Horde: Allow disabling issue creation on a per-workflow basis.

diff --git a/Engine/Source/Programs/Horde/Horde.Build.Tests/Fixture.cs b/Engine/Source/Programs/Horde/Horde.Build.Tests/Fixture.cs
--- a/Engine/Source/Programs/Horde/Horde.Build.Tests/Fixture.cs
+++ b/Engine/Source/Programs/Horde/Horde.Build.Tests/Fixture.cs
@@ -62,9 +62,9 @@ namespace Horde.Build.Tests
 			TemplateRefId1 = new TemplateRefId("template1");
 			TemplateRefId2 = new TemplateRefId("template2");
 
-			List<CreateTemplateRefRequest> templates = new List<CreateTemplateRefRequest>();
-			templates.Add(new CreateTemplateRefRequest { Id = TemplateRefId1.ToString(), Name = "Test Template" });
-			templates.Add(new CreateTemplateRefRequest { Id = TemplateRefId2.ToString(), Name = "Test Template" });
+			List<TemplateRefConfig> templates = new List<TemplateRefConfig>();
+			templates.Add(new TemplateRefConfig { Id = TemplateRefId1, Name = "Test Template" });
+			templates.Add(new TemplateRefConfig { Id = TemplateRefId2, Name = "Test Template" });
 
 			List<CreateStreamTabRequest> tabs = new List<CreateStreamTabRequest>();
 			tabs.Add(new CreateJobsTabRequest { Title = "foo", Templates = new List<string> { TemplateRefId1.ToString(), TemplateRefId2.ToString() } });
diff --git a/Engine/Source/Programs/Horde/Horde.Build.Tests/IssueServiceTests.cs b/Engine/Source/Programs/Horde/Horde.Build.Tests/IssueServiceTests.cs
--- a/Engine/Source/Programs/Horde/Horde.Build.Tests/IssueServiceTests.cs
+++ b/Engine/Source/Programs/Horde/Horde.Build.Tests/IssueServiceTests.cs
@@ -105,6 +105,18 @@ namespace Horde.Build.Tests
 		readonly UserId _bobId;
 		readonly DirectoryReference _workspaceDir;
 
+		async Task<IStream> CreateStreamAsync(ProjectId projectId, StreamId streamId, string streamName)
+		{
+			string revision = $"config:{streamId}";
+
+			StreamConfig streamConfig = new StreamConfig { Name = streamName };
+			streamConfig.Tabs.Add(new Api.CreateJobsTabRequest { Title = "General", Templates = new List<string> { "test-template" } });
+			streamConfig.Templates.Add(new Api.TemplateRefConfig { Id = new TemplateRefId("test-template") });
+			await ConfigCollection.AddConfigAsync(revision, streamConfig);
+
+			return Deref(await StreamCollection.TryCreateOrReplaceAsync(streamId, null, "", revision, projectId, streamConfig));
+		}
+
 		public IssueServiceTests()
 		{
 			if (RuntimeInformation.IsOSPlatform(OSPlatform.Windows))
@@ -118,9 +130,9 @@ namespace Horde.Build.Tests
 
 			IProject project = ProjectCollection.AddOrUpdateAsync(new ProjectId("ue4"), "", "", 0, new ProjectConfig { Name = "UE4" }).Result!;
 
-			IStream mainStream = StreamCollection.TryCreateOrReplaceAsync(_mainStreamId, null, "", "", project.Id, new StreamConfig { Name = MainStreamName }).Result!;
-			IStream releaseStream = StreamCollection.TryCreateOrReplaceAsync(_releaseStreamId, null, "", "", project.Id, new StreamConfig { Name = ReleaseStreamName }).Result!;
-			IStream devStream = StreamCollection.TryCreateOrReplaceAsync(_devStreamId, null, "", "", project.Id, new StreamConfig { Name = DevStreamName }).Result!;
+			IStream mainStream = CreateStreamAsync(project.Id, _mainStreamId, MainStreamName).Result;
+			IStream releaseStream = CreateStreamAsync(project.Id, _releaseStreamId, ReleaseStreamName).Result;
+			IStream devStream = CreateStreamAsync(project.Id, _devStreamId, DevStreamName).Result;
 
 			IUser bill = UserCollection.FindOrAddUserByLoginAsync("Bill").Result;
 			IUser anne = UserCollection.FindOrAddUserByLoginAsync("Anne").Result;
diff --git a/Engine/Source/Programs/Horde/Horde.Build.Tests/JobServiceTests.cs b/Engine/Source/Programs/Horde/Horde.Build.Tests/JobServiceTests.cs
--- a/Engine/Source/Programs/Horde/Horde.Build.Tests/JobServiceTests.cs
+++ b/Engine/Source/Programs/Horde/Horde.Build.Tests/JobServiceTests.cs
@@ -37,8 +37,8 @@ namespace Horde.Build.Tests
 			TemplateRefId templateRefId2 = new TemplateRefId("template2");
 
 			StreamConfig streamConfig = new StreamConfig();
-			streamConfig.Templates.Add(new CreateTemplateRefRequest { Id = templateRefId1.ToString(), Name = "Test Template", ChainedJobs = new List<CreateChainedJobTemplateRequest> { new CreateChainedJobTemplateRequest { TemplateId = templateRefId2.ToString(), Trigger = "Setup Build" } } });
-			streamConfig.Templates.Add(new CreateTemplateRefRequest { Id = templateRefId2.ToString(), Name = "Test Template" });
+			streamConfig.Templates.Add(new TemplateRefConfig { Id = templateRefId1, Name = "Test Template", ChainedJobs = new List<CreateChainedJobTemplateRequest> { new CreateChainedJobTemplateRequest { TemplateId = templateRefId2.ToString(), Trigger = "Setup Build" } } });
+			streamConfig.Templates.Add(new TemplateRefConfig { Id = templateRefId2, Name = "Test Template" });
 			streamConfig.Tabs.Add(new CreateJobsTabRequest { Title = "foo", Templates = new List<string> { templateRefId1.ToString(), templateRefId2.ToString() } });
 
 			StreamId streamId = new StreamId("ue5-main");
diff --git a/Engine/Source/Programs/Horde/Horde.Build.Tests/SchedulerTests.cs b/Engine/Source/Programs/Horde/Horde.Build.Tests/SchedulerTests.cs
--- a/Engine/Source/Programs/Horde/Horde.Build.Tests/SchedulerTests.cs
+++ b/Engine/Source/Programs/Horde/Horde.Build.Tests/SchedulerTests.cs
@@ -58,7 +58,7 @@ namespace Horde.Build.Tests
 			StreamConfig config = new StreamConfig();
 			config.Name = "//UE5/Main";
 			config.Tabs.Add(new CreateJobsTabRequest { Title = "foo", Templates = new List<string> { TemplateRefId.ToString() } });
-			config.Templates.Add(new CreateTemplateRefRequest { Id = TemplateRefId.ToString(), Name = "Test", Schedule = schedule });
+			config.Templates.Add(new TemplateRefConfig { Id = TemplateRefId, Name = "Test", Schedule = schedule });
 
 			return (await StreamService.StreamCollection.TryCreateOrReplaceAsync(StreamId, stream, "", "", ProjectId, config))!; 
 		}
@@ -484,12 +484,12 @@ namespace Horde.Build.Tests
 
 			// Create two templates, the second dependent on the first
 			TemplateRefId newTemplateRefId1 = new TemplateRefId("new-template-1");
-			CreateTemplateRefRequest newTemplate1 = new CreateTemplateRefRequest();
-			newTemplate1.Id = newTemplateRefId1.ToString();
+			TemplateRefConfig newTemplate1 = new TemplateRefConfig();
+			newTemplate1.Id = newTemplateRefId1;
 
 			TemplateRefId newTemplateRefId2 = new TemplateRefId("new-template-2");
-			CreateTemplateRefRequest newTemplate2 = new CreateTemplateRefRequest();
-			newTemplate2.Id = newTemplateRefId2.ToString();
+			TemplateRefConfig newTemplate2 = new TemplateRefConfig();
+			newTemplate2.Id = newTemplateRefId2;
 			newTemplate2.Name = "Test template 2";
 			newTemplate2.Schedule = new CreateScheduleRequest();
 			newTemplate2.Schedule.MaxChanges = 4;
diff --git a/Engine/Source/Programs/Horde/Horde.Build.Tests/TestSetup.cs b/Engine/Source/Programs/Horde/Horde.Build.Tests/TestSetup.cs
--- a/Engine/Source/Programs/Horde/Horde.Build.Tests/TestSetup.cs
+++ b/Engine/Source/Programs/Horde/Horde.Build.Tests/TestSetup.cs
@@ -13,6 +13,7 @@ using Horde.Build.Collections;
 using Horde.Build.Collections.Impl;
 using Horde.Build.Commits;
 using Horde.Build.Commits.Impl;
+using Horde.Build.Config;
 using Horde.Build.Controllers;
 using Horde.Build.Jobs;
 using Horde.Build.Logs;
@@ -54,6 +55,7 @@ namespace Horde.Build.Tests
 	{
 		public FakeClock Clock => ServiceProvider.GetRequiredService<FakeClock>();
 
+		public ConfigCollection ConfigCollection => ServiceProvider.GetRequiredService<ConfigCollection>();
 		public IGraphCollection GraphCollection => ServiceProvider.GetRequiredService<IGraphCollection>();
 		public INotificationTriggerCollection NotificationTriggerCollection => ServiceProvider.GetRequiredService<INotificationTriggerCollection>();
 		public IStreamCollection StreamCollection => ServiceProvider.GetRequiredService<IStreamCollection>();
@@ -174,6 +176,7 @@ namespace Horde.Build.Tests
 			services.AddSingleton<AutoscaleService>();
 			services.AddSingleton<AwsFleetManager, AwsFleetManager>();
 			services.AddSingleton<ConsistencyService>();
+			services.AddSingleton<ConfigCollection>();
 			services.AddSingleton<RequestTrackerService>();
 			services.AddSingleton<CredentialService>();
 			services.AddSingleton<JobTaskSource>();
diff --git a/Engine/Source/Programs/Horde/Horde.Build/Api/Streams.cs b/Engine/Source/Programs/Horde/Horde.Build/Api/Streams.cs
--- a/Engine/Source/Programs/Horde/Horde.Build/Api/Streams.cs
+++ b/Engine/Source/Programs/Horde/Horde.Build/Api/Streams.cs
@@ -187,12 +187,18 @@ namespace Horde.Build.Api
 	/// <summary>
 	/// Parameters to create a template within a stream
 	/// </summary>
-	public class CreateTemplateRefRequest : CreateTemplateRequest
+	public class TemplateRefConfig : TemplateConfig
 	{
+		TemplateRefId _id;
+
 		/// <summary>
 		/// Optional identifier for this ref. If not specified, an id will be generated from the name.
 		/// </summary>
-		public string? Id { get; set; }
+		public TemplateRefId Id
+		{
+			get => _id.IsEmpty ? TemplateRefId.Sanitize(Name) : _id;
+			set => _id = value;
+		}
 
 		/// <summary>
 		/// Whether to show badges in UGS for these jobs
diff --git a/Engine/Source/Programs/Horde/Horde.Build/Api/Templates.cs b/Engine/Source/Programs/Horde/Horde.Build/Api/Templates.cs
--- a/Engine/Source/Programs/Horde/Horde.Build/Api/Templates.cs
+++ b/Engine/Source/Programs/Horde/Horde.Build/Api/Templates.cs
@@ -388,7 +388,7 @@ namespace Horde.Build.Api
 	/// <summary>
 	/// Parameters to create a new template
 	/// </summary>
-	public class CreateTemplateRequest
+	public class TemplateConfig
 	{
 		/// <summary>
 		/// Name for the new template
diff --git a/Engine/Source/Programs/Horde/Horde.Build/Configuration/ConfigCollection.cs b/Engine/Source/Programs/Horde/Horde.Build/Configuration/ConfigCollection.cs
--- a/Engine/Source/Programs/Horde/Horde.Build/Configuration/ConfigCollection.cs
+++ b/Engine/Source/Programs/Horde/Horde.Build/Configuration/ConfigCollection.cs
@@ -83,12 +83,15 @@ namespace Horde.Build.Config
 		/// </summary>
 		/// <typeparam name="T"></typeparam>
 		/// <param name="revision"></param>
-		/// <param name="data"></param>
+		/// <param name="config"></param>
 		/// <returns></returns>
-		public async Task<T> AddConfigAsync<T>(string revision, ReadOnlyMemory<byte> data)
+		public async Task AddConfigAsync<T>(string revision, T config)
 		{
+			JsonSerializerOptions options = new JsonSerializerOptions();
+			Startup.ConfigureJsonSerializer(options);
+
+			byte[] data = JsonSerializer.SerializeToUtf8Bytes(config, options);
 			await AddConfigDataAsync(revision, data);
-			return await GetConfigAsync<T>(revision);
 		}
 
 		/// <summary>
diff --git a/Engine/Source/Programs/Horde/Horde.Build/Configuration/ConfigUpdateService.cs b/Engine/Source/Programs/Horde/Horde.Build/Configuration/ConfigUpdateService.cs
--- a/Engine/Source/Programs/Horde/Horde.Build/Configuration/ConfigUpdateService.cs
+++ b/Engine/Source/Programs/Horde/Horde.Build/Configuration/ConfigUpdateService.cs
@@ -376,7 +376,8 @@ namespace Horde.Build.Config
 		async Task<T> ReadDataAsync<T>(string revision, Uri configPath) where T : class
 		{
 			byte[] data = await ReadDataAsync(configPath);
-			return await _configCollection.AddConfigAsync<T>(revision, data);
+			await _configCollection.AddConfigDataAsync(revision, data);
+			return await _configCollection.GetConfigAsync<T>(revision);
 		}
 
 		Task<byte[]> ReadDataAsync(Uri configPath)
diff --git a/Engine/Source/Programs/Horde/Horde.Build/Issues/Impl/IssueReportService.cs b/Engine/Source/Programs/Horde/Horde.Build/Issues/Impl/IssueReportService.cs
--- a/Engine/Source/Programs/Horde/Horde.Build/Issues/Impl/IssueReportService.cs
+++ b/Engine/Source/Programs/Horde/Horde.Build/Issues/Impl/IssueReportService.cs
@@ -132,11 +132,11 @@ namespace Horde.Build.Issues.Impl
 						jobs ??= await _jobCollection.FindAsync(streamId: stream.Id, minCreateTime: prevScheduledReportTime);
 
 						HashSet<TemplateRefId> wfTemplateRefIds = new HashSet<TemplateRefId>();
-						foreach (Api.CreateTemplateRefRequest templateRef in config.Templates)
+						foreach (Api.TemplateRefConfig templateRef in config.Templates)
 						{
-							if (templateRef.Id != null && templateRef.Workflow != null && templateRef.Workflow == workflow.Id)
+							if (templateRef.Workflow != null && templateRef.Workflow == workflow.Id)
 							{
-								wfTemplateRefIds.Add(new TemplateRefId(templateRef.Id));
+								wfTemplateRefIds.Add(templateRef.Id);
 							}
 						}
 
diff --git a/Engine/Source/Programs/Horde/Horde.Build/Issues/Impl/IssueService.cs b/Engine/Source/Programs/Horde/Horde.Build/Issues/Impl/IssueService.cs
--- a/Engine/Source/Programs/Horde/Horde.Build/Issues/Impl/IssueService.cs
+++ b/Engine/Source/Programs/Horde/Horde.Build/Issues/Impl/IssueService.cs
@@ -9,10 +9,13 @@ using System.Text.RegularExpressions;
 using System.Threading;
 using System.Threading.Tasks;
 using EpicGames.Core;
+using Horde.Build.Api;
 using Horde.Build.Collections;
+using Horde.Build.Config;
 using Horde.Build.IssueHandlers;
 using Horde.Build.IssueHandlers.Impl;
 using Horde.Build.Models;
+using Horde.Build.Server;
 using Horde.Build.Utilities;
 using HordeCommon;
 using Microsoft.Extensions.Hosting;
@@ -128,6 +131,8 @@ namespace Horde.Build.Services.Impl
 		/// Maximum number of changes to query from Perforce in one go
 		/// </summary>
 		const int MaxChanges = 250;
+
+		readonly ConfigCollection _configCollection;
 		readonly IJobStepRefCollection _jobStepRefs;
 		readonly IIssueCollection _issueCollection;
 		readonly StreamService _streams;
@@ -175,7 +180,7 @@ namespace Horde.Build.Services.Impl
 		/// <summary>
 		/// Constructor
 		/// </summary>
-		public IssueService(IIssueCollection issueCollection, IJobStepRefCollection jobStepRefs, StreamService streams, IUserCollection userCollection, IPerforceService perforce, ILogFileService logFileService, IClock clock, ILogger<IssueService> logger)
+		public IssueService(ConfigCollection configCollection, IIssueCollection issueCollection, IJobStepRefCollection jobStepRefs, StreamService streams, IUserCollection userCollection, IPerforceService perforce, ILogFileService logFileService, IClock clock, ILogger<IssueService> logger)
 		{
 			Type[] issueTypes = Assembly.GetExecutingAssembly().GetTypes().Where(x => !x.IsAbstract && typeof(IIssue).IsAssignableFrom(x)).ToArray();
 			foreach (Type issueType in issueTypes)
@@ -184,6 +189,7 @@ namespace Horde.Build.Services.Impl
 			}
 
 			// Get all the collections
+			_configCollection = configCollection;
 			_issueCollection = issueCollection;
 			_jobStepRefs = jobStepRefs;
 			_streams = streams;
@@ -536,6 +542,23 @@ namespace Horde.Build.Services.Impl
 				throw new Exception($"Invalid stream id '{job.StreamId}' on job '{job.Id}'");
 			}
 
+			// Check whether issue creation is enabled for this template type
+			StreamConfig config = await _configCollection.GetConfigAsync<StreamConfig>(stream.ConfigRevision);
+
+			TemplateRefConfig? templateRef;
+			if (!config.TryGetTemplate(job.TemplateId, out templateRef))
+			{
+				return;
+			}
+
+			if (templateRef.Workflow != null)
+			{
+				if(!config.TryGetWorkflow(templateRef.Workflow.Value, out WorkflowConfig? workflow) || !workflow.CreateIssues)
+				{
+					return;
+				}
+			}
+
 			// Find the batch for this event
 			IJobStepBatch? batch;
 			if (!job.TryGetBatch(batchId, out batch))
diff --git a/Engine/Source/Programs/Horde/Horde.Build/Projects/Impl/StreamCollection.cs b/Engine/Source/Programs/Horde/Horde.Build/Projects/Impl/StreamCollection.cs
--- a/Engine/Source/Programs/Horde/Horde.Build/Projects/Impl/StreamCollection.cs
+++ b/Engine/Source/Programs/Horde/Horde.Build/Projects/Impl/StreamCollection.cs
@@ -174,24 +174,16 @@ namespace Horde.Build.Collections.Impl
 		/// <param name="stream">The current stream state</param>
 		/// <param name="templateCollection">The template service</param>
 		/// <returns>List of new template references</returns>
-		async Task<Dictionary<TemplateRefId, TemplateRef>> CreateTemplateRefsAsync(List<CreateTemplateRefRequest> requests, IStream? stream, ITemplateCollection templateCollection)
+		async Task<Dictionary<TemplateRefId, TemplateRef>> CreateTemplateRefsAsync(List<TemplateRefConfig> requests, IStream? stream, ITemplateCollection templateCollection)
 		{
 			Dictionary<TemplateRefId, TemplateRef> newTemplateRefs = new Dictionary<TemplateRefId, TemplateRef>();
-			foreach (CreateTemplateRefRequest request in requests)
+			foreach (TemplateRefConfig request in requests)
 			{
 				// Create the template
 				ITemplate newTemplate = await templateCollection.AddAsync(request.Name, request.Priority, request.AllowPreflights, request.UpdateIssues, request.PromoteIssuesByDefault, request.InitialAgentType, request.SubmitNewChange, request.SubmitDescription, request.Arguments, request.Parameters.ConvertAll(x => x.ToModel()));
 
 				// Get an identifier for the new template ref
-				TemplateRefId newTemplateRefId;
-				if (request.Id != null)
-				{
-					newTemplateRefId = new TemplateRefId(request.Id);
-				}
-				else
-				{
-					newTemplateRefId = TemplateRefId.Sanitize(request.Name);
-				}
+				TemplateRefId newTemplateRefId = request.Id;
 
 				// Create the schedule object
 				Schedule? schedule = null;
diff --git a/Engine/Source/Programs/Horde/Horde.Build/Server/Config.cs b/Engine/Source/Programs/Horde/Horde.Build/Server/Config.cs
--- a/Engine/Source/Programs/Horde/Horde.Build/Server/Config.cs
+++ b/Engine/Source/Programs/Horde/Horde.Build/Server/Config.cs
@@ -3,6 +3,8 @@
 using System;
 using System.Collections.Generic;
 using System.ComponentModel.DataAnnotations;
+using System.Diagnostics.CodeAnalysis;
+using System.Linq;
 using EpicGames.Core;
 using EpicGames.Horde.Common;
 using EpicGames.Serialization;
@@ -15,6 +17,7 @@ namespace Horde.Build.Server
 {
 	using ProjectId = StringId<IProject>;
 	using StreamId = StringId<IStream>;
+	using TemplateRefId = StringId<TemplateRef>;
 	using WorkflowId = StringId<WorkflowConfig>;
 	
 	/// <summary>
@@ -237,7 +240,7 @@ namespace Horde.Build.Server
 		/// <summary>
 		/// List of templates to create
 		/// </summary>
-		public List<CreateTemplateRefRequest> Templates { get; set; } = new List<CreateTemplateRefRequest>();
+		public List<TemplateRefConfig> Templates { get; set; } = new List<TemplateRefConfig>();
 
 		/// <summary>
 		/// Custom permissions for this object
@@ -268,6 +271,30 @@ namespace Horde.Build.Server
 		/// Workflows for dealing with new issues
 		/// </summary>
 		public List<WorkflowConfig> Workflows { get; set; } = new List<WorkflowConfig>();
+
+		/// <summary>
+		/// Tries to find a template with the given id
+		/// </summary>
+		/// <param name="templateRefId"></param>
+		/// <param name="templateConfig"></param>
+		/// <returns></returns>
+		public bool TryGetTemplate(TemplateRefId templateRefId, [NotNullWhen(true)] out TemplateRefConfig? templateConfig)
+		{
+			templateConfig = Templates.FirstOrDefault(x => x.Id == templateRefId);
+			return templateConfig != null;
+		}
+
+		/// <summary>
+		/// Tries to find a workflow with the given id
+		/// </summary>
+		/// <param name="workflowId"></param>
+		/// <param name="workflowConfig"></param>
+		/// <returns></returns>
+		public bool TryGetWorkflow(WorkflowId workflowId, [NotNullWhen(true)] out WorkflowConfig? workflowConfig)
+		{
+			workflowConfig = Workflows.FirstOrDefault(x => x.Id == workflowId);
+			return workflowConfig != null;
+		}
 	}
 
 	/// <summary>
@@ -280,6 +307,11 @@ namespace Horde.Build.Server
 		/// </summary>
 		public WorkflowId Id { get; set; } = WorkflowId.Empty;
 
+		/// <summary>
+		/// Whether to allow creating issues for these jobs
+		/// </summary>
+		public bool CreateIssues { get; set; } = true;
+
 		/// <summary>
 		/// Times of day at which to send a report
 		/// </summary>
