Horde: Clean up messages posted to dev environment.

diff --git a/Engine/Source/Programs/Horde/Horde.Build/Notifications/Impl/SlackNotificationSink.cs b/Engine/Source/Programs/Horde/Horde.Build/Notifications/Impl/SlackNotificationSink.cs
--- a/Engine/Source/Programs/Horde/Horde.Build/Notifications/Impl/SlackNotificationSink.cs
+++ b/Engine/Source/Programs/Horde/Horde.Build/Notifications/Impl/SlackNotificationSink.cs
@@ -918,8 +918,8 @@ namespace Horde.Build.Notifications.Impl
 			if (report.Workflow.ReportChannel != null)
 			{
 				List<BlockBase> blocks = new List<BlockBase>();
-				blocks.Add(new HeaderBlock($"{report.Stream.Name}: {report.Time:d}"));
-				await SendMessageAsync(report.Workflow.ReportChannel, blocks: blocks.ToArray());// B, message.ToString());/// "Hello world");//, blocks: blocks.ToArray());
+				blocks.Add(new HeaderBlock(AddEnvironmentAnnotation($"{report.Stream.Name}: {report.Time:d}")));
+				await SendMessageAsync(report.Workflow.ReportChannel, blocks: blocks.ToArray(), withEnvironment: false);
 
 				StringBuilder body = new StringBuilder();
 				if (report.Issues.Count == 0)
@@ -971,7 +971,7 @@ namespace Horde.Build.Notifications.Impl
 					body.Append($"\n\n{report.NumPassingSteps:n0} of {report.NumSteps:n0} (*{totalPct:0.0}%*) build steps succeeded since last status update.");
 				}
 
-				await SendMessageAsync(report.Workflow.ReportChannel, text: body.ToString());
+				await SendMessageAsync(report.Workflow.ReportChannel, text: body.ToString(), withEnvironment: false);
 			}
 		}
 
@@ -1155,7 +1155,7 @@ namespace Horde.Build.Notifications.Impl
 				attachment.FallbackText += $" - Device: {device.Name} Pool: {pool.Name}";
 			}
 				
-			attachment.Blocks.Add(new HeaderBlock($"{message}", false));
+			attachment.Blocks.Add(new HeaderBlock(AddEnvironmentAnnotation(message), false));
 
 			if (stream != null && job != null && step != null && node != null)
 			{
@@ -1307,7 +1307,7 @@ namespace Horde.Build.Notifications.Impl
 			}
 		}
 
-		private async Task SendMessageAsync(string recipient, string? text = null, BlockBase[]? blocks = null, BlockKitAttachment[]? attachments = null)
+		private async Task SendMessageAsync(string recipient, string? text = null, BlockBase[]? blocks = null, BlockKitAttachment[]? attachments = null, bool withEnvironment = true)
 		{
 			if (_allowUsers != null && !_allowUsers.Contains(recipient) && !recipient.StartsWith("#", StringComparison.Ordinal))
 			{
@@ -1319,7 +1319,7 @@ namespace Horde.Build.Notifications.Impl
 			message.Channel = recipient;
 			if (text != null)
 			{
-				message.Text = AddEnvironmentAnnotation(text);
+				message.Text = withEnvironment ? AddEnvironmentAnnotation(text) : text;
 			}
 			if (blocks != null)
 			{
