Horde: Add an option to create a thread in the triage channel for every newly uncovered issue. The thread will [REDACTED_HANDLE] suspects for the issue up to a defined number of users.

diff --git a/Engine/Source/Programs/Horde/Horde.Build/Notifications/Impl/SlackNotificationSink.cs b/Engine/Source/Programs/Horde/Horde.Build/Notifications/Impl/SlackNotificationSink.cs
--- a/Engine/Source/Programs/Horde/Horde.Build/Notifications/Impl/SlackNotificationSink.cs
+++ b/Engine/Source/Programs/Horde/Horde.Build/Notifications/Impl/SlackNotificationSink.cs
@@ -17,6 +17,7 @@ using System.Threading.Tasks;
 using EpicGames.Core;
 using Horde.Build.Api;
 using Horde.Build.Collections;
+using Horde.Build.Config;
 using Horde.Build.Issues.Impl;
 using Horde.Build.Models;
 using Horde.Build.Server;
@@ -199,6 +200,7 @@ namespace Horde.Build.Notifications.Impl
 		readonly IUserCollection _userCollection;
 		readonly ILogFileService _logFileService;
 		readonly StreamService _streamService;
+		readonly ConfigCollection _configCollection;
 		readonly IWebHostEnvironment _environment;
 		readonly ServerSettings _settings;
 		readonly IMongoCollection<MessageStateDocument> _messageStates;
@@ -214,12 +216,13 @@ namespace Horde.Build.Notifications.Impl
 		/// <summary>
 		/// Constructor
 		/// </summary>
-		public SlackNotificationSink(MongoService mongoService, IIssueService issueService, IUserCollection userCollection, ILogFileService logFileService, StreamService streamService, IWebHostEnvironment environment, IOptions<ServerSettings> settings, ILogger<SlackNotificationSink> logger)
+		public SlackNotificationSink(MongoService mongoService, IIssueService issueService, IUserCollection userCollection, ILogFileService logFileService, StreamService streamService, ConfigCollection configCollection, IWebHostEnvironment environment, IOptions<ServerSettings> settings, ILogger<SlackNotificationSink> logger)
 		{
 			_issueService = issueService;
 			_userCollection = userCollection;
 			_logFileService = logFileService;
 			_streamService = streamService;
+			_configCollection = configCollection;
 			_environment = environment;
 			_settings = settings.Value;
 			_messageStates = mongoService.Database.GetCollection<MessageStateDocument>("Slack");
@@ -599,6 +602,23 @@ namespace Horde.Build.Notifications.Impl
 				}
 			}
 
+			if (details.Spans.Count > 0)
+			{
+				IIssueSpan span = details.Spans[0];
+				IStream? stream = await _streamService.GetStreamAsync(span.StreamId);
+				if(stream != null)
+				{
+					StreamConfig config = await _configCollection.GetConfigAsync<StreamConfig>(stream.ConfigRevision);
+					if (config.TryGetTemplate(span.TemplateRefId, out TemplateRefConfig? templateRef))
+					{
+						if (templateRef.Workflow != null && config.TryGetWorkflow(templateRef.Workflow.Value, out WorkflowConfig? workflow))
+						{
+							await SendTriageMessageAsync(issue, span, workflow);
+						}
+					}
+				}
+			}
+
 			if (issue.OwnerId == null && (details.Suspects.Count == 0 || details.Suspects.All(x => x.DeclinedAt != null) || issue.CreatedAt < DateTime.UtcNow - TimeSpan.FromHours(1.0)))
 			{
 				foreach (IIssueSpan span in details.Spans)
@@ -644,6 +664,94 @@ namespace Horde.Build.Notifications.Impl
 			}
 		}
 
+		async Task SendTriageMessageAsync(IIssue issue, IIssueSpan span, WorkflowConfig workflow)
+		{
+			string? triageChannel = workflow.TriageChannel;
+			if (triageChannel != null)
+			{
+				Uri issueUrl = GetIssueUrl(issue, span.FirstFailure);
+
+				BlockKitMessage message = new BlockKitMessage();
+				message.Channel = triageChannel;
+				message.Text = $"{workflow.TriagePrefix}*Issue <{issueUrl}|{issue.Id}>*: {issue.Summary}{workflow.TriageSuffix}";
+
+				string eventId = $"issue_triage_{issue.Id}";
+
+				(MessageStateDocument state, bool isNew) = await AddOrUpdateMessageStateAsync(triageChannel, eventId, null, "");
+				if (isNew)
+				{
+					PostMessageResponse response = await SendRequestAsync<PostMessageResponse>(PostMessageUrl, message);
+					if (String.IsNullOrEmpty(response.Ts) || String.IsNullOrEmpty(response.Channel))
+					{
+						_logger.LogWarning("Missing 'ts' or 'channel' field on slack response");
+					}
+
+					state.Channel = response.Channel ?? String.Empty;
+					state.Ts = response.Ts ?? String.Empty;
+
+					await SetMessageTimestampAsync(state.Id, state.Channel, state.Ts);
+
+					// Create the summary text
+					StringBuilder summary = new StringBuilder();
+					summary.AppendLine($"From <{GetStepUrl(issue, span.FirstFailure)}|{span.FirstFailure.JobName}: {span.NodeName}>.");
+
+					if (span.FirstFailure.LogId != null)
+					{
+						LogId logId = span.FirstFailure.LogId.Value;
+						ILogFile? logFile = await _logFileService.GetLogFileAsync(logId);
+						if (logFile != null)
+						{
+							List<ILogEvent> events = await _logFileService.FindEventsAsync(logFile, span.Id, 0, 20);
+							if (events.Any(x => x.Severity == EventSeverity.Error))
+							{
+								events.RemoveAll(x => x.Severity == EventSeverity.Warning);
+							}
+
+							List<string> eventStrings = new List<string>();
+							for (int idx = 0; idx < Math.Min(events.Count, 3); idx++)
+							{
+								ILogEventData data = await _logFileService.GetEventDataAsync(logFile, events[idx].LineIndex, events[idx].LineCount);
+								summary.AppendLine($"```{data.Message}```");
+							}
+							if (events.Count > 3)
+							{
+								summary.AppendLine("```...```");
+							}
+						}
+					}
+					await SendMessageToThread(triageChannel, state.Ts, summary.ToString());
+
+					List<IIssueSuspect> suspects = await _issueService.GetIssueSuspectsAsync(issue);
+					IGrouping<UserId, IIssueSuspect>[] suspectGroups = suspects.GroupBy(x => x.AuthorId).ToArray();
+
+					if (suspectGroups.Length > 0 && suspectGroups.Length <= workflow.MaxMentions)
+					{
+						List<string> suspectList = new List<string>();
+						foreach (IGrouping<UserId, IIssueSuspect> suspectGroup in suspectGroups)
+						{
+							string mention = await FormatMentionAsync(suspectGroup.Key);
+							string changes = String.Join(", ", suspectGroup.Select(x => $"CL {x.Change}"));
+							suspectList.Add($"{mention} ({changes})");
+						}
+
+						string suspectMessage = $"Possibly {StringUtils.FormatList(suspectList, "or")}.";
+						await SendMessageToThread(triageChannel, state.Ts, suspectMessage);
+					}
+				}
+			}
+		}
+
+		async Task<string?> SendMessageToThread(string channel, string threadTs, string text)
+		{
+			BlockKitMessage message = new BlockKitMessage();
+			message.Channel = channel;
+			message.ThreadTs = threadTs;
+			message.Text = text;
+
+			PostMessageResponse response = await SendRequestAsync<PostMessageResponse>(PostMessageUrl, message);
+			return response.Ts;
+		}
+
 		async Task NotifyIssueUpdatedAsync(IUser user, IIssue issue, IIssueDetails details)
 		{
 			string? slackUserId = await GetSlackUserId(user);
@@ -655,6 +763,16 @@ namespace Horde.Build.Notifications.Impl
 			await SendIssueMessageAsync(slackUserId, issue, details, user.Id);
 		}
 
+		Uri GetStepUrl(IIssue issue, IIssueStep step)
+		{
+			return new Uri(_settings.DashboardUrl, $"job/{step.JobId}?step={step.StepId}");
+		}
+
+		Uri GetIssueUrl(IIssue issue, IIssueStep step)
+		{
+			return new Uri(_settings.DashboardUrl, $"job/{step.JobId}?step={step.StepId}&issue={issue.Id}");
+		}
+
 		async Task SendIssueMessageAsync(string recipient, IIssue issue, IIssueDetails details, UserId? userId)
 		{
 			using IDisposable scope = _logger.BeginScope("SendIssueMessageAsync (User: {SlackUser}, Issue: {IssueId})", recipient, issue.Id);
@@ -665,8 +783,7 @@ namespace Horde.Build.Notifications.Impl
 			Uri issueUrl = _settings.DashboardUrl;
 			if (details.Steps.Count > 0)
 			{
-				IIssueStep firstStep = details.Steps[0];
-				issueUrl = new Uri(issueUrl, $"job/{firstStep.JobId}?step={firstStep.StepId}&issue={issue.Id}");
+				issueUrl = GetIssueUrl(issue, details.Steps[0]);
 			}
 			attachment.Blocks.Add(new SectionBlock($"*<{issueUrl}|Issue #{issue.Id}: {issue.Summary}>*"));
 
@@ -849,6 +966,11 @@ namespace Horde.Build.Notifications.Impl
 				return user.Login;
 			}
 
+			if (!_environment.IsProduction())
+			{
+				return $"{user.Name} [{slackUserId}]";
+			}
+
 			return $"<@{slackUserId}>";
 		}
 
diff --git a/Engine/Source/Programs/Horde/Horde.Build/Properties/launchSettings.json b/Engine/Source/Programs/Horde/Horde.Build/Properties/launchSettings.json
--- a/Engine/Source/Programs/Horde/Horde.Build/Properties/launchSettings.json
+++ b/Engine/Source/Programs/Horde/Horde.Build/Properties/launchSettings.json
@@ -3,7 +3,7 @@
     "HordeServer": {
       "commandName": "Project",
       "environmentVariables": {
-        "Horde__DisableSchedules": "true"
+        "ASPNETCORE_ENVIRONMENT": "Development"
       }
     }
   }
diff --git a/Engine/Source/Programs/Horde/Horde.Build/Utilities/BlockKit/BlockKitMessage.cs b/Engine/Source/Programs/Horde/Horde.Build/Utilities/BlockKit/BlockKitMessage.cs
--- a/Engine/Source/Programs/Horde/Horde.Build/Utilities/BlockKit/BlockKitMessage.cs
+++ b/Engine/Source/Programs/Horde/Horde.Build/Utilities/BlockKit/BlockKitMessage.cs
@@ -22,6 +22,12 @@ namespace Horde.Build.Utilities.Slack.BlockKit
 		[JsonPropertyName("ts")]
 		public string? Ts { get; set; }
 
+		/// <summary>
+		/// Timestamp of the thread to post to
+		/// </summary>
+		[JsonPropertyName("thread_ts")]
+		public string? ThreadTs { get; set; }
+
 		/// <summary>
 		/// If no blocks are provided, this text will be displayed. This is also what is used on the pop up notification.
 		/// </summary>
diff --git a/Engine/Source/Programs/Shared/EpicGames.Core/StringUtils.cs b/Engine/Source/Programs/Shared/EpicGames.Core/StringUtils.cs
--- a/Engine/Source/Programs/Shared/EpicGames.Core/StringUtils.cs
+++ b/Engine/Source/Programs/Shared/EpicGames.Core/StringUtils.cs
@@ -229,15 +229,15 @@ namespace EpicGames.Core
 		/// <param name="arguments">List of strings to format</param>
 		/// <param name="conjunction">Conjunction to use between the last two items in the list (eg. "and" or "or")</param>
 		/// <returns>Formatted list of strings</returns>
-		public static string FormatList(string[] arguments, string conjunction = "and")
+		public static string FormatList(IReadOnlyList<string> arguments, string conjunction = "and")
 		{
 			StringBuilder result = new StringBuilder();
-			if (arguments.Length > 0)
+			if (arguments.Count > 0)
 			{
 				result.Append(arguments[0]);
-				for (int idx = 1; idx < arguments.Length; idx++)
+				for (int idx = 1; idx < arguments.Count; idx++)
 				{
-					if (idx == arguments.Length - 1)
+					if (idx == arguments.Count - 1)
 					{
 						result.AppendFormat(" {0} ", conjunction);
 					}
