Horde: Missing file.

diff --git a/Engine/Source/Programs/Horde/Horde.Build/Server/Config.cs b/Engine/Source/Programs/Horde/Horde.Build/Server/Config.cs
--- a/Engine/Source/Programs/Horde/Horde.Build/Server/Config.cs
+++ b/Engine/Source/Programs/Horde/Horde.Build/Server/Config.cs
@@ -321,6 +321,26 @@ namespace Horde.Build.Server
 		/// Channel to post summary information for these templates.
 		/// </summary>
 		public string? ReportChannel { get; set; }
+
+		/// <summary>
+		/// Channel to post threads for triaging new issues
+		/// </summary>
+		public string? TriageChannel { get; set; }
+
+		/// <summary>
+		/// Prefix for all triage messages
+		/// </summary>
+		public string? TriagePrefix { get; set; } = ":thread: ";
+
+		/// <summary>
+		/// Suffix for all triage messages
+		/// </summary>
+		public string? TriageSuffix { get; set; }
+
+		/// <summary>
+		/// Maximum number of people to mention on a triage thread
+		/// </summary>
+		public int MaxMentions { get; set; } = 5;
 	}
 
 	/// <summary>
