Horde: Add an error code explaining why a step failed to run, and use it to mark steps that timed out.

diff --git a/Engine/Source/Programs/Horde/Horde.Agent/Services/WorkerService.cs b/Engine/Source/Programs/Horde/Horde.Agent/Services/WorkerService.cs
--- a/Engine/Source/Programs/Horde/Horde.Agent/Services/WorkerService.cs
+++ b/Engine/Source/Programs/Horde/Horde.Agent/Services/WorkerService.cs
@@ -1256,13 +1256,6 @@ namespace Horde.Agent.Services
 			{
 				try
 				{
-					if (timer.Elapsed > TimeSpan.FromHours(24.0))
-					{
-						_logger.LogDebug("Step was aborted after running for {NumHours} (JobId={JobId} BatchId={BatchId} StepId={StepId})", timer.Elapsed, jobId, batchId, stepId);
-						stepCancelSource.Cancel();
-						break;
-					}
-
 					GetStepResponse res = await epcClient.InvokeAsync(x => x.GetStepAsync(new GetStepRequest(jobId, batchId, stepId), null, null, cancellationToken), new RpcContext(), cancellationToken);
 					if (res.AbortRequested)
 					{
diff --git a/Engine/Source/Programs/Horde/Horde.Build/Api/Jobs.cs b/Engine/Source/Programs/Horde/Horde.Build/Api/Jobs.cs
--- a/Engine/Source/Programs/Horde/Horde.Build/Api/Jobs.cs
+++ b/Engine/Source/Programs/Horde/Horde.Build/Api/Jobs.cs
@@ -517,6 +517,11 @@ namespace Horde.Build.Api
 		/// Current outcome of the jobstep
 		/// </summary>
 		public JobStepOutcome Outcome { get; set; }
+
+		/// <summary>
+		/// Error describing additional context for why a step failed to complete
+		/// </summary>
+		public JobStepError Error { get; set; }
 		
 		/// <summary>
 		/// If the step has been requested to abort
@@ -585,6 +590,7 @@ namespace Horde.Build.Api
 			NodeIdx = step.NodeIdx;
 			State = step.State;
 			Outcome = step.Outcome;
+			Error = step.Error;
 			AbortRequested = step.AbortRequested;
 			AbortByUser = abortedByUserInfo?.Login;
 			AbortedByUserInfo = abortedByUserInfo;
diff --git a/Engine/Source/Programs/Horde/Horde.Build/Controllers/JobsController.cs b/Engine/Source/Programs/Horde/Horde.Build/Controllers/JobsController.cs
--- a/Engine/Source/Programs/Horde/Horde.Build/Controllers/JobsController.cs
+++ b/Engine/Source/Programs/Horde/Horde.Build/Controllers/JobsController.cs
@@ -1225,7 +1225,7 @@ namespace Horde.Build.Controllers
 
 			try
 			{
-				IJob? newJob = await _jobService.UpdateStepAsync(job, batchId, stepId, request.State, request.Outcome, request.AbortRequested, abortByUser, request.LogId?.ToObjectId<ILogFile>(), null, retryByUser, request.Priority, null, request.Properties);
+				IJob? newJob = await _jobService.UpdateStepAsync(job, batchId, stepId, request.State, request.Outcome, null, request.AbortRequested, abortByUser, request.LogId?.ToObjectId<ILogFile>(), null, retryByUser, request.Priority, null, request.Properties);
 				if (newJob == null)
 				{
 					return NotFound(jobId);
@@ -1357,7 +1357,7 @@ namespace Horde.Build.Controllers
 			if (triggerId == null)
 			{
 				triggerId = ObjectId.GenerateNewId();
-				if (await _jobService.UpdateStepAsync(job, batchId, stepId, JobStepState.Unspecified, JobStepOutcome.Unspecified, null, null, null, triggerId, null, null, null) == null)
+				if (await _jobService.UpdateStepAsync(job, batchId, stepId, JobStepState.Unspecified, JobStepOutcome.Unspecified, newNotificationTriggerId: triggerId) == null)
 				{
 					return NotFound(jobId, batchId, stepId);
 				}
diff --git a/Engine/Source/Programs/Horde/Horde.Build/Jobs/IJob.cs b/Engine/Source/Programs/Horde/Horde.Build/Jobs/IJob.cs
--- a/Engine/Source/Programs/Horde/Horde.Build/Jobs/IJob.cs
+++ b/Engine/Source/Programs/Horde/Horde.Build/Jobs/IJob.cs
@@ -85,6 +85,11 @@ namespace Horde.Build.Models
 		/// </summary>
 		public JobStepOutcome Outcome { get; }
 
+		/// <summary>
+		/// Error from executing this step
+		/// </summary>
+		public JobStepError Error { get; }
+
 		/// <summary>
 		/// The log id for this step
 		/// </summary>
@@ -141,6 +146,22 @@ namespace Horde.Build.Models
 		public Dictionary<string, string>? Properties { get; }
 	}
 
+	/// <summary>
+	/// Systemic error codes for a job failing
+	/// </summary>
+	public enum JobStepError
+	{
+		/// <summary>
+		/// No systemic error
+		/// </summary>
+		None,
+
+		/// <summary>
+		/// Step did not complete in the required amount of time
+		/// </summary>
+		TimedOut,
+	}
+
 	/// <summary>
 	/// Extension methods for job steps
 	/// </summary>
@@ -163,6 +184,25 @@ namespace Horde.Build.Models
 		{
 			return step.State != JobStepState.Aborted && step.State != JobStepState.Completed && step.State != JobStepState.Skipped;
 		}
+
+		/// <summary>
+		/// Determine if a step should be timed out
+		/// </summary>
+		/// <param name="step"></param>
+		/// <param name="utcNow"></param>
+		/// <returns></returns>
+		public static bool HasTimedOut(this IJobStep step, DateTime utcNow)
+		{
+			if (step.State == JobStepState.Running && step.StartTimeUtc != null)
+			{
+				TimeSpan elapsed = utcNow - step.StartTimeUtc.Value;
+				if (elapsed > TimeSpan.FromHours(24.0))
+				{
+					return true;
+				}
+			}
+			return false;
+		}
 	}
 
 	/// <summary>
diff --git a/Engine/Source/Programs/Horde/Horde.Build/Jobs/IJobCollection.cs b/Engine/Source/Programs/Horde/Horde.Build/Jobs/IJobCollection.cs
--- a/Engine/Source/Programs/Horde/Horde.Build/Jobs/IJobCollection.cs
+++ b/Engine/Source/Programs/Horde/Horde.Build/Jobs/IJobCollection.cs
@@ -157,6 +157,7 @@ namespace Horde.Build.Collections
 		/// <param name="stepId">Unique id of the step to update</param>
 		/// <param name="newState">New state of the jobstep</param>
 		/// <param name="newOutcome">New outcome of the jobstep</param>
+		/// <param name="newError">New error annotation for this jobstep</param>
 		/// <param name="newAbortRequested">New state of request abort</param>
 		/// <param name="newAbortByUserId">New name of user that requested the abort</param>
 		/// <param name="newLogId">New log id for the jobstep</param>
@@ -166,7 +167,7 @@ namespace Horde.Build.Collections
 		/// <param name="newReports">New report documents</param>
 		/// <param name="newProperties">Property changes. Any properties with a null value will be removed.</param>
 		/// <returns>True if the job was updated, false if it was deleted in the meantime</returns>
-		Task<IJob?> TryUpdateStepAsync(IJob job, IGraph graph, SubResourceId batchId, SubResourceId stepId, JobStepState newState = default, JobStepOutcome newOutcome = default, bool? newAbortRequested = null, UserId? newAbortByUserId = null, LogId? newLogId = null, ObjectId? newNotificationTriggerId = null, UserId? newRetryByUserId = null, Priority? newPriority = null, List<Report>? newReports = null, Dictionary<string, string?>? newProperties = null);
+		Task<IJob?> TryUpdateStepAsync(IJob job, IGraph graph, SubResourceId batchId, SubResourceId stepId, JobStepState newState = default, JobStepOutcome newOutcome = default, JobStepError? newError = null, bool? newAbortRequested = null, UserId? newAbortByUserId = null, LogId? newLogId = null, ObjectId? newNotificationTriggerId = null, UserId? newRetryByUserId = null, Priority? newPriority = null, List<Report>? newReports = null, Dictionary<string, string?>? newProperties = null);
 
 		/// <summary>
 		/// Attempts to update the node groups to be executed for a job. Fails if another write happens in the meantime.
diff --git a/Engine/Source/Programs/Horde/Horde.Build/Jobs/Impl/JobCollection.cs b/Engine/Source/Programs/Horde/Horde.Build/Jobs/Impl/JobCollection.cs
--- a/Engine/Source/Programs/Horde/Horde.Build/Jobs/Impl/JobCollection.cs
+++ b/Engine/Source/Programs/Horde/Horde.Build/Jobs/Impl/JobCollection.cs
@@ -56,6 +56,8 @@ namespace Horde.Build.Collections.Impl
 
 			public JobStepOutcome Outcome { get; set; } = JobStepOutcome.Success;
 
+			public JobStepError Error { get; set; } = JobStepError.None;
+
 			[BsonIgnoreIfNull]
 			public LogId? LogId { get; set; }
 
@@ -735,7 +737,7 @@ namespace Horde.Build.Collections.Impl
 		}
 
 		/// <inheritdoc/>
-		public Task<IJob?> TryUpdateStepAsync(IJob job, IGraph graph, SubResourceId batchId, SubResourceId stepId, JobStepState newState, JobStepOutcome newOutcome, bool? newAbortRequested, UserId? newAbortByUserId, LogId? newLogId, ObjectId? newNotificationTriggerId, UserId? newRetryByUserId, Priority? newPriority, List<Report>? newReports, Dictionary<string, string?>? newProperties)
+		public Task<IJob?> TryUpdateStepAsync(IJob job, IGraph graph, SubResourceId batchId, SubResourceId stepId, JobStepState newState, JobStepOutcome newOutcome, JobStepError? newError, bool? newAbortRequested, UserId? newAbortByUserId, LogId? newLogId, ObjectId? newNotificationTriggerId, UserId? newRetryByUserId, Priority? newPriority, List<Report>? newReports, Dictionary<string, string?>? newProperties)
 		{
 			JobDocument jobDocument = Clone((JobDocument)job);
 
@@ -813,6 +815,13 @@ namespace Horde.Build.Collections.Impl
 								bRefreshDependentJobSteps = true;
 							}
 
+							// Update the job step error
+							if (newError != null)
+							{
+								step.Error = newError.Value;
+								updates.Add(updateBuilder.Set(x => x.Batches[batchIdx].Steps[stepIdx].Error, step.Error));
+							}
+
 							// Update the log id
 							if (newLogId != null && step.LogId != newLogId.Value)
 							{
diff --git a/Engine/Source/Programs/Horde/Horde.Build/Jobs/Impl/JobService.cs b/Engine/Source/Programs/Horde/Horde.Build/Jobs/Impl/JobService.cs
--- a/Engine/Source/Programs/Horde/Horde.Build/Jobs/Impl/JobService.cs
+++ b/Engine/Source/Programs/Horde/Horde.Build/Jobs/Impl/JobService.cs
@@ -857,6 +857,7 @@ namespace Horde.Build.Services
 		/// <param name="stepId">Unique id of the step to update</param>
 		/// <param name="newState">New state of the jobstep</param>
 		/// <param name="newOutcome">New outcome of the jobstep</param>
+		/// <param name="newError">New error for the step</param>
 		/// <param name="newAbortRequested">New state of abort request</param>
 		/// <param name="newAbortByUserId">New user that requested the abort</param>
 		/// <param name="newLogId">New log id for the jobstep</param>
@@ -866,7 +867,7 @@ namespace Horde.Build.Services
 		/// <param name="newReports">New list of reports</param>
 		/// <param name="newProperties">Property changes. Any properties with a null value will be removed.</param>
 		/// <returns>True if the job was updated, false if it was deleted in the meantime</returns>
-		public async Task<IJob?> UpdateStepAsync(IJob job, SubResourceId batchId, SubResourceId stepId, JobStepState newState = JobStepState.Unspecified, JobStepOutcome newOutcome = JobStepOutcome.Unspecified, bool? newAbortRequested = null, UserId? newAbortByUserId = null, LogId? newLogId = null, ObjectId? newNotificationTriggerId = null, UserId? newRetryByUserId = null, Priority? newPriority = null, List<Report>? newReports = null, Dictionary<string, string?>? newProperties = null)
+		public async Task<IJob?> UpdateStepAsync(IJob job, SubResourceId batchId, SubResourceId stepId, JobStepState newState = JobStepState.Unspecified, JobStepOutcome newOutcome = JobStepOutcome.Unspecified, JobStepError? newError = null, bool? newAbortRequested = null, UserId? newAbortByUserId = null, LogId? newLogId = null, ObjectId? newNotificationTriggerId = null, UserId? newRetryByUserId = null, Priority? newPriority = null, List<Report>? newReports = null, Dictionary<string, string?>? newProperties = null)
 		{
 			using IScope traceScope = GlobalTracer.Instance.BuildSpan("JobService.UpdateStepAsync").StartActive();
 			traceScope.Span.SetTag("Job", job.Id);
@@ -876,7 +877,7 @@ namespace Horde.Build.Services
 			using IDisposable scope = _logger.BeginScope("UpdateStepAsync({JobId})", job.Id);
 			for (; ;)
 			{
-				IJob? newJob = await TryUpdateStepAsync(job, batchId, stepId, newState, newOutcome, newAbortRequested, newAbortByUserId, newLogId, newNotificationTriggerId, newRetryByUserId, newPriority, newReports, newProperties);
+				IJob? newJob = await TryUpdateStepAsync(job, batchId, stepId, newState, newOutcome, newError, newAbortRequested, newAbortByUserId, newLogId, newNotificationTriggerId, newRetryByUserId, newPriority, newReports, newProperties);
 				if (newJob != null)
 				{
 					return newJob;
@@ -900,6 +901,7 @@ namespace Horde.Build.Services
 		/// <param name="stepId">Unique id of the step to update</param>
 		/// <param name="newState">New state of the jobstep</param>
 		/// <param name="newOutcome">New outcome of the jobstep</param>
+		/// <param name="newError">New error for the step</param>
 		/// <param name="newAbortRequested">New state for request abort</param>
 		/// <param name="newAbortByUserId">New name of user that requested the abort</param>
 		/// <param name="newLogId">New log id for the jobstep</param>
@@ -909,7 +911,7 @@ namespace Horde.Build.Services
 		/// <param name="newReports">New reports</param>
 		/// <param name="newProperties">Property changes. Any properties with a null value will be removed.</param>
 		/// <returns>True if the job was updated, false if it was deleted in the meantime</returns>
-		public async Task<IJob?> TryUpdateStepAsync(IJob job, SubResourceId batchId, SubResourceId stepId, JobStepState newState = JobStepState.Unspecified, JobStepOutcome newOutcome = JobStepOutcome.Unspecified, bool? newAbortRequested = null, UserId? newAbortByUserId = null, LogId? newLogId = null, ObjectId? newTriggerId = null, UserId? newRetryByUserId = null, Priority? newPriority = null, List<Report>? newReports = null, Dictionary<string, string?>? newProperties = null)
+		public async Task<IJob?> TryUpdateStepAsync(IJob job, SubResourceId batchId, SubResourceId stepId, JobStepState newState = JobStepState.Unspecified, JobStepOutcome newOutcome = JobStepOutcome.Unspecified, JobStepError? newError = null, bool? newAbortRequested = null, UserId? newAbortByUserId = null, LogId? newLogId = null, ObjectId? newTriggerId = null, UserId? newRetryByUserId = null, Priority? newPriority = null, List<Report>? newReports = null, Dictionary<string, string?>? newProperties = null)
 		{
 			using IScope traceScope = GlobalTracer.Instance.BuildSpan("JobService.TryUpdateStepAsync").StartActive();
 			traceScope.Span.SetTag("Job", job.Id);
@@ -941,7 +943,7 @@ namespace Horde.Build.Services
 			}
 
 			// Update the step
-			IJob? newJob = await _jobs.TryUpdateStepAsync(job, graph, batchId, stepId, newState, newOutcome, newAbortRequested, newAbortByUserId, newLogId, newTriggerId, newRetryByUserId, newPriority, newReports, newProperties);
+			IJob? newJob = await _jobs.TryUpdateStepAsync(job, graph, batchId, stepId, newState, newOutcome, newError, newAbortRequested, newAbortByUserId, newLogId, newTriggerId, newRetryByUserId, newPriority, newReports, newProperties);
 			if (newJob != null)
 			{
 				job = newJob;
diff --git a/Engine/Source/Programs/Horde/Horde.Build/Services/RpcService.cs b/Engine/Source/Programs/Horde/Horde.Build/Services/RpcService.cs
--- a/Engine/Source/Programs/Horde/Horde.Build/Services/RpcService.cs
+++ b/Engine/Source/Programs/Horde/Horde.Build/Services/RpcService.cs
@@ -60,12 +60,13 @@ namespace Horde.Build.Services
 		readonly ITestDataCollection _testData;
 		readonly IJobStepRefCollection _jobStepRefCollection;
 		readonly ConformTaskSource _conformTaskSource;
+		readonly IClock _clock;
 		readonly ILogger<RpcService> _logger;
 
 		/// <summary>
 		/// Constructor
 		/// </summary>
-		public RpcService(AclService aclService, AgentService agentService, StreamService streamService, JobService jobService, AgentSoftwareService agentSoftwareService, IArtifactCollection artifactCollection, ILogFileService logFileService, PoolService poolService, LifetimeService lifetimeService, IGraphCollection graphs, ITestDataCollection testData, IJobStepRefCollection jobStepRefCollection, ConformTaskSource conformTaskSource, ILogger<RpcService> logger)
+		public RpcService(AclService aclService, AgentService agentService, StreamService streamService, JobService jobService, AgentSoftwareService agentSoftwareService, IArtifactCollection artifactCollection, ILogFileService logFileService, PoolService poolService, LifetimeService lifetimeService, IGraphCollection graphs, ITestDataCollection testData, IJobStepRefCollection jobStepRefCollection, ConformTaskSource conformTaskSource, IClock clock, ILogger<RpcService> logger)
 		{
 			_aclService = aclService;
 			_agentService = agentService;
@@ -80,6 +81,7 @@ namespace Horde.Build.Services
 			_testData = testData;
 			_jobStepRefCollection = jobStepRefCollection;
 			_conformTaskSource = conformTaskSource;
+			_clock = clock;
 			_logger = logger;
 		}
 
@@ -586,7 +588,7 @@ namespace Horde.Build.Services
 			//				}
 
 			// Update the step state
-			IJob? newJob = await _jobService.TryUpdateStepAsync(job, batch.Id, step.Id, JobStepState.Running, JobStepOutcome.Unspecified, null, null, log.Value.Id, null, null, null, null);
+			IJob? newJob = await _jobService.TryUpdateStepAsync(job, batch.Id, step.Id, JobStepState.Running, JobStepOutcome.Unspecified, newLogId: log.Value.Id);
 			if (newJob != null)
 			{
 				BeginStepResponse response = new BeginStepResponse();
@@ -678,7 +680,19 @@ namespace Horde.Build.Services
 
 			IJobStepBatch batch = AuthorizeBatch(job, request.BatchId.ToSubResourceId(), context);
 
-			await _jobService.UpdateStepAsync(job, batch.Id, request.StepId.ToSubResourceId(), request.State, request.Outcome, null, null, null, null, null);
+			SubResourceId stepId = request.StepId.ToSubResourceId();
+			if (!batch.TryGetStep(stepId, out IJobStep? step))
+			{
+				throw new StructuredRpcException(StatusCode.NotFound, "Unable to find step {JobId}:{BatchId}:{StepId}", job.Id, batch.Id, stepId);
+			}
+
+			JobStepError? error = null;
+			if (step.Error == JobStepError.None && step.HasTimedOut(_clock.UtcNow))
+			{
+				error = JobStepError.TimedOut;
+			}
+
+			await _jobService.UpdateStepAsync(job, batch.Id, request.StepId.ToSubResourceId(), request.State, request.Outcome, error, null, null, null, null, null);
 			return new Empty();
 		}
 
@@ -699,7 +713,7 @@ namespace Horde.Build.Services
 				throw new StructuredRpcException(StatusCode.NotFound, "Unable to find step {JobId}:{BatchId}:{StepId}", job.Id, batch.Id, stepId);
 			}
 
-			return new GetStepResponse { Outcome = step.Outcome, State = step.State, AbortRequested = step.AbortRequested };
+			return new GetStepResponse { Outcome = step.Outcome, State = step.State, AbortRequested = step.AbortRequested || step.HasTimedOut(_clock.UtcNow) };
 		}
 
 		/// <summary>
