Horde: Include a link to the issue thread in the periodic summary.

diff --git a/Engine/Source/Programs/Horde/Horde.Build/Notifications/Impl/SlackNotificationSink.cs b/Engine/Source/Programs/Horde/Horde.Build/Notifications/Impl/SlackNotificationSink.cs
--- a/Engine/Source/Programs/Horde/Horde.Build/Notifications/Impl/SlackNotificationSink.cs
+++ b/Engine/Source/Programs/Horde/Horde.Build/Notifications/Impl/SlackNotificationSink.cs
@@ -46,6 +46,7 @@ namespace Horde.Build.Notifications.Impl
 	{
 		const string PostMessageUrl = "https://slack.com/api/chat.postMessage";
 		const string UpdateMessageUrl = "https://slack.com/api/chat.update";
+		const string GetPermalinkUrl = "https://slack.com/api/chat.getPermalink";
 
 		class SocketResponse
 		{
@@ -147,6 +148,9 @@ namespace Horde.Build.Notifications.Impl
 
 			[BsonElement("dig")]
 			public string Digest { get; set; } = String.Empty;
+
+			[BsonElement("lnk"), BsonIgnoreIfNull]
+			public string? Permalink { get; set; }
 		}
 
 		class SlackUser : IAvatar
@@ -268,15 +272,29 @@ namespace Horde.Build.Notifications.Impl
 			return (state, state.Id == newId);
 		}
 
-		async Task SetMessageTimestampAsync(ObjectId messageId, string channel, string ts)
+		async Task<MessageStateDocument?> GetMessageStateAsync(string recipient, string eventId)
+		{
+			FilterDefinition<MessageStateDocument> filter = Builders<MessageStateDocument>.Filter.Eq(x => x.Recipient, recipient) & Builders<MessageStateDocument>.Filter.Eq(x => x.EventId, eventId);
+			return await _messageStates.Find(filter).FirstOrDefaultAsync();
+		}
+
+		async Task SetMessageTimestampAsync(ObjectId messageId, string channel, string ts, string? permalink = null)
 		{
 			FilterDefinition<MessageStateDocument> filter = Builders<MessageStateDocument>.Filter.Eq(x => x.Id, messageId);
 			UpdateDefinition<MessageStateDocument> update = Builders<MessageStateDocument>.Update.Set(x => x.Channel, channel).Set(x => x.Ts, ts);
+			if (permalink == null)
+			{
+				update = update.Unset(x => x.Permalink);
+			}
+			else
+			{
+				update = update.Set(x => x.Permalink, permalink);
+			}
 			await _messageStates.FindOneAndUpdateAsync(filter, update);
 		}
 
 		#endregion
-		
+
 		/// <inheritdoc/>
 		public async Task NotifyJobScheduledAsync(List<JobScheduledNotification> notifications)
 		{
@@ -664,6 +682,8 @@ namespace Horde.Build.Notifications.Impl
 			}
 		}
 
+		static string GetTriageThreadEventId(int issueId) => $"issue_triage_{issueId}";
+
 		async Task SendTriageMessageAsync(IIssue issue, IIssueSpan span, WorkflowConfig workflow)
 		{
 			string? triageChannel = workflow.TriageChannel;
@@ -675,7 +695,7 @@ namespace Horde.Build.Notifications.Impl
 				message.Channel = triageChannel;
 				message.Text = $"{workflow.TriagePrefix}*Issue <{issueUrl}|{issue.Id}>*: {issue.Summary}{workflow.TriageSuffix}";
 
-				string eventId = $"issue_triage_{issue.Id}";
+				string eventId = GetTriageThreadEventId(issue.Id);
 
 				(MessageStateDocument state, bool isNew) = await AddOrUpdateMessageStateAsync(triageChannel, eventId, null, "");
 				if (isNew)
@@ -689,7 +709,13 @@ namespace Horde.Build.Notifications.Impl
 					state.Channel = response.Channel ?? String.Empty;
 					state.Ts = response.Ts ?? String.Empty;
 
-					await SetMessageTimestampAsync(state.Id, state.Channel, state.Ts);
+					string? permalink = null;
+					if (response.Channel != null && response.Ts != null)
+					{
+						permalink = await GetPermalinkAsync(response.Channel, response.Ts);
+					}
+
+					await SetMessageTimestampAsync(state.Id, state.Channel, state.Ts, permalink);
 
 					// Create the summary text
 					StringBuilder summary = new StringBuilder();
@@ -1085,6 +1111,16 @@ namespace Horde.Build.Notifications.Impl
 					}
 
 					body.Append($"\n\t\u2022 *Issue <{issueUrl}|{issue.Id}>*: {issue.Summary} [{FormatReadableTimeSpan(report.Time - issue.CreatedAt)}] - {status}");
+
+					string? triageChannel = report.Workflow.TriageChannel;
+					if (triageChannel != null)
+					{
+						MessageStateDocument? state = await GetMessageStateAsync(triageChannel, GetTriageThreadEventId(issue.Id));
+						if (state != null && state.Permalink != null)
+						{
+							body.Append($" (<{state.Permalink}|View Thread>)");
+						}
+					}
 				}
 
 				if (report.NumSteps > 0)
@@ -1408,6 +1444,24 @@ namespace Horde.Build.Notifications.Impl
 			public string? Ts { get; set; }
 		}
 
+		class GetPermalinkRequest
+		{
+			[JsonPropertyName("channel")]
+			public string? Channel { get; set; }
+
+			[JsonPropertyName("message_ts")]
+			public string? MessageTs { get; set; }
+		}
+
+		class GetPermalinkResponse : SlackResponse
+		{
+			[JsonPropertyName("channel")]
+			public string? Channel { get; set; }
+
+			[JsonPropertyName("permalink")]
+			public string? Permalink { get; set; }
+		}
+
 		[return: NotNullIfNotNull("message")]
 		private string? AddEnvironmentAnnotation(string? message)
 		{
@@ -1500,30 +1554,45 @@ namespace Horde.Build.Notifications.Impl
 			}
 		}
 
+		private async Task<string?> GetPermalinkAsync(string channel, string messageTs)
+		{
+			string requestUrl = $"{GetPermalinkUrl}?channel={channel}&message_ts={messageTs}";
+			using (HttpRequestMessage sendMessageRequest = new HttpRequestMessage(HttpMethod.Get, requestUrl))
+			{
+				GetPermalinkResponse response = await SendRequestAsync<GetPermalinkResponse>(sendMessageRequest, null);
+				return response?.Permalink;
+			}
+		}
+
 		private async Task<TResponse> SendRequestAsync<TResponse>(string requestUrl, object request) where TResponse : SlackResponse
 		{
-			using (HttpClient client = new HttpClient())
+			using (HttpRequestMessage sendMessageRequest = new HttpRequestMessage(HttpMethod.Post, requestUrl))
 			{
-				using (HttpRequestMessage sendMessageRequest = new HttpRequestMessage(HttpMethod.Post, requestUrl))
+				string requestJson = JsonSerializer.Serialize(request, new JsonSerializerOptions { DefaultIgnoreCondition = JsonIgnoreCondition.WhenWritingNull });
+				using (StringContent messageContent = new StringContent(requestJson, Encoding.UTF8, "application/json"))
 				{
-					string requestJson = JsonSerializer.Serialize(request, new JsonSerializerOptions { DefaultIgnoreCondition = JsonIgnoreCondition.WhenWritingNull });
-					using (StringContent messageContent = new StringContent(requestJson, Encoding.UTF8, "application/json"))
-					{
-						sendMessageRequest.Content = messageContent;
-						sendMessageRequest.Headers.Add("Authorization", $"Bearer {_settings.SlackToken ?? ""}");
+					sendMessageRequest.Content = messageContent;
+					return await SendRequestAsync<TResponse>(sendMessageRequest, requestJson);
+				}
+			}
+		}
 
-						HttpResponseMessage response = await client.SendAsync(sendMessageRequest);
-						byte[] responseBytes = await response.Content.ReadAsByteArrayAsync();
+		private async Task<TResponse> SendRequestAsync<TResponse>(HttpRequestMessage request, string? requestJson) where TResponse : SlackResponse
+		{
+			using (HttpClient client = new HttpClient())
+			{
+				request.Headers.Add("Authorization", $"Bearer {_settings.SlackToken ?? ""}");
 
-						TResponse responseObject = JsonSerializer.Deserialize<TResponse>(responseBytes)!;
-						if(!responseObject.Ok)
-						{
-							_logger.LogError("Failed to send Slack message ({Error}). Request: {Request}. Response: {Response}", responseObject.Error, requestJson, Encoding.UTF8.GetString(responseBytes));
-						}
+				HttpResponseMessage response = await client.SendAsync(request);
+				byte[] responseBytes = await response.Content.ReadAsByteArrayAsync();
 
-						return responseObject;
-					}
+				TResponse responseObject = JsonSerializer.Deserialize<TResponse>(responseBytes)!;
+				if (!responseObject.Ok)
+				{
+					_logger.LogError("Failed to send Slack message ({Error}). Request: {Request}. Response: {Response}", responseObject.Error, requestJson, Encoding.UTF8.GetString(responseBytes));
 				}
+
+				return responseObject;
 			}
 		}
 
