Horde: Add UGS URI handler links to changelist numbers.

diff --git a/Engine/Source/Programs/Horde/Horde.Build/Notifications/Impl/SlackNotificationSink.cs b/Engine/Source/Programs/Horde/Horde.Build/Notifications/Impl/SlackNotificationSink.cs
--- a/Engine/Source/Programs/Horde/Horde.Build/Notifications/Impl/SlackNotificationSink.cs
+++ b/Engine/Source/Programs/Horde/Horde.Build/Notifications/Impl/SlackNotificationSink.cs
@@ -36,6 +36,7 @@ using MongoDB.Driver;
 
 namespace Horde.Build.Notifications.Impl
 {
+	using JobId = ObjectId<IJob>;
 	using LogId = ObjectId<ILogFile>;
 	using UserId = ObjectId<IUser>;
 
@@ -718,8 +719,8 @@ namespace Horde.Build.Notifications.Impl
 					await SetMessageTimestampAsync(state.Id, state.Channel, state.Ts, permalink);
 
 					// Create the summary text
-					StringBuilder summary = new StringBuilder();
-					summary.AppendLine($"From <{GetStepUrl(issue, span.FirstFailure)}|{span.FirstFailure.JobName}: {span.NodeName}>.");
+					List<ILogEvent> events = new List<ILogEvent>();
+					List<ILogEventData> eventDataItems = new List<ILogEventData>();
 
 					if (span.FirstFailure.LogId != null)
 					{
@@ -727,7 +728,7 @@ namespace Horde.Build.Notifications.Impl
 						ILogFile? logFile = await _logFileService.GetLogFileAsync(logId);
 						if (logFile != null)
 						{
-							List<ILogEvent> events = await _logFileService.FindEventsAsync(logFile, span.Id, 0, 20);
+							events = await _logFileService.FindEventsAsync(logFile, span.Id, 0, 50);
 							if (events.Any(x => x.Severity == EventSeverity.Error))
 							{
 								events.RemoveAll(x => x.Severity == EventSeverity.Warning);
@@ -737,14 +738,21 @@ namespace Horde.Build.Notifications.Impl
 							for (int idx = 0; idx < Math.Min(events.Count, 3); idx++)
 							{
 								ILogEventData data = await _logFileService.GetEventDataAsync(logFile, events[idx].LineIndex, events[idx].LineCount);
-								summary.AppendLine($"```{data.Message}```");
-							}
-							if (events.Count > 3)
-							{
-								summary.AppendLine("```...```");
+								eventDataItems.Add(data);
 							}
 						}
 					}
+
+					StringBuilder summary = new StringBuilder();
+					summary.AppendLine($"From <{GetJobUrl(span.FirstFailure.JobId)}|{span.FirstFailure.JobName}> / <{GetStepUrl(span.FirstFailure.JobId, span.FirstFailure.StepId)}|{span.NodeName}>:");
+					foreach (ILogEventData eventDataItem in eventDataItems)
+					{
+						summary.AppendLine($"```{eventDataItem.Message}```");
+					}
+					if (events.Count > eventDataItems.Count)
+					{
+						summary.AppendLine("```...```");
+					}
 					await SendMessageToThread(triageChannel, state.Ts, summary.ToString());
 
 					List<IIssueSuspect> suspects = await _issueService.GetIssueSuspectsAsync(issue);
@@ -756,7 +764,7 @@ namespace Horde.Build.Notifications.Impl
 						foreach (IGrouping<UserId, IIssueSuspect> suspectGroup in suspectGroups)
 						{
 							string mention = await FormatMentionAsync(suspectGroup.Key);
-							string changes = String.Join(", ", suspectGroup.Select(x => $"CL {x.Change}"));
+							string changes = String.Join(", ", suspectGroup.Select(x => $"<ugs://change?number={x.Change}|CL {x.Change}>"));
 							suspectList.Add($"{mention} ({changes})");
 						}
 
@@ -789,9 +797,14 @@ namespace Horde.Build.Notifications.Impl
 			await SendIssueMessageAsync(slackUserId, issue, details, user.Id);
 		}
 
-		Uri GetStepUrl(IIssue issue, IIssueStep step)
+		Uri GetJobUrl(JobId jobId)
+		{
+			return new Uri(_settings.DashboardUrl, $"job/{jobId}");
+		}
+
+		Uri GetStepUrl(JobId jobId, SubResourceId stepId)
 		{
-			return new Uri(_settings.DashboardUrl, $"job/{step.JobId}?step={step.StepId}");
+			return new Uri(_settings.DashboardUrl, $"job/{jobId}?step={stepId}");
 		}
 
 		Uri GetIssueUrl(IIssue issue, IIssueStep step)
@@ -1077,7 +1090,7 @@ namespace Horde.Build.Notifications.Impl
 				else
 				{
 					TimeSpan averageAge = TimeSpan.FromHours(report.Issues.Select(x => (report.Time - x.CreatedAt).TotalHours).Average());
-					body.Append($"*{report.Issues.Count} issues* are currently open (average age {FormatReadableTimeSpan(averageAge)}):");
+					body.Append($"*{report.Issues.Count} issues* currently open (average age {FormatReadableTimeSpan(averageAge)}):");
 				}
 				body.Append('\n');
 
