Horde: Tidy up formatting for stream summaries. Slack splits long messages after a certain number of characters then strips leading spaces, which breaks tabulation.

diff --git a/Engine/Source/Programs/Horde/Horde.Build/Notifications/Impl/SlackNotificationSink.cs b/Engine/Source/Programs/Horde/Horde.Build/Notifications/Impl/SlackNotificationSink.cs
--- a/Engine/Source/Programs/Horde/Horde.Build/Notifications/Impl/SlackNotificationSink.cs
+++ b/Engine/Source/Programs/Horde/Horde.Build/Notifications/Impl/SlackNotificationSink.cs
@@ -1123,17 +1123,19 @@ namespace Horde.Build.Notifications.Impl
 						}
 					}
 
-					body.Append($"\n\t\u2022 *Issue <{issueUrl}|{issue.Id}>*: {issue.Summary} [{FormatReadableTimeSpan(report.Time - issue.CreatedAt)}] - {status}");
-
+					body.Append($"\n\u2022 *Issue <{issueUrl}|{issue.Id}>");
+					
 					string? triageChannel = report.Workflow.TriageChannel;
 					if (triageChannel != null)
 					{
 						MessageStateDocument? state = await GetMessageStateAsync(triageChannel, GetTriageThreadEventId(issue.Id));
 						if (state != null && state.Permalink != null)
 						{
-							body.Append($" (<{state.Permalink}|View Thread>)");
+							body.Append($" (<{state.Permalink}|Thread>)");
 						}
 					}
+
+					body.Append($"*: {issue.Summary} [{FormatReadableTimeSpan(report.Time - issue.CreatedAt)}] - {status}");
 				}
 
 				if (report.NumSteps > 0)
