Horde: Link to message inside the triage channel thread rather than the initial message.

diff --git a/Engine/Source/Programs/Horde/Horde.Build/Notifications/Impl/SlackNotificationSink.cs b/Engine/Source/Programs/Horde/Horde.Build/Notifications/Impl/SlackNotificationSink.cs
--- a/Engine/Source/Programs/Horde/Horde.Build/Notifications/Impl/SlackNotificationSink.cs
+++ b/Engine/Source/Programs/Horde/Horde.Build/Notifications/Impl/SlackNotificationSink.cs
@@ -710,14 +710,6 @@ namespace Horde.Build.Notifications.Impl
 					state.Channel = response.Channel ?? String.Empty;
 					state.Ts = response.Ts ?? String.Empty;
 
-					string? permalink = null;
-					if (response.Channel != null && response.Ts != null)
-					{
-						permalink = await GetPermalinkAsync(response.Channel, response.Ts);
-					}
-
-					await SetMessageTimestampAsync(state.Id, state.Channel, state.Ts, permalink);
-
 					// Create the summary text
 					List<ILogEvent> events = new List<ILogEvent>();
 					List<ILogEventData> eventDataItems = new List<ILogEventData>();
@@ -753,7 +745,14 @@ namespace Horde.Build.Notifications.Impl
 					{
 						summary.AppendLine("```...```");
 					}
-					await SendMessageToThread(triageChannel, state.Ts, summary.ToString());
+					string? summaryTs = await SendMessageToThread(triageChannel, state.Ts, summary.ToString());
+
+					// Permalink to the summary text so we link inside the thread rather than just to the original message
+					string? permalink = null;
+					if (response.Channel != null && summaryTs != null)
+					{
+						permalink = await GetPermalinkAsync(response.Channel, summaryTs);
+					}
 
 					List<IIssueSuspect> suspects = await _issueService.GetIssueSuspectsAsync(issue);
 					IGrouping<UserId, IIssueSuspect>[] suspectGroups = suspects.GroupBy(x => x.AuthorId).ToArray();
@@ -771,6 +770,8 @@ namespace Horde.Build.Notifications.Impl
 						string suspectMessage = $"Possibly {StringUtils.FormatList(suspectList, "or")}.";
 						await SendMessageToThread(triageChannel, state.Ts, suspectMessage);
 					}
+
+					await SetMessageTimestampAsync(state.Id, state.Channel, state.Ts, permalink);
 				}
 			}
 		}
