Horde: Fix annotations not being passed on to graph definitions.

diff --git a/Engine/Source/Programs/Horde/Horde.Build/Services/RpcService.cs b/Engine/Source/Programs/Horde/Horde.Build/Services/RpcService.cs
--- a/Engine/Source/Programs/Horde/Horde.Build/Services/RpcService.cs
+++ b/Engine/Source/Programs/Horde/Horde.Build/Services/RpcService.cs
@@ -730,7 +730,7 @@ namespace Horde.Build.Services
 				List<NewNode> newNodes = new List<NewNode>();
 				foreach (CreateNodeRequest node in group.Nodes)
 				{
-					NewNode newNode = new NewNode(node.Name, node.InputDependencies.ToList(), node.OrderDependencies.ToList(), node.Priority, node.AllowRetry, node.RunEarly, node.Warnings, new Dictionary<string, string>(node.Credentials), new Dictionary<string, string>(node.Properties));
+					NewNode newNode = new NewNode(node.Name, node.InputDependencies.ToList(), node.OrderDependencies.ToList(), node.Priority, node.AllowRetry, node.RunEarly, node.Warnings, new Dictionary<string, string>(node.Credentials), new Dictionary<string, string>(node.Properties), new CaseInsensitiveDictionary<string>(node.Annotations));
 					newNodes.Add(newNode);
 				}
 				newGroups.Add(new NewGroup(group.AgentType, newNodes));
