Horde: Serialize annotations as an array of documents to handle invalid Bson property names.

diff --git a/Engine/Source/Programs/Horde/Horde.Build/Jobs/Impl/GraphCollection.cs b/Engine/Source/Programs/Horde/Horde.Build/Jobs/Impl/GraphCollection.cs
--- a/Engine/Source/Programs/Horde/Horde.Build/Jobs/Impl/GraphCollection.cs
+++ b/Engine/Source/Programs/Horde/Horde.Build/Jobs/Impl/GraphCollection.cs
@@ -14,6 +14,7 @@ using Horde.Build.Utilities;
 using HordeCommon;
 using MongoDB.Bson;
 using MongoDB.Bson.Serialization.Attributes;
+using MongoDB.Bson.Serialization.Options;
 using MongoDB.Driver;
 
 namespace Horde.Build.Collections.Impl
@@ -50,6 +51,7 @@ namespace Horde.Build.Collections.Impl
 			public Dictionary<string, string>? Properties { get; set; }
 
 			[BsonIgnoreIfNull]
+			[BsonDictionaryOptions(DictionaryRepresentation.ArrayOfDocuments)]
 			public CaseInsensitiveDictionary<string>? Annotations { get; set; }
 
 			IReadOnlyDictionary<string, string>? INode.Credentials => Credentials;
