Horde: Include the workflow id on returned issue spans.

diff --git a/Engine/Source/Programs/Horde/Horde.Build/Api/Issues.cs b/Engine/Source/Programs/Horde/Horde.Build/Api/Issues.cs
--- a/Engine/Source/Programs/Horde/Horde.Build/Api/Issues.cs
+++ b/Engine/Source/Programs/Horde/Horde.Build/Api/Issues.cs
@@ -5,6 +5,7 @@ using System.Collections.Generic;
 using System.Linq;
 using EpicGames.Core;
 using Horde.Build.Models;
+using Horde.Build.Server;
 using Horde.Build.Utilities;
 using MongoDB.Bson;
 
@@ -12,6 +13,7 @@ namespace Horde.Build.Api
 {
 	using StreamId = StringId<IStream>;
 	using TemplateRefId = StringId<TemplateRef>;
+	using WorkflowId = StringId<WorkflowConfig>;
 
 	/// <summary>
 	/// Identifies a particular changelist and job
@@ -95,6 +97,11 @@ namespace Horde.Build.Api
 		/// </summary>
 		public string Name { get; set; }
 
+		/// <summary>
+		/// Workflow that this span belongs to
+		/// </summary>
+		public WorkflowId? WorkflowId { get; set; }
+
 		/// <summary>
 		/// The previous build 
 		/// </summary>
@@ -115,11 +122,13 @@ namespace Horde.Build.Api
 		/// </summary>
 		/// <param name="span">The node to construct from</param>
 		/// <param name="steps">Failing steps for this span</param>
-		public GetIssueSpanResponse(IIssueSpan span, List<IIssueStep> steps)
+		/// <param name="workflowId"></param>
+		public GetIssueSpanResponse(IIssueSpan span, List<IIssueStep> steps, WorkflowId? workflowId)
 		{
 			Id = span.Id.ToString();
 			Name = span.NodeName;
 			TemplateId = span.TemplateRefId.ToString();
+			WorkflowId = workflowId;
 			LastSuccess = (span.LastSuccess != null) ? new GetIssueStepResponse(span.LastSuccess) : null;
 			Steps = steps.ConvertAll(x => new GetIssueStepResponse(x));
 			NextSuccess = (span.NextSuccess != null) ? new GetIssueStepResponse(span.NextSuccess) : null;
@@ -149,17 +158,17 @@ namespace Horde.Build.Api
 		/// <summary>
 		/// Map of steps to (event signature id -> trace id)
 		/// </summary>
-		public List<GetIssueSpanResponse> Nodes { get; set; }
+		public List<GetIssueSpanResponse> Nodes { get; set; } = new List<GetIssueSpanResponse>();
 
 		/// <summary>
 		/// Constructor
 		/// </summary>
-		/// <param name="streamId">The stream to construct from</param>
+		/// <param name="stream">The stream to construct from</param>
 		/// <param name="spans">List of spans for the given stream</param>
 		/// <param name="steps">List of steps for the given stream</param>
-		public GetIssueStreamResponse(StreamId streamId, List<IIssueSpan> spans, List<IIssueStep> steps)
+		public GetIssueStreamResponse(IStream stream, List<IIssueSpan> spans, List<IIssueStep> steps)
 		{
-			StreamId = streamId.ToString();
+			StreamId = stream.Id.ToString();
 
 			foreach (IIssueSpan span in spans)
 			{
@@ -171,9 +180,14 @@ namespace Horde.Build.Api
 				{
 					MaxChange = span.NextSuccess.Change;
 				}
-			}
 
-			Nodes = spans.ConvertAll(x => new GetIssueSpanResponse(x, steps.Where(y => y.SpanId == x.Id).ToList()));
+				WorkflowId? workflowId = null;
+				if (stream.Config.TryGetTemplate(span.TemplateRefId, out TemplateRefConfig? templateConfig))
+				{
+					workflowId = templateConfig.Workflow;
+				}
+				Nodes.Add(new GetIssueSpanResponse(span, steps.Where(y => y.SpanId == span.Id).ToList(), workflowId));
+			}
 		}
 	}
 
diff --git a/Engine/Source/Programs/Horde/Horde.Build/Api/Streams.cs b/Engine/Source/Programs/Horde/Horde.Build/Api/Streams.cs
--- a/Engine/Source/Programs/Horde/Horde.Build/Api/Streams.cs
+++ b/Engine/Source/Programs/Horde/Horde.Build/Api/Streams.cs
@@ -561,6 +561,11 @@ namespace Horde.Build.Api
 		/// </summary>
 		public string? PauseComment { get; set; }
 
+		/// <summary>
+		/// Workflows for this stream
+		/// </summary>
+		public List<WorkflowConfig> Workflows { get; set; }
+
 		/// <summary>
 		/// Constructor
 		/// </summary>
@@ -581,7 +586,8 @@ namespace Horde.Build.Api
 		/// <param name="acl">Permissions for this object</param>
 		/// <param name="pausedUntil">Stream paused for new builds until this date</param>
 		/// <param name="pauseComment">Reason for stream being paused</param>
-		public GetStreamResponse(string id, string projectId, string name, string configPath, string configRevision, int order, string? notificationChannel, string? notificationChannelFilter, string? triageChannel, DefaultPreflightRequest? defaultPreflight, List<GetStreamTabResponse> tabs, Dictionary<string, GetAgentTypeResponse> agentTypes, Dictionary<string, GetWorkspaceTypeResponse>? workspaceTypes, List<GetTemplateRefResponse> templates, GetAclResponse? acl, DateTime? pausedUntil, string? pauseComment)
+		/// <param name="workflows">Workflows for this stream</param>
+		public GetStreamResponse(string id, string projectId, string name, string configPath, string configRevision, int order, string? notificationChannel, string? notificationChannelFilter, string? triageChannel, DefaultPreflightRequest? defaultPreflight, List<GetStreamTabResponse> tabs, Dictionary<string, GetAgentTypeResponse> agentTypes, Dictionary<string, GetWorkspaceTypeResponse>? workspaceTypes, List<GetTemplateRefResponse> templates, GetAclResponse? acl, DateTime? pausedUntil, string? pauseComment, List<WorkflowConfig> workflows)
 		{
 			Id = id;
 			ProjectId = projectId;
@@ -601,6 +607,7 @@ namespace Horde.Build.Api
 			Acl = acl;
 			PausedUntil = pausedUntil;
 			PauseComment = pauseComment;
+			Workflows = workflows;
 		}
 	}
 }
diff --git a/Engine/Source/Programs/Horde/Horde.Build/Controllers/IssuesController.cs b/Engine/Source/Programs/Horde/Horde.Build/Controllers/IssuesController.cs
--- a/Engine/Source/Programs/Horde/Horde.Build/Controllers/IssuesController.cs
+++ b/Engine/Source/Programs/Horde/Horde.Build/Controllers/IssuesController.cs
@@ -7,7 +7,9 @@ using System.Threading.Tasks;
 using Horde.Build.Acls;
 using Horde.Build.Api;
 using Horde.Build.Collections;
+using Horde.Build.Config;
 using Horde.Build.Models;
+using Horde.Build.Server;
 using Horde.Build.Services;
 using Horde.Build.Utilities;
 using HordeCommon;
@@ -30,6 +32,7 @@ namespace Horde.Build.Controllers
 	[Route("[controller]")]
 	public class IssuesController : HordeControllerBase
 	{
+		private readonly ConfigCollection _configCollection;
 		private readonly IIssueCollection _issueCollection;
 		private readonly IIssueService _issueService;
 		private readonly JobService _jobService;
@@ -40,8 +43,9 @@ namespace Horde.Build.Controllers
 		/// <summary>
 		/// Constructor
 		/// </summary>
-		public IssuesController(IIssueCollection issueCollection, IIssueService issueService, JobService jobService, StreamService streamService, IUserCollection userCollection, ILogFileService logFileService)
+		public IssuesController(ConfigCollection configCollection, IIssueCollection issueCollection, IIssueService issueService, JobService jobService, StreamService streamService, IUserCollection userCollection, ILogFileService logFileService)
 		{
+			_configCollection = configCollection;
 			_issueCollection = issueCollection;
 			_issueService = issueService;
 			_jobService = jobService;
@@ -326,9 +330,13 @@ namespace Horde.Build.Controllers
 			{
 				if (await _streamService.AuthorizeAsync(spanGroup.Key, AclAction.ViewStream, User, cache))
 				{
-					HashSet<ObjectId> spanIds = new HashSet<ObjectId>(spanGroup.Select(x => x.Id));
-					List<IIssueStep> steps = issue.Steps.Where(x => spanIds.Contains(x.SpanId)).ToList();
-					responses.Add(PropertyFilter.Apply(new GetIssueStreamResponse(spanGroup.Key, spanGroup.ToList(), steps), filter));
+					IStream? stream = await _streamService.GetCachedStream(spanGroup.Key);
+					if (stream != null)
+					{
+						HashSet<ObjectId> spanIds = new HashSet<ObjectId>(spanGroup.Select(x => x.Id));
+						List<IIssueStep> steps = issue.Steps.Where(x => spanIds.Contains(x.SpanId)).ToList();
+						responses.Add(PropertyFilter.Apply(new GetIssueStreamResponse(stream, spanGroup.ToList(), steps), filter));
+					}
 				}
 			}
 			return responses;
@@ -358,6 +366,12 @@ namespace Horde.Build.Controllers
 				return Forbid();
 			}
 
+			IStream? stream = await _streamService.GetCachedStream(streamIdValue);
+			if (stream == null)
+			{
+				return NotFound();
+			}
+
 			List<IIssueSpan> spans = details.Spans.Where(x => x.StreamId == streamIdValue).ToList();
 			if(spans.Count == 0)
 			{
@@ -367,7 +381,7 @@ namespace Horde.Build.Controllers
 			HashSet<ObjectId> spanIds = new HashSet<ObjectId>(spans.Select(x => x.Id));
 			List<IIssueStep> steps = details.Steps.Where(x => spanIds.Contains(x.SpanId)).ToList();
 
-			return PropertyFilter.Apply(new GetIssueStreamResponse(streamIdValue, spans, steps), filter);
+			return PropertyFilter.Apply(new GetIssueStreamResponse(stream, spans, steps), filter);
 		}
 
 		/// <summary>
diff --git a/Engine/Source/Programs/Horde/Horde.Build/Projects/IStream.cs b/Engine/Source/Programs/Horde/Horde.Build/Projects/IStream.cs
--- a/Engine/Source/Programs/Horde/Horde.Build/Projects/IStream.cs
+++ b/Engine/Source/Programs/Horde/Horde.Build/Projects/IStream.cs
@@ -8,6 +8,7 @@ using System.Text.RegularExpressions;
 using EpicGames.Core;
 using Horde.Build.Acls;
 using Horde.Build.Api;
+using Horde.Build.Server;
 using Horde.Build.Utilities;
 using HordeCommon;
 using MongoDB.Bson.Serialization.Attributes;
@@ -532,6 +533,11 @@ namespace Horde.Build.Models
 		/// </summary>
 		public string ConfigRevision { get; }
 
+		/// <summary>
+		/// Configuration settings for the stream
+		/// </summary>
+		public StreamConfig Config { get; }
+
 		/// <summary>
 		/// Order to display on the dashboard's drop-down list
 		/// </summary>
@@ -688,7 +694,7 @@ namespace Horde.Build.Models
 			Dictionary<string, GetAgentTypeResponse> apiAgentTypes = stream.AgentTypes.ToDictionary(x => x.Key, x => x.Value.ToApiResponse());
 			Dictionary<string, GetWorkspaceTypeResponse> apiWorkspaceTypes = stream.WorkspaceTypes.ToDictionary(x => x.Key, x => x.Value.ToApiResponse());
 			GetAclResponse? apiAcl = (bIncludeAcl && stream.Acl != null)? new GetAclResponse(stream.Acl) : null;
-			return new Api.GetStreamResponse(stream.Id.ToString(), stream.ProjectId.ToString(), stream.Name, stream.ConfigPath, stream.ConfigRevision, stream.Order, stream.NotificationChannel, stream.NotificationChannelFilter, stream.TriageChannel, stream.DefaultPreflight?.ToRequest(), apiTabs, apiAgentTypes, apiWorkspaceTypes, apiTemplateRefs, apiAcl, stream.PausedUntil, stream.PauseComment);
+			return new Api.GetStreamResponse(stream.Id.ToString(), stream.ProjectId.ToString(), stream.Name, stream.ConfigPath, stream.ConfigRevision, stream.Order, stream.NotificationChannel, stream.NotificationChannelFilter, stream.TriageChannel, stream.DefaultPreflight?.ToRequest(), apiTabs, apiAgentTypes, apiWorkspaceTypes, apiTemplateRefs, apiAcl, stream.PausedUntil, stream.PauseComment, stream.Config.Workflows);
 		}
 
 		/// <summary>
diff --git a/Engine/Source/Programs/Horde/Horde.Build/Projects/Impl/StreamCollection.cs b/Engine/Source/Programs/Horde/Horde.Build/Projects/Impl/StreamCollection.cs
--- a/Engine/Source/Programs/Horde/Horde.Build/Projects/Impl/StreamCollection.cs
+++ b/Engine/Source/Programs/Horde/Horde.Build/Projects/Impl/StreamCollection.cs
@@ -8,6 +8,7 @@ using System.Linq;
 using System.Threading.Tasks;
 using Horde.Build.Acls;
 using Horde.Build.Api;
+using Horde.Build.Config;
 using Horde.Build.Models;
 using Horde.Build.Server;
 using Horde.Build.Utilities;
@@ -46,6 +47,7 @@ namespace Horde.Build.Collections.Impl
 			public string? ClusterName { get; set; }
 			public string ConfigPath { get; set; } = String.Empty;
 			public string ConfigRevision { get; set; } = String.Empty;
+			public StreamConfig? Config { get; set; }
 
 			public int Order { get; set; } = DefaultOrder;
 			public string? NotificationChannel { get; set; }
@@ -70,6 +72,7 @@ namespace Horde.Build.Collections.Impl
 			public int UpdateIndex { get; set; }
 			public bool Deleted { get; set; }
 
+			StreamConfig IStream.Config => Config!;
 			string IStream.ClusterName => ClusterName ?? PerforceCluster.DefaultName;
 			IReadOnlyList<StreamTab> IStream.Tabs => Tabs;
 			IReadOnlyDictionary<string, AgentType> IStream.AgentTypes => AgentTypes;
@@ -102,34 +105,31 @@ namespace Horde.Build.Collections.Impl
 			public static readonly ProjectionDefinition<StreamDocument> Projection = Builders<StreamDocument>.Projection.Include(x => x.Acl).Include(x => x.ProjectId);
 		}
 
-		/// <summary>
-		/// The stream collection
-		/// </summary>
 		readonly IMongoCollection<StreamDocument> _streams;
-
-		/// <summary>
-		/// Clock
-		/// </summary>
+		readonly ConfigCollection _configCollection;
 		readonly IClock _clock;
-
-		/// <summary>
-		/// The template collection
-		/// </summary>
 		readonly ITemplateCollection _templateCollection;
 
 		/// <summary>
 		/// Constructor
 		/// </summary>
 		/// <param name="mongoService">The database service instance</param>
+		/// <param name="configCollection"></param>
 		/// <param name="clock"></param>
 		/// <param name="templateCollection"></param>
-		public StreamCollection(MongoService mongoService, IClock clock, ITemplateCollection templateCollection)
+		public StreamCollection(MongoService mongoService, ConfigCollection configCollection, IClock clock, ITemplateCollection templateCollection)
 		{
 			_streams = mongoService.GetCollection<StreamDocument>("Streams");
+			_configCollection = configCollection;
 			_clock = clock;
 			_templateCollection = templateCollection;
 		}
 
+		async Task PostLoadAsync(StreamDocument stream)
+		{
+			stream.Config = await _configCollection.GetConfigAsync<StreamConfig>(stream.ConfigRevision);
+		}
+
 		/// <inheritdoc/>
 		public async Task<IStream?> TryCreateOrReplaceAsync(StreamId id, IStream? stream, string configPath, string revision, ProjectId projectId, StreamConfig config)
 		{
@@ -292,7 +292,12 @@ namespace Horde.Build.Collections.Impl
 		/// <inheritdoc/>
 		public async Task<IStream?> GetAsync(StreamId streamId)
 		{
-			return await _streams.Find<StreamDocument>(x => x.Id == streamId).FirstOrDefaultAsync();
+			StreamDocument? stream = await _streams.Find<StreamDocument>(x => x.Id == streamId).FirstOrDefaultAsync();
+			if (stream != null)
+			{
+				await PostLoadAsync(stream);
+			}
+			return stream;
 		}
 
 		/// <inheritdoc/>
@@ -305,6 +310,10 @@ namespace Horde.Build.Collections.Impl
 		public async Task<List<IStream>> FindAllAsync()
 		{
 			List<StreamDocument> results = await _streams.Find(Builders<StreamDocument>.Filter.Ne(x => x.Deleted, true)).ToListAsync();
+			foreach (StreamDocument result in results)
+			{
+				await PostLoadAsync(result);
+			}
 			return results.ConvertAll<IStream>(x => x);
 		}
 
@@ -312,7 +321,12 @@ namespace Horde.Build.Collections.Impl
 		public async Task<List<IStream>> FindForProjectsAsync(ProjectId[] projectIds)
 		{
 			FilterDefinition<StreamDocument> filter = Builders<StreamDocument>.Filter.In(x => x.ProjectId, projectIds) & Builders<StreamDocument>.Filter.Ne(x => x.Deleted, true);
+
 			List<StreamDocument> results = await _streams.Find(filter).ToListAsync();
+			foreach (StreamDocument result in results)
+			{
+				await PostLoadAsync(result);
+			}
 			return results.ConvertAll<IStream>(x => x);
 		}
 
@@ -375,7 +389,13 @@ namespace Horde.Build.Collections.Impl
 			update = update.Set(x => x.UpdateIndex, stream.UpdateIndex + 1);
 
 			FindOneAndUpdateOptions<StreamDocument> options = new FindOneAndUpdateOptions<StreamDocument> { ReturnDocument = ReturnDocument.After };
-			return await _streams.FindOneAndUpdateAsync(filter, update, options);
+
+			StreamDocument? result = await _streams.FindOneAndUpdateAsync(filter, update, options);
+			if(result != null)
+			{
+				await PostLoadAsync(result);
+			}
+			return result;
 		}
 
 		/// <inheritdoc/>
diff --git a/Engine/Source/Programs/Horde/Horde.Build/Server/Config.cs b/Engine/Source/Programs/Horde/Horde.Build/Server/Config.cs
--- a/Engine/Source/Programs/Horde/Horde.Build/Server/Config.cs
+++ b/Engine/Source/Programs/Horde/Horde.Build/Server/Config.cs
@@ -317,6 +317,11 @@ namespace Horde.Build.Server
 		/// </summary>
 		public List<TimeSpan> ReportTimes { get; set; } = new List<TimeSpan> { TimeSpan.Zero };
 
+		/// <summary>
+		/// Name of the tab to post summary data to
+		/// </summary>
+		public string? SummaryTab { get; set; }
+
 		/// <summary>
 		/// Channel to post summary information for these templates.
 		/// </summary>
