Horde: Fix tests.

diff --git a/Engine/Source/Programs/Horde/Horde.Build.Tests/ControllerIntegrationTest.cs b/Engine/Source/Programs/Horde/Horde.Build.Tests/ControllerIntegrationTest.cs
--- a/Engine/Source/Programs/Horde/Horde.Build.Tests/ControllerIntegrationTest.cs
+++ b/Engine/Source/Programs/Horde/Horde.Build.Tests/ControllerIntegrationTest.cs
@@ -6,6 +6,7 @@ using System.Net.Http;
 using System.Threading.Tasks;
 using Horde.Build.Collections;
 using Horde.Build.Collections.Impl;
+using Horde.Build.Config;
 using Horde.Build.Jobs;
 using Horde.Build.Server;
 using Horde.Build.Services;
@@ -79,6 +80,7 @@ public class ControllerIntegrationTest : IDisposable
 	private async Task<Fixture> CreateFixture()
 	{
 		IServiceProvider services = Factory.Services;
+		ConfigCollection configCollection = services.GetRequiredService<ConfigCollection>();
 		MongoService mongoService = services.GetRequiredService<MongoService>();
 		ITemplateCollection templateService = services.GetRequiredService<ITemplateCollection>();
 		JobService jobService = services.GetRequiredService<JobService>();
@@ -87,7 +89,7 @@ public class ControllerIntegrationTest : IDisposable
 		AgentService agentService = services.GetRequiredService<AgentService>();
 		GraphCollection graphCollection = new(mongoService);
 
-		return await Fixture.Create(graphCollection, templateService, jobService, artifactCollection, streamService,
+		return await Fixture.Create(configCollection, graphCollection, templateService, jobService, artifactCollection, streamService,
 			agentService);
 	}
 }
\ No newline at end of file
diff --git a/Engine/Source/Programs/Horde/Horde.Build.Tests/Fixture.cs b/Engine/Source/Programs/Horde/Horde.Build.Tests/Fixture.cs
--- a/Engine/Source/Programs/Horde/Horde.Build.Tests/Fixture.cs
+++ b/Engine/Source/Programs/Horde/Horde.Build.Tests/Fixture.cs
@@ -7,6 +7,7 @@ using System.Threading.Tasks;
 using EpicGames.Core;
 using Horde.Build.Api;
 using Horde.Build.Collections;
+using Horde.Build.Config;
 using Horde.Build.Jobs;
 using Horde.Build.Models;
 using Horde.Build.Server;
@@ -36,10 +37,10 @@ namespace Horde.Build.Tests
 		public string Agent1Name { get; private set; } = null!;
 		public const string PoolName = "TestingPool";
 
-		public static async Task<Fixture> Create(IGraphCollection graphCollection, ITemplateCollection templateCollection, JobService jobService, IArtifactCollection artifactCollection, StreamService streamService, AgentService agentService)
+		public static async Task<Fixture> Create(ConfigCollection configCollection, IGraphCollection graphCollection, ITemplateCollection templateCollection, JobService jobService, IArtifactCollection artifactCollection, StreamService streamService, AgentService agentService)
 		{
 			Fixture fixture = new Fixture();
-			await fixture.Populate(graphCollection, templateCollection, jobService, artifactCollection, streamService, agentService);
+			await fixture.Populate(configCollection, graphCollection, templateCollection, jobService, artifactCollection, streamService, agentService);
 
 //			(PerforceService as PerforceServiceStub)?.AddChange("//UE5/Main", 112233, "leet.coder", "Did stuff", new []{"file.cpp"});
 //			(PerforceService as PerforceServiceStub)?.AddChange("//UE5/Main", 1111, "swarm", "A shelved CL here", new []{"renderer.cpp"});
@@ -47,7 +48,7 @@ namespace Horde.Build.Tests
 			return fixture;
 		}
 
-		private async Task Populate(IGraphCollection graphCollection, ITemplateCollection templateCollection, JobService jobService, IArtifactCollection artifactCollection, StreamService streamService, AgentService agentService)
+		private async Task Populate(ConfigCollection configCollection, IGraphCollection graphCollection, ITemplateCollection templateCollection, JobService jobService, IArtifactCollection artifactCollection, StreamService streamService, AgentService agentService)
 		{
 			FixtureGraph fg = new FixtureGraph();
 			fg.Id = ContentHash.Empty;
@@ -74,14 +75,15 @@ namespace Horde.Build.Tests
 				{ "Win64", new() { Pool = PoolName} }
 			};
 
+			StreamConfig config = new StreamConfig { Name = "//UE5/Main", Tabs = tabs, Templates = templates, AgentTypes = agentTypes };
+			await configCollection.AddConfigAsync("rev1", config);
+
 			Stream = await streamService.StreamCollection.GetAsync(new StreamId("ue5-main"));
 			Stream = await streamService.StreamCollection.TryCreateOrReplaceAsync(
 				new StreamId("ue5-main"),
 				Stream,
-				"",
-				"",
-				new ProjectId("does-not-exist"),
-				new StreamConfig { Name = "//UE5/Main", Tabs = tabs, Templates = templates, AgentTypes = agentTypes}
+				"rev1",
+				new ProjectId("does-not-exist")
 			);
 			
 			Job1 = await jobService.CreateJobAsync(
diff --git a/Engine/Source/Programs/Horde/Horde.Build.Tests/Fleet/JobQueueStrategyTest.cs b/Engine/Source/Programs/Horde/Horde.Build.Tests/Fleet/JobQueueStrategyTest.cs
--- a/Engine/Source/Programs/Horde/Horde.Build.Tests/Fleet/JobQueueStrategyTest.cs
+++ b/Engine/Source/Programs/Horde/Horde.Build.Tests/Fleet/JobQueueStrategyTest.cs
@@ -110,13 +110,12 @@ namespace Horde.Build.Tests.Fleet
 			
 			string agentTypeName1 = "bogusAgentType1";
 			Dictionary<string, CreateAgentTypeRequest> agentTypes = new() { {agentTypeName1, new() { Pool = pool1.Name} }, };
-			IStream stream = (await StreamCollection.TryCreateOrReplaceAsync(
+
+			IStream stream = (await CreateOrReplaceStreamAsync(
 				new StreamId("ue5-main"),
 				null,
-				"",
-				"",
 				new ProjectId("does-not-exist"),
-				new StreamConfig { Name = "//UE5/Main", AgentTypes = agentTypes}
+				new StreamConfig { Name = "//UE5/Main", AgentTypes = agentTypes }
 			))!;
 
 			string nodeForAgentType1 = "bogusNodeOnAgentType1";
diff --git a/Engine/Source/Programs/Horde/Horde.Build.Tests/IssueServiceTests.cs b/Engine/Source/Programs/Horde/Horde.Build.Tests/IssueServiceTests.cs
--- a/Engine/Source/Programs/Horde/Horde.Build.Tests/IssueServiceTests.cs
+++ b/Engine/Source/Programs/Horde/Horde.Build.Tests/IssueServiceTests.cs
@@ -114,7 +114,7 @@ namespace Horde.Build.Tests
 			streamConfig.Templates.Add(new Api.TemplateRefConfig { Id = new TemplateRefId("test-template") });
 			await ConfigCollection.AddConfigAsync(revision, streamConfig);
 
-			return Deref(await StreamCollection.TryCreateOrReplaceAsync(streamId, null, "", revision, projectId, streamConfig));
+			return Deref(await StreamCollection.TryCreateOrReplaceAsync(streamId, null, revision, projectId));
 		}
 
 		public IssueServiceTests()
diff --git a/Engine/Source/Programs/Horde/Horde.Build.Tests/JobServiceTests.cs b/Engine/Source/Programs/Horde/Horde.Build.Tests/JobServiceTests.cs
--- a/Engine/Source/Programs/Horde/Horde.Build.Tests/JobServiceTests.cs
+++ b/Engine/Source/Programs/Horde/Horde.Build.Tests/JobServiceTests.cs
@@ -43,7 +43,7 @@ namespace Horde.Build.Tests
 
 			StreamId streamId = new StreamId("ue5-main");
 			IStream? stream = await StreamService.GetStreamAsync(streamId);
-			stream = await StreamService.StreamCollection.TryCreateOrReplaceAsync(new StreamId("ue5-main"), stream, String.Empty, String.Empty, projectId, streamConfig);
+			stream = await CreateOrReplaceStreamAsync(new StreamId("ue5-main"), stream, projectId, streamConfig);
 
 			IJob job = await JobService.CreateJobAsync(null, stream!, templateRefId1, template.Id, graph, "Hello", 1234, 1233, 999, null, null, null, null, null, null, stream!.Templates[templateRefId1].ChainedJobs, true, true, null, null, new List<string>());
 			Assert.AreEqual(1, job.ChainedJobs.Count);
@@ -155,7 +155,7 @@ namespace Horde.Build.Tests
 
 			StreamId streamId = new StreamId("ue5-main");
 			IStream? stream = await StreamCollection.GetAsync(streamId);
-			stream = await StreamCollection.TryCreateOrReplaceAsync(streamId, stream, "", "", project!.Id, new StreamConfig { Name = "//UE5/Main" });
+			stream = await CreateOrReplaceStreamAsync(streamId, stream, project!.Id, new StreamConfig { Name = "//UE5/Main" });
 
 			ITemplate template = await TemplateCollection.AddAsync("Test template");
 			IGraph graph = await GraphCollection.AddAsync(template);
diff --git a/Engine/Source/Programs/Horde/Horde.Build.Tests/SchedulerTests.cs b/Engine/Source/Programs/Horde/Horde.Build.Tests/SchedulerTests.cs
--- a/Engine/Source/Programs/Horde/Horde.Build.Tests/SchedulerTests.cs
+++ b/Engine/Source/Programs/Horde/Horde.Build.Tests/SchedulerTests.cs
@@ -60,7 +60,7 @@ namespace Horde.Build.Tests
 			config.Tabs.Add(new CreateJobsTabRequest { Title = "foo", Templates = new List<string> { TemplateRefId.ToString() } });
 			config.Templates.Add(new TemplateRefConfig { Id = TemplateRefId, Name = "Test", Schedule = schedule });
 
-			return (await StreamService.StreamCollection.TryCreateOrReplaceAsync(StreamId, stream, "", "", ProjectId, config))!; 
+			return (await CreateOrReplaceStreamAsync(StreamId, stream, ProjectId, config))!; 
 		}
 
 		public async Task<List<IJob>> FileTestHelperAsync(params string[] files)
@@ -424,7 +424,7 @@ namespace Horde.Build.Tests
 			config.Name = "//UE5/Main";
 			config.Tabs.Add(new CreateJobsTabRequest { Title = "foo", Templates = new List<string> { newTemplateRefId1.ToString(), newTemplateRefId2.ToString() } });
 
-			stream = (await StreamService.StreamCollection.TryCreateOrReplaceAsync(StreamId, stream, "", "", ProjectId, config))!;
+			stream = (await CreateOrReplaceStreamAsync(StreamId, stream, ProjectId, config))!;
 
 			// Create the TriggerNext step and mark it as complete
 			IGraph graphA = await GraphCollection.AddAsync(newTemplate1);
@@ -506,7 +506,7 @@ namespace Horde.Build.Tests
 			config.Templates.Add(newTemplate1);
 			config.Templates.Add(newTemplate2);
 
-			stream = (await StreamService.StreamCollection.TryCreateOrReplaceAsync(StreamId, stream, "", "", ProjectId, config))!;
+			stream = (await CreateOrReplaceStreamAsync(StreamId, stream, ProjectId, config))!;
 
 			ITemplate template1 = (await TemplateCollection.GetAsync(stream.Templates[newTemplateRefId1].Hash))!;
 			Assert.IsNotNull(template1);
diff --git a/Engine/Source/Programs/Horde/Horde.Build.Tests/TestSetup.cs b/Engine/Source/Programs/Horde/Horde.Build.Tests/TestSetup.cs
--- a/Engine/Source/Programs/Horde/Horde.Build.Tests/TestSetup.cs
+++ b/Engine/Source/Programs/Horde/Horde.Build.Tests/TestSetup.cs
@@ -4,6 +4,7 @@ using System;
 using System.Collections.Generic;
 using System.Reflection;
 using System.Security.Claims;
+using System.Text.Json;
 using System.Threading.Tasks;
 using Datadog.Trace;
 using EpicGames.Core;
@@ -46,6 +47,9 @@ using StatsdClient;
 
 namespace Horde.Build.Tests
 {
+	using StreamId = StringId<IStream>;
+	using ProjectId = StringId<IProject>;
+
 	/// <summary>
 	/// Handles set up of collections, services, fixtures etc during testing
 	///
@@ -212,9 +216,22 @@ namespace Horde.Build.Tests
 			services.AddSingleton<ISingletonDocument<AgentSoftwareChannels>>(new SingletonDocumentStub<AgentSoftwareChannels>());
 		}
 
+		public async Task<IStream?> CreateOrReplaceStreamAsync(StreamId streamId, IStream? stream, ProjectId projectId, StreamConfig config)
+		{
+			JsonSerializerOptions options = new JsonSerializerOptions();
+			Startup.ConfigureJsonSerializer(options);
+
+			byte[] data = JsonSerializer.SerializeToUtf8Bytes(config, options);
+			IoHash hash = IoHash.Compute(data);
+			string revision = hash.ToString();
+			await ConfigCollection.AddConfigDataAsync(revision, data);
+
+			return await StreamService.StreamCollection.TryCreateOrReplaceAsync(streamId, stream, hash.ToString(), projectId);
+		}
+
 		public Task<Fixture> CreateFixtureAsync()
 		{
-			return Fixture.Create(GraphCollection, TemplateCollection, JobService, ArtifactCollection, StreamService, AgentService);
+			return Fixture.Create(ConfigCollection, GraphCollection, TemplateCollection, JobService, ArtifactCollection, StreamService, AgentService);
 		}
 
 		private JobsController GetJobsController()
diff --git a/Engine/Source/Programs/Horde/Horde.Build/Api/Streams.cs b/Engine/Source/Programs/Horde/Horde.Build/Api/Streams.cs
--- a/Engine/Source/Programs/Horde/Horde.Build/Api/Streams.cs
+++ b/Engine/Source/Programs/Horde/Horde.Build/Api/Streams.cs
@@ -572,7 +572,6 @@ namespace Horde.Build.Api
 		/// <param name="id">Unique id of the stream</param>
 		/// <param name="projectId">Unique id of the project containing the stream</param>
 		/// <param name="name">Name of the stream</param>
-		/// <param name="configPath">Path to the config file for this stream</param>
 		/// <param name="configRevision">The config file path on the server</param>
 		/// <param name="order">Order to display this stream</param>
 		/// <param name="notificationChannel"></param>
@@ -587,12 +586,11 @@ namespace Horde.Build.Api
 		/// <param name="pausedUntil">Stream paused for new builds until this date</param>
 		/// <param name="pauseComment">Reason for stream being paused</param>
 		/// <param name="workflows">Workflows for this stream</param>
-		public GetStreamResponse(string id, string projectId, string name, string configPath, string configRevision, int order, string? notificationChannel, string? notificationChannelFilter, string? triageChannel, DefaultPreflightRequest? defaultPreflight, List<GetStreamTabResponse> tabs, Dictionary<string, GetAgentTypeResponse> agentTypes, Dictionary<string, GetWorkspaceTypeResponse>? workspaceTypes, List<GetTemplateRefResponse> templates, GetAclResponse? acl, DateTime? pausedUntil, string? pauseComment, List<WorkflowConfig> workflows)
+		public GetStreamResponse(string id, string projectId, string name, string configRevision, int order, string? notificationChannel, string? notificationChannelFilter, string? triageChannel, DefaultPreflightRequest? defaultPreflight, List<GetStreamTabResponse> tabs, Dictionary<string, GetAgentTypeResponse> agentTypes, Dictionary<string, GetWorkspaceTypeResponse>? workspaceTypes, List<GetTemplateRefResponse> templates, GetAclResponse? acl, DateTime? pausedUntil, string? pauseComment, List<WorkflowConfig> workflows)
 		{
 			Id = id;
 			ProjectId = projectId;
 			Name = name;
-			ConfigPath = configPath;
 			ConfigRevision = configRevision;
 			Order = order;
 			NotificationChannel = notificationChannel;
diff --git a/Engine/Source/Programs/Horde/Horde.Build/Configuration/ConfigUpdateService.cs b/Engine/Source/Programs/Horde/Horde.Build/Configuration/ConfigUpdateService.cs
--- a/Engine/Source/Programs/Horde/Horde.Build/Configuration/ConfigUpdateService.cs
+++ b/Engine/Source/Programs/Horde/Horde.Build/Configuration/ConfigUpdateService.cs
@@ -264,13 +264,13 @@ namespace Horde.Build.Config
 				if (streamRevisions.TryGetValue(streamPath, out string? revision))
 				{
 					IStream? stream = streams.FirstOrDefault(x => x.Id == streamRef.Id);
-					if (stream == null || stream.ConfigPath != streamPath.ToString() || stream.ConfigRevision != revision)
+					if (stream == null || stream.ConfigRevision != revision)
 					{
 						_logger.LogInformation("Updating configuration for stream {StreamRef} ({Revision})", streamRef.Id, revision);
 						try
 						{
 							StreamConfig streamConfig = await ReadDataAsync<StreamConfig>(revision, streamPath);
-							stream = await _streamService.StreamCollection.CreateOrReplaceAsync(streamRef.Id, stream, streamPath.ToString(), revision, projectId, streamConfig);
+							stream = await _streamService.StreamCollection.CreateOrReplaceAsync(streamRef.Id, stream, revision, projectId);
 						}
 						catch (Exception ex)
 						{
diff --git a/Engine/Source/Programs/Horde/Horde.Build/Projects/IStream.cs b/Engine/Source/Programs/Horde/Horde.Build/Projects/IStream.cs
--- a/Engine/Source/Programs/Horde/Horde.Build/Projects/IStream.cs
+++ b/Engine/Source/Programs/Horde/Horde.Build/Projects/IStream.cs
@@ -523,11 +523,6 @@ namespace Horde.Build.Models
 		/// </summary>
 		public string ClusterName { get; }
 
-		/// <summary>
-		/// Path to the configuration file for this stream
-		/// </summary>
-		public string ConfigPath { get; }
-
 		/// <summary>
 		/// The revision of config file used for this stream
 		/// </summary>
@@ -694,7 +689,7 @@ namespace Horde.Build.Models
 			Dictionary<string, GetAgentTypeResponse> apiAgentTypes = stream.AgentTypes.ToDictionary(x => x.Key, x => x.Value.ToApiResponse());
 			Dictionary<string, GetWorkspaceTypeResponse> apiWorkspaceTypes = stream.WorkspaceTypes.ToDictionary(x => x.Key, x => x.Value.ToApiResponse());
 			GetAclResponse? apiAcl = (bIncludeAcl && stream.Acl != null)? new GetAclResponse(stream.Acl) : null;
-			return new Api.GetStreamResponse(stream.Id.ToString(), stream.ProjectId.ToString(), stream.Name, stream.ConfigPath, stream.ConfigRevision, stream.Order, stream.NotificationChannel, stream.NotificationChannelFilter, stream.TriageChannel, stream.DefaultPreflight?.ToRequest(), apiTabs, apiAgentTypes, apiWorkspaceTypes, apiTemplateRefs, apiAcl, stream.PausedUntil, stream.PauseComment, stream.Config.Workflows);
+			return new Api.GetStreamResponse(stream.Id.ToString(), stream.ProjectId.ToString(), stream.Name, stream.ConfigRevision, stream.Order, stream.NotificationChannel, stream.NotificationChannelFilter, stream.TriageChannel, stream.DefaultPreflight?.ToRequest(), apiTabs, apiAgentTypes, apiWorkspaceTypes, apiTemplateRefs, apiAcl, stream.PausedUntil, stream.PauseComment, stream.Config.Workflows);
 		}
 
 		/// <summary>
diff --git a/Engine/Source/Programs/Horde/Horde.Build/Projects/IStreamCollection.cs b/Engine/Source/Programs/Horde/Horde.Build/Projects/IStreamCollection.cs
--- a/Engine/Source/Programs/Horde/Horde.Build/Projects/IStreamCollection.cs
+++ b/Engine/Source/Programs/Horde/Horde.Build/Projects/IStreamCollection.cs
@@ -24,12 +24,10 @@ namespace Horde.Build.Collections
 		/// </summary>
 		/// <param name="id">Unique id for the new stream</param>
 		/// <param name="stream">The current stream value. If not-null, this will attempt to replace the existing instance.</param>
-		/// <param name="configPath">Path to the config file</param>
 		/// <param name="revision">The config file revision</param>
 		/// <param name="projectId">The project id</param>
-		/// <param name="config">The stream configuration</param>
 		/// <returns></returns>
-		Task<IStream?> TryCreateOrReplaceAsync(StreamId id, IStream? stream, string configPath, string revision, ProjectId projectId, StreamConfig config);
+		Task<IStream?> TryCreateOrReplaceAsync(StreamId id, IStream? stream, string revision, ProjectId projectId);
 
 		/// <summary>
 		/// Gets a stream by ID
@@ -94,16 +92,14 @@ namespace Horde.Build.Collections
 		/// <param name="streamCollection">The stream collection</param>
 		/// <param name="id">Unique id for the new stream</param>
 		/// <param name="stream">The current stream value. If not-null, this will attempt to replace the existing instance.</param>
-		/// <param name="configPath">Path to the config file</param>
 		/// <param name="revision">The config file revision</param>
 		/// <param name="projectId">The project id</param>
-		/// <param name="config">The stream configuration</param>
 		/// <returns></returns>
-		public static async Task<IStream> CreateOrReplaceAsync(this IStreamCollection streamCollection, StreamId id, IStream? stream, string configPath, string revision, ProjectId projectId, StreamConfig config)
+		public static async Task<IStream> CreateOrReplaceAsync(this IStreamCollection streamCollection, StreamId id, IStream? stream, string revision, ProjectId projectId)
 		{
 			for (; ; )
 			{
-				stream = await streamCollection.TryCreateOrReplaceAsync(id, stream, configPath, revision, projectId, config);
+				stream = await streamCollection.TryCreateOrReplaceAsync(id, stream, revision, projectId);
 				if (stream != null)
 				{
 					return stream;
diff --git a/Engine/Source/Programs/Horde/Horde.Build/Projects/Impl/StreamCollection.cs b/Engine/Source/Programs/Horde/Horde.Build/Projects/Impl/StreamCollection.cs
--- a/Engine/Source/Programs/Horde/Horde.Build/Projects/Impl/StreamCollection.cs
+++ b/Engine/Source/Programs/Horde/Horde.Build/Projects/Impl/StreamCollection.cs
@@ -45,7 +45,6 @@ namespace Horde.Build.Collections.Impl
 			public string Name { get; set; }
 
 			public string? ClusterName { get; set; }
-			public string ConfigPath { get; set; } = String.Empty;
 			public string ConfigRevision { get; set; } = String.Empty;
 			public StreamConfig? Config { get; set; }
 
@@ -131,8 +130,10 @@ namespace Horde.Build.Collections.Impl
 		}
 
 		/// <inheritdoc/>
-		public async Task<IStream?> TryCreateOrReplaceAsync(StreamId id, IStream? stream, string configPath, string revision, ProjectId projectId, StreamConfig config)
+		public async Task<IStream?> TryCreateOrReplaceAsync(StreamId id, IStream? stream, string revision, ProjectId projectId)
 		{
+			StreamConfig config = await _configCollection.GetConfigAsync<StreamConfig>(revision);
+
 			List<StreamTab> tabs = config.Tabs.ConvertAll(x => StreamTab.FromRequest(x));
 			Dictionary<TemplateRefId, TemplateRef> templateRefs = await CreateTemplateRefsAsync(config.Templates, stream, _templateCollection);
 
@@ -159,11 +160,11 @@ namespace Horde.Build.Collections.Impl
 			Acl? acl = Acl.Merge(new Acl(), config.Acl);
 			if (stream == null)
 			{
-				return await TryCreateAsync(id, projectId, configPath, revision, config, defaultPreflight, tabs, agentTypes, workspaceTypes, templateRefs, acl);
+				return await TryCreateAsync(id, projectId, revision, config, defaultPreflight, tabs, agentTypes, workspaceTypes, templateRefs, acl);
 			}
 			else
 			{
-				return await TryReplaceAsync(stream, projectId, configPath, revision, config, defaultPreflight, tabs, agentTypes, workspaceTypes, templateRefs, acl);
+				return await TryReplaceAsync(stream, projectId, revision, config, defaultPreflight, tabs, agentTypes, workspaceTypes, templateRefs, acl);
 			}
 		}
 
@@ -220,12 +221,12 @@ namespace Horde.Build.Collections.Impl
 		}
 
 		/// <inheritdoc/>
-		async Task<IStream?> TryCreateAsync(StreamId id, ProjectId projectId, string configPath, string configRevision, StreamConfig config, DefaultPreflight? defaultPreflight, List<StreamTab> tabs, Dictionary<string, AgentType> agentTypes, Dictionary<string, WorkspaceType> workspaceTypes, Dictionary<TemplateRefId, TemplateRef> templateRefs, Acl? acl)
+		async Task<IStream?> TryCreateAsync(StreamId id, ProjectId projectId, string configRevision, StreamConfig config, DefaultPreflight? defaultPreflight, List<StreamTab> tabs, Dictionary<string, AgentType> agentTypes, Dictionary<string, WorkspaceType> workspaceTypes, Dictionary<TemplateRefId, TemplateRef> templateRefs, Acl? acl)
 		{
 			StreamDocument newStream = new StreamDocument(id, config.Name, projectId);
 			newStream.ClusterName = config.ClusterName;
-			newStream.ConfigPath = configPath;
 			newStream.ConfigRevision = configRevision;
+			newStream.Config = config;
 			newStream.Order = config.Order ?? StreamDocument.DefaultOrder;
 			newStream.NotificationChannel = config.NotificationChannel;
 			newStream.NotificationChannelFilter = config.NotificationChannelFilter;
@@ -258,7 +259,7 @@ namespace Horde.Build.Collections.Impl
 		}
 
 		/// <inheritdoc/>
-		async Task<IStream?> TryReplaceAsync(IStream streamInterface, ProjectId projectId, string configPath, string configRevision, StreamConfig config, DefaultPreflight? defaultPreflight, List<StreamTab> tabs, Dictionary<string, AgentType>? agentTypes, Dictionary<string, WorkspaceType>? workspaceTypes, Dictionary<TemplateRefId, TemplateRef>? templateRefs, Acl? acl)
+		async Task<IStream?> TryReplaceAsync(IStream streamInterface, ProjectId projectId, string configRevision, StreamConfig config, DefaultPreflight? defaultPreflight, List<StreamTab> tabs, Dictionary<string, AgentType>? agentTypes, Dictionary<string, WorkspaceType>? workspaceTypes, Dictionary<TemplateRefId, TemplateRef>? templateRefs, Acl? acl)
 		{
 			int order = config.Order ?? StreamDocument.DefaultOrder;
 
@@ -270,7 +271,6 @@ namespace Horde.Build.Collections.Impl
 			updates.Add(updateBuilder.Set(x => x.Name, config.Name));
 			updates.Add(updateBuilder.Set(x => x.ProjectId, projectId));
 			updates.Add(updateBuilder.Set(x => x.ClusterName, config.ClusterName));
-			updates.Add(updateBuilder.Set(x => x.ConfigPath, configPath));
 			updates.Add(updateBuilder.Set(x => x.ConfigRevision, configRevision));
 			updates.Add(updateBuilder.Set(x => x.Order, order));
 			updates.Add(updateBuilder.Set(x => x.NotificationChannel, config.NotificationChannel));
