Horde: Add cross-compilation of remote exec test binary

Will build binaries for all platforms/archs that remote execution supports.
Previously, only a Windows binary was checked in but testing now requires more platforms.

diff --git a/Engine/Source/Programs/Horde/Horde.Agent.Tests/remote-exec-test-bin/README.md b/Engine/Source/Programs/Horde/Horde.Agent.Tests/remote-exec-test-bin/README.md
new file mode 100644
--- /dev/null
+++ b/Engine/Source/Programs/Horde/Horde.Agent.Tests/remote-exec-test-bin/README.md
@@ -0,0 +1,11 @@
+# Test binary for remote execution
+
+Provided in this directory are binaries for the different platforms that remote execution in Horde supports.
+The binary accepts two arguments: one for sleeping an arbitrary amount of time (for simulating work) and one for specifying a forced exit code.
+This is useful for testing and verifying behavior of the different platforms.
+
+## Building
+The binaries are built using the Dockerfile together with Zig.
+The cross-compilation capability offered by Zig out of the box is great for producing the binaries for Windows, Linux and macOS together with their corresponding processor architectures.
+
+Simply run `build.bat` on Windows with Docker installed to produce new binaries in the same directory (works on Linux or macOS too if the command-line is copied from batch file).
\ No newline at end of file
diff --git a/Engine/Source/Programs/Horde/Horde.Agent.Tests/remote-exec-test-bin/build.bat b/Engine/Source/Programs/Horde/Horde.Agent.Tests/remote-exec-test-bin/build.bat
--- a/Engine/Source/Programs/Horde/Horde.Agent.Tests/remote-exec-test-bin/build.bat
+++ b/Engine/Source/Programs/Horde/Horde.Agent.Tests/remote-exec-test-bin/build.bat
@@ -1,2 +1,3 @@
 @echo off
-cl.exe /O1 /MT remote-exec-test-bin.c
\ No newline at end of file
+set DOCKER_BUILDKIT=1
+docker build --progress=plain --file Dockerfile --output . .
\ No newline at end of file
diff --git a/Engine/Source/Programs/Horde/Horde.Agent.Tests/remote-exec-test-bin/remote-exec-test-bin.c b/Engine/Source/Programs/Horde/Horde.Agent.Tests/remote-exec-test-bin/remote-exec-test-bin.c
--- a/Engine/Source/Programs/Horde/Horde.Agent.Tests/remote-exec-test-bin/remote-exec-test-bin.c
+++ b/Engine/Source/Programs/Horde/Horde.Agent.Tests/remote-exec-test-bin/remote-exec-test-bin.c
@@ -1,5 +1,12 @@
 #include <stdio.h>
+#include <stdlib.h>
+
+#ifdef _WINDOWS
 #include <windows.h>
+#else
+#include <unistd.h>
+#define Sleep(x) usleep((x)*1000)
+#endif
 
 int main(int argc, char *argv[]) {
     printf("Horde Remote Execution Test Binary\n");
