Horde: Include workflows when searching for issues.

diff --git a/Engine/Source/Programs/Horde/Horde.Build/Api/Issues.cs b/Engine/Source/Programs/Horde/Horde.Build/Api/Issues.cs
--- a/Engine/Source/Programs/Horde/Horde.Build/Api/Issues.cs
+++ b/Engine/Source/Programs/Horde/Horde.Build/Api/Issues.cs
@@ -600,6 +600,11 @@ namespace Horde.Build.Api
 		/// </summary>
 		public string Name { get; set; }
 
+		/// <summary>
+		/// Workflow for this span
+		/// </summary>
+		public WorkflowId? WorkflowId { get; set; }
+
 		/// <summary>
 		/// The previous build 
 		/// </summary>
@@ -614,11 +619,13 @@ namespace Horde.Build.Api
 		/// Constructor
 		/// </summary>
 		/// <param name="span"></param>
-		public FindIssueSpanResponse(IIssueSpan span)
+		/// <param name="workflowId"></param>
+		public FindIssueSpanResponse(IIssueSpan span, WorkflowId? workflowId)
 		{
 			Id = span.Id.ToString();
 			TemplateId = span.TemplateRefId.ToString();
 			Name = span.NodeName;
+			WorkflowId = workflowId;
 			if (span.LastSuccess != null)
 			{
 				LastSuccess = new GetIssueStepResponse(span.LastSuccess);
diff --git a/Engine/Source/Programs/Horde/Horde.Build/Controllers/IssuesController.cs b/Engine/Source/Programs/Horde/Horde.Build/Controllers/IssuesController.cs
--- a/Engine/Source/Programs/Horde/Horde.Build/Controllers/IssuesController.cs
+++ b/Engine/Source/Programs/Horde/Horde.Build/Controllers/IssuesController.cs
@@ -23,6 +23,7 @@ namespace Horde.Build.Controllers
 	using LogId = ObjectId<ILogFile>;
 	using StreamId = StringId<IStream>;
 	using UserId = ObjectId<IUser>;
+	using WorkflowId = StringId<WorkflowConfig>;
 
 	/// <summary>
 	/// Controller for the /api/v1/issues endpoint
@@ -68,7 +69,7 @@ namespace Horde.Build.Controllers
 		/// <returns>List of matching agents</returns>
 		[HttpGet]
 		[Route("/api/v2/issues")]
-		[ProducesResponseType(typeof(List<GetIssueResponse>), 200)]
+		[ProducesResponseType(typeof(List<FindIssueResponse>), 200)]
 		public async Task<ActionResult<object>> FindIssuesV2Async([FromQuery(Name = "Id")] int[]? ids = null, [FromQuery] StreamId? streamId = null, [FromQuery] int? minChange = null, [FromQuery] int? maxChange = null, [FromQuery] bool? resolved = null, [FromQuery] int index = 0, [FromQuery] int count = 10, [FromQuery] PropertyFilter? filter = null)
 		{
 			if (ids != null && ids.Length == 0)
@@ -84,6 +85,12 @@ namespace Horde.Build.Controllers
 					return Forbid();
 				}
 
+				IStream? stream = await _streamService.GetStreamAsync(streamId.Value);
+				if (stream == null)
+				{
+					return NotFound(streamId.Value);
+				}
+
 				List<IIssueSpan> spans = await _issueCollection.FindSpansAsync(null, ids, streamId.Value, minChange, maxChange, resolved);
 				if(spans.Count > 0)
 				{
@@ -128,7 +135,16 @@ namespace Horde.Build.Controllers
 								}
 							}
 
-							spanResponses.AddRange(spansForIssue.Select(x => new FindIssueSpanResponse(x)));
+							// Convert each issue to a response
+							foreach (IIssueSpan span in spansForIssue)
+							{
+								WorkflowId? workflowId = null;
+								if(stream.Config.TryGetTemplate(span.TemplateRefId, out TemplateRefConfig? templateRefConfig))
+								{
+									workflowId = templateRefConfig.Workflow;
+								}
+								spanResponses.Add(new FindIssueSpanResponse(span, workflowId));
+							}
 						}
 					
 						IUser? owner = null;
