Horde: Strikethrough issues in Slack once they've been resolved.

diff --git a/Engine/Source/Programs/Horde/Horde.Build/Notifications/Impl/SlackNotificationSink.cs b/Engine/Source/Programs/Horde/Horde.Build/Notifications/Impl/SlackNotificationSink.cs
--- a/Engine/Source/Programs/Horde/Horde.Build/Notifications/Impl/SlackNotificationSink.cs
+++ b/Engine/Source/Programs/Horde/Horde.Build/Notifications/Impl/SlackNotificationSink.cs
@@ -692,24 +692,17 @@ namespace Horde.Build.Notifications.Impl
 			{
 				Uri issueUrl = GetIssueUrl(issue, span.FirstFailure);
 
-				BlockKitMessage message = new BlockKitMessage();
-				message.Channel = triageChannel;
-				message.Text = $"{workflow.TriagePrefix}*Issue <{issueUrl}|{issue.Id}>*: {issue.Summary}{workflow.TriageSuffix}";
-
 				string eventId = GetTriageThreadEventId(issue.Id);
 
-				(MessageStateDocument state, bool isNew) = await AddOrUpdateMessageStateAsync(triageChannel, eventId, null, "");
-				if (isNew)
+				string text = $"{workflow.TriagePrefix}*Issue <{issueUrl}|{issue.Id}>*: {issue.Summary}{workflow.TriageSuffix}";
+				if (issue.ResolvedAt != null)
 				{
-					PostMessageResponse response = await SendRequestAsync<PostMessageResponse>(PostMessageUrl, message);
-					if (String.IsNullOrEmpty(response.Ts) || String.IsNullOrEmpty(response.Channel))
-					{
-						_logger.LogWarning("Missing 'ts' or 'channel' field on slack response");
-					}
-
-					state.Channel = response.Channel ?? String.Empty;
-					state.Ts = response.Ts ?? String.Empty;
+					text = $"~{text}~";
+				}
 
+				(MessageStateDocument state, bool isNew) = await SendOrUpdateMessageAsync(triageChannel, eventId, null, text);
+				if (isNew)
+				{
 					// Create the summary text
 					List<ILogEvent> events = new List<ILogEvent>();
 					List<ILogEventData> eventDataItems = new List<ILogEventData>();
@@ -749,9 +742,9 @@ namespace Horde.Build.Notifications.Impl
 
 					// Permalink to the summary text so we link inside the thread rather than just to the original message
 					string? permalink = null;
-					if (response.Channel != null && summaryTs != null)
+					if (summaryTs != null)
 					{
-						permalink = await GetPermalinkAsync(response.Channel, summaryTs);
+						permalink = await GetPermalinkAsync(state.Channel, summaryTs);
 					}
 
 					List<IIssueSuspect> suspects = await _issueService.GetIssueSuspectsAsync(issue);
@@ -974,7 +967,10 @@ namespace Horde.Build.Notifications.Impl
 				attachment.Blocks.Add(new SectionBlock(String.Join("\n", declinedLines)));
 			}
 
-			await SendOrUpdateMessageAsync(recipient, GetIssueEventId(issue), userId, attachments: new[] { attachment });
+			if (IsRecipientAllowed(recipient, "issue update"))
+			{
+				await SendOrUpdateMessageAsync(recipient, GetIssueEventId(issue), userId, attachments: new[] { attachment });
+			}
 		}
 
 		static string GetIssueEventId(IIssue issue)
@@ -1501,9 +1497,8 @@ namespace Horde.Build.Notifications.Impl
 
 		private async Task SendMessageAsync(string recipient, string? text = null, BlockBase[]? blocks = null, BlockKitAttachment[]? attachments = null, bool withEnvironment = true)
 		{
-			if (_allowUsers != null && !_allowUsers.Contains(recipient) && !recipient.StartsWith("#", StringComparison.Ordinal))
+			if (!IsRecipientAllowed(recipient, text))
 			{
-				_logger.LogDebug("Suppressing message to {Recipient}: {Text}", recipient, text);
 				return;
 			}
 
@@ -1525,14 +1520,18 @@ namespace Horde.Build.Notifications.Impl
 			await SendRequestAsync<PostMessageResponse>(PostMessageUrl, message);
 		}
 
-		private async Task SendOrUpdateMessageAsync(string recipient, string eventId, UserId? userId, string? text = null, BlockBase[]? blocks = null, BlockKitAttachment[]? attachments = null)
+		bool IsRecipientAllowed(string recipient, string? description)
 		{
-			if (_allowUsers != null && !_allowUsers.Contains(recipient))
+			if (_allowUsers != null && !_allowUsers.Contains(recipient) && !recipient.StartsWith("#", StringComparison.Ordinal))
 			{
-				_logger.LogDebug("Suppressing message to {Recipient}: {Text}", recipient, text);
-				return;
+				_logger.LogDebug("Suppressing message to {Recipient}: {Description}", recipient, description);
+				return false;
 			}
+			return true;
+		}
 
+		private async Task<(MessageStateDocument, bool)> SendOrUpdateMessageAsync(string recipient, string eventId, UserId? userId, string? text = null, BlockBase[]? blocks = null, BlockKitAttachment[]? attachments = null)
+		{
 			BlockKitMessage message = new BlockKitMessage();
 			message.Text = AddEnvironmentAnnotation(text);
 
@@ -1568,6 +1567,7 @@ namespace Horde.Build.Notifications.Impl
 				message.Ts = state.Ts;
 				await SendRequestAsync<PostMessageResponse>(UpdateMessageUrl, message);
 			}
+			return (state, isNew);
 		}
 
 		private async Task<string?> GetPermalinkAsync(string channel, string messageTs)
