Horde: Adding workflows to stream issue summary
#skipci
#rnx

diff --git a/Engine/Source/Programs/Horde/HordeDashboard/src/backend/Api.ts b/Engine/Source/Programs/Horde/HordeDashboard/src/backend/Api.ts
--- a/Engine/Source/Programs/Horde/HordeDashboard/src/backend/Api.ts
+++ b/Engine/Source/Programs/Horde/HordeDashboard/src/backend/Api.ts
@@ -1777,6 +1777,18 @@ export type DefaultPreflightRequest = {
 	changeTemplateId?: string;
 }
 
+	
+/**  Configuration for an issue workflow */	
+export type WorkflowConfig = {
+
+	/** Identifier for this workflow */
+	id: string;
+	
+	/** Name of the tab to post summary data to */
+	summaryTab?: string;
+
+}
+
 /**Response describing a stream */
 export type GetStreamResponse = {
 
@@ -1816,6 +1828,9 @@ export type GetStreamResponse = {
 	/**Properties for this stream */
 	properties: { [key: string]: string };
 
+	/** Workflows for this stream */
+	workflows: WorkflowConfig[];
+
 	/** Custom permissions for this object */
 	acl?: GetAclResponse;
 
@@ -2570,6 +2585,9 @@ export type GetIssueSpanResponse = {
 	/**The previous build  */
 	lastSuccess?: GetIssueStepResponse;
 
+	/// Workflow that this span belongs to
+	workflowId?: string;
+
 	/**The failing builds for a particular event */
 	steps: GetIssueStepResponse[];
 
@@ -3749,6 +3767,9 @@ export type FindIssueSpanResponse = {
 	/// Name of the step
 	name: string;
 
+	/// Workflow for this span
+	workflowId?: string;
+
 	/// The previous build 
 	lastSuccess?: GetIssueStepResponse;
 
@@ -3833,4 +3854,4 @@ export type GetJiraIssueResponse = {
 	assigneeDisplayName?: string;
 	assigneeEmailAddress?: string;
 
-}
\ No newline at end of file
+}
