Horde: Server side of pause/quarantine step support
#skipci
#rnx

diff --git a/Engine/Source/Programs/Horde/Horde.Build.Tests/IssueServiceTests.cs b/Engine/Source/Programs/Horde/Horde.Build.Tests/IssueServiceTests.cs
--- a/Engine/Source/Programs/Horde/Horde.Build.Tests/IssueServiceTests.cs
+++ b/Engine/Source/Programs/Horde/Horde.Build.Tests/IssueServiceTests.cs
@@ -15,6 +15,7 @@ using Horde.Build.Server;
 using Horde.Build.Services;
 using Horde.Build.Tests.Stubs.Services;
 using Horde.Build.Utilities;
+using Horde.Build.Api;
 using HordeCommon;
 using Microsoft.Extensions.DependencyInjection;
 using Microsoft.Extensions.Logging;
@@ -1464,7 +1465,72 @@ namespace Horde.Build.Tests
 			}
 		}
 
-				private async Task ParseAsync(LogId logId, string[] lines)
+		[TestMethod]
+		public async Task QuarantinedIssueTest()
+		{
+			IStream stream = Deref(await StreamCollection.GetAsync(_mainStreamId));
+			TemplateRefId templateRefId = new TemplateRefId("test-template");
+
+			// #1
+			// Scenario: Warning in first step, step is quarantined
+			// Expected: No issue generated
+			{
+				stream = Deref(await StreamService.TryUpdateTemplateRefAsync(stream, templateRefId, new List<UpdateStepStateRequest>() { new UpdateStepStateRequest() { Name = "Update Version Files", QuarantinedByUserId = UserId.GenerateNewId().ToString() } }));
+
+				IJob job = CreateJob(_mainStreamId, 105, "Test Build", _graph);
+				await AddEvent(job, 0, 0, new { level = nameof(LogLevel.Warning) }, EventSeverity.Warning);
+				await UpdateCompleteStep(job, 0, 0, JobStepOutcome.Warnings);
+
+				List<IIssue> issues = await IssueService.FindIssuesAsync();
+				Assert.AreEqual(0, issues.Count);				
+			}
+
+			// #2
+			// Scenario: Warning in first step, step is not quarantined
+			// Expected: issue is generated
+			{
+				stream = Deref(await StreamService.TryUpdateTemplateRefAsync(stream, templateRefId, new List<UpdateStepStateRequest>()));
+
+				// step states should be cleared now
+				Assert.AreEqual(null, stream.Templates[templateRefId].StepStates);
+
+				IJob job = CreateJob(_mainStreamId, 105, "Test Build", _graph);
+				await AddEvent(job, 0, 0, new { level = nameof(LogLevel.Warning) }, EventSeverity.Warning);
+				await UpdateCompleteStep(job, 0, 0, JobStepOutcome.Warnings);
+
+				List<IIssue> issues = await IssueService.FindIssuesAsync();
+				Assert.AreEqual(1, issues.Count);
+
+				Assert.AreEqual("Warnings in Update Version Files", issues[0].Summary);
+
+				List<IIssueSpan> spans = await IssueCollection.FindSpansAsync();
+				Assert.AreEqual(1, spans.Count);
+				Assert.AreEqual(job.Id, spans[0].LastFailure?.JobId);
+			}
+
+			// #3
+			// Scenario: Warning in first step, step is quarantined, issue exists
+			// Expected: issue is updated with new failure
+			{
+				stream = Deref(await StreamService.TryUpdateTemplateRefAsync(stream, templateRefId, new List<UpdateStepStateRequest>() { new UpdateStepStateRequest() { Name = "Update Version Files", QuarantinedByUserId = UserId.GenerateNewId().ToString() } }));
+
+
+				IJob job = CreateJob(_mainStreamId, 105, "Test Build", _graph);
+				await AddEvent(job, 0, 0, new { level = nameof(LogLevel.Warning) }, EventSeverity.Warning);
+				await UpdateCompleteStep(job, 0, 0, JobStepOutcome.Warnings);
+
+				List<IIssue> issues = await IssueService.FindIssuesAsync();
+				Assert.AreEqual(1, issues.Count);
+
+				List<IIssueSpan> spans = await IssueCollection.FindSpansAsync();
+				Assert.AreEqual(1, spans.Count);
+				Assert.AreEqual(job.Id, spans[0].LastFailure?.JobId);
+			}
+
+		}
+
+
+		private async Task ParseAsync(LogId logId, string[] lines)
 		{
 			LogParserContext context = new LogParserContext();
 			context.WorkspaceDir = DirectoryReference.GetCurrentDirectory();
diff --git a/Engine/Source/Programs/Horde/Horde.Build.Tests/JobTaskSourceTests.cs b/Engine/Source/Programs/Horde/Horde.Build.Tests/JobTaskSourceTests.cs
--- a/Engine/Source/Programs/Horde/Horde.Build.Tests/JobTaskSourceTests.cs
+++ b/Engine/Source/Programs/Horde/Horde.Build.Tests/JobTaskSourceTests.cs
@@ -8,15 +8,40 @@ using Horde.Build.Models;
 using Horde.Build.Utilities;
 using HordeCommon;
 using Microsoft.VisualStudio.TestTools.UnitTesting;
+using Horde.Build.Services;
+using System.Linq;
+using System;
+
 
 namespace Horde.Build.Tests
 {
+	using JobId = ObjectId<IJob>;
+	using UserId = ObjectId<IUser>;
+
 	[TestClass]
 	public class JobTaskSourceTests : TestSetup
 	{
 		private bool _eventReceived;
 		private bool? _eventPoolHasAgentsOnline;
-		
+
+		static NewGroup AddGroup(List<NewGroup> groups)
+		{
+			NewGroup @group = new NewGroup("Win64", new List<NewNode>());
+			groups.Add(@group);
+			return @group;
+		}
+
+		static NewNode AddNode(NewGroup @group, string name, string[]? inputDependencies, Action<NewNode>? action = null)
+		{
+			NewNode node = new NewNode(name, inputDependencies?.ToList(), inputDependencies?.ToList(), null, null, null, null, null, null);
+			if (action != null)
+			{
+				action.Invoke(node);
+			}
+			@group.Nodes.Add(node);
+			return node;
+		}
+
 		[TestMethod]
 		public async Task UpdateJobQueueNormal()
 		{
@@ -78,6 +103,62 @@ namespace Horde.Build.Tests
 			Assert.IsFalse(_eventPoolHasAgentsOnline!.Value);
 		}
 
+		[TestMethod]
+		public async Task UpdateJobQueueWithPausedStep()
+		{
+			Fixture fixture = await SetupPoolWithAgentAsync(isPoolAutoScaled: true, shouldCreateAgent: true, isAgentEnabled: true);
+
+			// update template with some step states
+			IStream Stream = Deref(await StreamService.TryUpdateTemplateRefAsync(fixture.Stream!, fixture.TemplateRefId1, new List<UpdateStepStateRequest>() { new UpdateStepStateRequest() { Name = "Quarantined Step", QuarantinedByUserId = UserId.GenerateNewId().ToString() } }));
+			Stream = Deref(await StreamService.TryUpdateTemplateRefAsync(Stream, fixture.TemplateRefId1, new List<UpdateStepStateRequest>() { new UpdateStepStateRequest() { Name = "Paused Step", PausedByUserId = UserId.GenerateNewId().ToString() } }));
+
+			// create a new graph with the associated nodes
+			List<NewGroup> newGroups = new List<NewGroup>();
+
+			NewGroup initialGroup = AddGroup(newGroups);			
+			AddNode(initialGroup, "Update Version Files", null);
+			AddNode(initialGroup, "Paused Step", new[] { "Update Version Files" });
+			AddNode(initialGroup, "Step That Depends on Paused Step", new[] { "Paused Step" });
+			AddNode(initialGroup, "Step That Depends on Update Version Files", new[] { "Update Version Files" });
+
+			IGraph graph = await GraphCollection.AppendAsync(null, newGroups);
+
+			// remove the default fixture jobs
+			List<IJob> Jobs = await JobCollection.FindAsync();
+			for (int i = 0; i < Jobs.Count; i++)
+			{
+				await JobCollection.RemoveAsync(Jobs[i]);
+			}
+
+			// create a new job
+			IJob job = await JobCollection.AddAsync(JobId.GenerateNewId(), Stream.Id,
+				fixture.TemplateRefId1, fixture.Template.Id, graph, "Test Paused Step Job",
+				1000, 1000, null, null, null, Priority.Highest, null, null, null, null, false,
+				false, null, null, new List<string> { "-Target=" + "Step That Depends on Paused Step;Step That Depends on Update Version Files" });
+
+			// validate
+			Assert.AreEqual(0, JobTaskSource.GetQueueForTesting().Count);
+			await JobTaskSource.TickAsync(CancellationToken.None);
+			Assert.AreEqual(1, JobTaskSource.GetQueueForTesting().Count);
+			Assert.AreEqual(job.Id, JobTaskSource.GetQueueForTesting().Min!.Id.Item1);
+			Assert.AreEqual(JobStepBatchState.Ready, JobTaskSource.GetQueueForTesting().Min!.Batch.State);
+
+			Assert.IsTrue(_eventReceived);
+			Assert.IsTrue(_eventPoolHasAgentsOnline!.Value);
+
+			// make sure we get db value
+			job = Deref(await JobService.GetJobAsync(job.Id));
+
+			IJobStepBatch batch = job.Batches[0];
+
+			Assert.AreEqual(batch.Steps.Count, 4);
+			Assert.AreEqual(batch.Steps[0].State, JobStepState.Ready);
+			Assert.AreEqual(batch.Steps[1].State, JobStepState.Skipped);
+			Assert.AreEqual(batch.Steps[2].State, JobStepState.Skipped);
+			Assert.AreEqual(batch.Steps[3].State, JobStepState.Waiting);
+
+		}
+
 		private async Task<Fixture> SetupPoolWithAgentAsync(bool isPoolAutoScaled, bool shouldCreateAgent, bool isAgentEnabled)
 		{
 			Fixture fixture = await CreateFixtureAsync();
diff --git a/Engine/Source/Programs/Horde/Horde.Build/Api/Streams.cs b/Engine/Source/Programs/Horde/Horde.Build/Api/Streams.cs
--- a/Engine/Source/Programs/Horde/Horde.Build/Api/Streams.cs
+++ b/Engine/Source/Programs/Horde/Horde.Build/Api/Streams.cs
@@ -246,6 +246,43 @@ namespace Horde.Build.Api
 		public UpdateAclRequest? Acl { get; set; }
 	}
 
+	/// <summary>
+	/// Step state update request
+	/// </summary>
+	public class UpdateStepStateRequest
+	{
+		/// <summary>
+		/// Name of the step
+		/// </summary>
+		[Required]
+		public string Name { get; set; } = String.Empty;
+
+		/// <summary>
+		/// User who paused the step
+		/// </summary>
+		public string? PausedByUserId { get; set; }
+
+		/// <summary>
+		/// User who quarantined the step
+		/// </summary>
+		public string? QuarantinedByUserId { get; set; }
+	}
+
+
+	/// <summary>
+	/// Updates an existing stream template ref
+	/// </summary>
+	public class UpdateTemplateRefRequest
+	{
+		
+		/// <summary>
+		/// Step states to update
+		/// </summary>
+		public List<UpdateStepStateRequest>? StepStates { get; set; }
+
+	}
+
+
 	/// <summary>
 	/// Mapping from a BuildGraph agent type to a set of machines on the farm
 	/// </summary>
diff --git a/Engine/Source/Programs/Horde/Horde.Build/Controllers/StreamsController.cs b/Engine/Source/Programs/Horde/Horde.Build/Controllers/StreamsController.cs
--- a/Engine/Source/Programs/Horde/Horde.Build/Controllers/StreamsController.cs
+++ b/Engine/Source/Programs/Horde/Horde.Build/Controllers/StreamsController.cs
@@ -31,16 +31,17 @@ namespace Horde.Build.Controllers
 		private readonly ITemplateCollection _templateCollection;
 		private readonly IJobStepRefCollection _jobStepRefCollection;
 		private readonly IPerforceService _perforceService;
-
+		private readonly AclService _aclService;
 		/// <summary>
 		/// Constructor
 		/// </summary>
-		public StreamsController(StreamService streamService, ITemplateCollection templateCollection, IJobStepRefCollection jobStepRefCollection, IPerforceService perforceService)
+		public StreamsController(AclService aclService, StreamService streamService, ITemplateCollection templateCollection, IJobStepRefCollection jobStepRefCollection, IPerforceService perforceService)
 		{
 			_streamService = streamService;
 			_templateCollection = templateCollection;
 			_jobStepRefCollection = jobStepRefCollection;
 			_perforceService = perforceService;
+			_aclService = aclService;
 		}
 
 		/// <summary>
@@ -256,7 +257,7 @@ namespace Horde.Build.Controllers
 		}
 
 		/// <summary>
-		/// Gets all the templates for a stream
+		/// Gets a template for a stream
 		/// </summary>
 		/// <param name="streamId">Unique id of the stream to query</param>
 		/// <param name="templateId">Unique id of the template to query</param>
@@ -291,5 +292,36 @@ namespace Horde.Build.Controllers
 
 			return new GetTemplateResponse(template).ApplyFilter(filter);
 		}
+
+		/// <summary>
+		/// Update a stream template ref
+		/// </summary>
+		[HttpPut]
+		[Authorize]
+		[Route("/api/v1/streams/{streamId}/templates/{templateRefId}")]
+		public async Task<ActionResult> UpdateStreamTemplateRefAsync(StreamId streamId, TemplateRefId templateRefId, [FromBody] UpdateTemplateRefRequest update)
+		{
+			if (!await _aclService.AuthorizeAsync(AclAction.AdminWrite, User))
+			{
+				return Forbid(AclAction.AdminWrite);
+			}
+
+			IStream? stream = await _streamService.GetStreamAsync(streamId);
+			if (stream == null)
+			{
+				return NotFound(streamId);
+			}
+
+			if (!stream.Templates.ContainsKey(templateRefId))
+			{
+				return NotFound(streamId, templateRefId);
+			}
+
+			await _streamService.TryUpdateTemplateRefAsync(stream, templateRefId, update.StepStates);
+
+			return Ok();
+
+		}
+
 	}
 }
diff --git a/Engine/Source/Programs/Horde/Horde.Build/Issues/Impl/IssueService.cs b/Engine/Source/Programs/Horde/Horde.Build/Issues/Impl/IssueService.cs
--- a/Engine/Source/Programs/Horde/Horde.Build/Issues/Impl/IssueService.cs
+++ b/Engine/Source/Programs/Horde/Horde.Build/Issues/Impl/IssueService.cs
@@ -572,6 +572,23 @@ namespace Horde.Build.Services.Impl
 			{
 				throw new ArgumentException($"Invalid step id {stepId}");
 			}
+
+			TemplateRef? template;
+			bool SuppressNewIssues = false;
+			if (stream.Templates.TryGetValue(job.TemplateId, out template))
+			{
+				if (template.StepStates != null)
+				{
+					string StepName = graph.Groups[batch.GroupIdx].Nodes[step.NodeIdx].Name;
+					TemplateStepState? state = template.StepStates.FirstOrDefault(x => x.Name.Equals(StepName, StringComparison.Ordinal));
+					if (state?.QuarantinedByUserId != null)
+					{
+						_logger.LogInformation("Issue creation disabled for {JobId}:{BatchId}:{StepId}, Quaratined by user: {UserId}", job.Id, batchId, stepId, state.QuarantinedByUserId);
+						SuppressNewIssues = true;
+					}
+				}
+			}
+
 			scope.Span.SetTag("LogId", step.LogId.ToString());
 
 			// Get the node associated with this step
@@ -597,9 +614,12 @@ namespace Horde.Build.Services.Impl
 						{
 							continue;
 						}
-						if (!await AddEventsToNewSpans(stream, job, batch, step, node, openSpans, eventGroups, job.PromoteIssuesByDefault))
+						if (!SuppressNewIssues)
 						{
-							continue;
+							if (!await AddEventsToNewSpans(stream, job, batch, step, node, openSpans, eventGroups, job.PromoteIssuesByDefault))
+							{
+								continue;
+							}
 						}
 					}
 
diff --git a/Engine/Source/Programs/Horde/Horde.Build/Jobs/Impl/JobTaskSource.cs b/Engine/Source/Programs/Horde/Horde.Build/Jobs/Impl/JobTaskSource.cs
--- a/Engine/Source/Programs/Horde/Horde.Build/Jobs/Impl/JobTaskSource.cs
+++ b/Engine/Source/Programs/Horde/Horde.Build/Jobs/Impl/JobTaskSource.cs
@@ -436,15 +436,49 @@ namespace Horde.Build.Tasks.Impl
 					}
 					else
 					{
-						(AgentWorkspace workspace, bool useAutoSdk) = workspaceResult.Value;
-						QueueItem newQueueItem = new QueueItem(stream, newJob, batchIdx, agentType.Pool, workspace, useAutoSdk);
-						newQueue.Add(newQueueItem);
-						newBatchIdToQueueItem[(newJob.Id, batch.Id)] = newQueueItem;
-						
-						IPool? newJobPool = pools.Find(p => p.Id == agentType.Pool);
-						if (newJobPool != null)
+						TemplateRef? templateRef;
+						if (stream.Templates.TryGetValue(newJob.TemplateId, out templateRef))
+						{							
+							if (templateRef.StepStates != null)
+							{
+								for (int i = 0; i < templateRef.StepStates.Count; i++)
+								{
+									TemplateStepState state = templateRef.StepStates[i];
+									if (state.PausedByUserId != null)
+									{
+										IJobStep? step = batch.Steps.FirstOrDefault(x => graph.Groups[batch.GroupIdx].Nodes[x.NodeIdx].Name.Equals(state.Name, StringComparison.Ordinal));
+
+										if (step != null)
+										{
+											JobId jobId = newJob.Id;
+											newJob = await _jobs.TryUpdateStepAsync(newJob, graph, batch.Id, step.Id, JobStepState.Skipped);
+											if (newJob == null)
+											{
+												_logger.LogError("Job {JobId} failed to update step {StepName} pause state", jobId, state.Name);
+												break;
+											}
+											else
+											{
+												_logger.LogInformation("Job {JobId} step {StepName} has been skipped due to being paused", jobId, state.Name);
+											}
+										}
+									}
+								}
+							}
+						}
+
+						if (newJob != null)
 						{
-							OnJobScheduled?.Invoke(newJobPool, IsPoolOnline(agentType.Pool), newJob, graph, batch.Id);
+							(AgentWorkspace workspace, bool useAutoSdk) = workspaceResult.Value;
+							QueueItem newQueueItem = new QueueItem(stream, newJob, batchIdx, agentType.Pool, workspace, useAutoSdk);
+							newQueue.Add(newQueueItem);
+							newBatchIdToQueueItem[(newJob.Id, batch.Id)] = newQueueItem;
+
+							IPool? newJobPool = pools.Find(p => p.Id == agentType.Pool);
+							if (newJobPool != null)
+							{
+								OnJobScheduled?.Invoke(newJobPool, IsPoolOnline(agentType.Pool), newJob, graph, batch.Id);
+							}
 						}
 					}
 				}
diff --git a/Engine/Source/Programs/Horde/Horde.Build/Projects/IStream.cs b/Engine/Source/Programs/Horde/Horde.Build/Projects/IStream.cs
--- a/Engine/Source/Programs/Horde/Horde.Build/Projects/IStream.cs
+++ b/Engine/Source/Programs/Horde/Horde.Build/Projects/IStream.cs
@@ -1,4 +1,4 @@
-﻿// Copyright Epic Games, Inc. All Rights Reserved.
+// Copyright Epic Games, Inc. All Rights Reserved.
 
 using System;
 using System.Collections.Generic;
@@ -19,6 +19,7 @@ namespace Horde.Build.Models
 	using ProjectId = StringId<IProject>;
 	using StreamId = StringId<IStream>;
 	using TemplateRefId = StringId<TemplateRef>;
+	using UserId = ObjectId<IUser>;
 
 	/// <summary>
 	/// Exception thrown when stream validation fails
@@ -284,6 +285,46 @@ namespace Horde.Build.Models
 		}
 	}
 
+	/// <summary>
+	/// A
+	/// </summary>
+	public class TemplateStepState
+	{
+		/// <summary>
+		/// Name of the step
+		/// </summary>
+		public string Name { get; set; } = String.Empty ;
+
+		/// <summary>
+		/// User who paused the step
+		/// </summary>
+		public UserId? PausedByUserId { get; set; }
+
+		/// <summary>
+		/// User who quarantined the step
+		/// </summary>
+		public UserId? QuarantinedByUserId { get; set; }
+
+		/// <summary>
+		/// Default constructor for serialization
+		/// </summary>
+		private TemplateStepState()
+		{
+		}
+
+		/// <summary>
+		/// Constructor
+		/// </summary>
+		public TemplateStepState( string name, UserId? pausedByUserId = null, UserId? quarantinedByUserId = null)
+		{
+			Name = name;
+			PausedByUserId = pausedByUserId;
+			QuarantinedByUserId = quarantinedByUserId;
+		}
+
+	}
+
+
 	/// <summary>
 	/// Reference to a template
 	/// </summary>
@@ -336,6 +377,12 @@ namespace Horde.Build.Models
 		[BsonIgnoreIfNull]
 		public List<ChainedJobTemplate>? ChainedJobs { get; set; }
 
+		/// <summary>
+		/// List of template step states
+		/// </summary>
+		[BsonIgnoreIfNull]
+		public List<TemplateStepState>? StepStates { get;set; }
+
 		/// <summary>
 		/// Custom permissions for this template
 		/// </summary>
@@ -361,8 +408,9 @@ namespace Horde.Build.Models
 		/// <param name="triageChannel"></param>
 		/// <param name="schedule">Schedule for this template</param>
 		/// <param name="triggers">List of downstream templates to trigger</param>
+		/// <param name="stepStates">List of template step states</param>
 		/// <param name="acl">ACL for this template</param>
-		public TemplateRef(ITemplate template, bool showUgsBadges = false, bool showUgsAlerts = false, string? notificationChannel = null, string? notificationChannelFilter = null, string? triageChannel = null, Schedule? schedule = null, List<ChainedJobTemplate>? triggers = null, Acl? acl = null)
+		public TemplateRef(ITemplate template, bool showUgsBadges = false, bool showUgsAlerts = false, string? notificationChannel = null, string? notificationChannelFilter = null, string? triageChannel = null, Schedule? schedule = null, List<ChainedJobTemplate>? triggers = null, List<TemplateStepState>? stepStates = null, Acl? acl = null)
 		{
 			Name = template.Name;
 			Hash = template.Id;
@@ -373,6 +421,7 @@ namespace Horde.Build.Models
 			TriageChannel = triageChannel;
 			Schedule = schedule;
 			ChainedJobs = triggers;
+			StepStates = stepStates;
 			Acl = acl;
 		}
 	}
diff --git a/Engine/Source/Programs/Horde/Horde.Build/Projects/IStreamCollection.cs b/Engine/Source/Programs/Horde/Horde.Build/Projects/IStreamCollection.cs
--- a/Engine/Source/Programs/Horde/Horde.Build/Projects/IStreamCollection.cs
+++ b/Engine/Source/Programs/Horde/Horde.Build/Projects/IStreamCollection.cs
@@ -3,6 +3,7 @@
 using System;
 using System.Collections.Generic;
 using System.Threading.Tasks;
+using Horde.Build.Api;
 using Horde.Build.Models;
 using Horde.Build.Server;
 using Horde.Build.Utilities;
@@ -76,6 +77,15 @@ namespace Horde.Build.Collections
 		/// <returns>The updated stream if successful, null otherwise</returns>
 		Task<IStream?> TryUpdateScheduleTriggerAsync(IStream stream, TemplateRefId templateRefId, DateTime? lastTriggerTimeUtc, int? lastTriggerChange, List<JobId> newActiveJobs);
 
+		/// <summary>
+		/// Attempts to update a stream template ref
+		/// </summary>
+		/// <param name="streamInterface">The stream containing the template ref</param>
+		/// <param name="templateRefId">The template ref to update</param>
+		/// <param name="stepStates">The stream states to update, pass an empty list to clear all step states, otherwise will be a partial update based on included step updates</param>
+		/// <returns></returns>
+		Task<IStream?> TryUpdateTemplateRefAsync(IStream streamInterface, TemplateRefId templateRefId, List<UpdateStepStateRequest>? stepStates = null);
+
 		/// <summary>
 		/// Delete a stream
 		/// </summary>
diff --git a/Engine/Source/Programs/Horde/Horde.Build/Projects/Impl/StreamCollection.cs b/Engine/Source/Programs/Horde/Horde.Build/Projects/Impl/StreamCollection.cs
--- a/Engine/Source/Programs/Horde/Horde.Build/Projects/Impl/StreamCollection.cs
+++ b/Engine/Source/Programs/Horde/Horde.Build/Projects/Impl/StreamCollection.cs
@@ -22,6 +22,7 @@ namespace Horde.Build.Collections.Impl
 	using ProjectId = StringId<IProject>;
 	using StreamId = StringId<IStream>;
 	using TemplateRefId = StringId<TemplateRef>;
+	using UserId = ObjectId<IUser>;
 
 	/// <summary>
 	/// Collection of stream documents
@@ -194,13 +195,19 @@ namespace Horde.Build.Collections.Impl
 				}
 
 				// Add it to the list
-				TemplateRef newTemplateRef = new TemplateRef(newTemplate, request.ShowUgsBadges, request.ShowUgsAlerts, request.NotificationChannel, request.NotificationChannelFilter, request.TriageChannel, schedule, request.ChainedJobs?.ConvertAll(x => new ChainedJobTemplate(x)), Acl.Merge(null, request.Acl));
+				TemplateRef newTemplateRef = new TemplateRef(newTemplate, request.ShowUgsBadges, request.ShowUgsAlerts, request.NotificationChannel, request.NotificationChannelFilter, request.TriageChannel, schedule, request.ChainedJobs?.ConvertAll(x => new ChainedJobTemplate(x)), null, Acl.Merge(null, request.Acl));
 				if (stream != null && stream.Templates.TryGetValue(newTemplateRefId, out TemplateRef? oldTemplateRef))
 				{
 					if (oldTemplateRef.Schedule != null && newTemplateRef.Schedule != null)
 					{
 						newTemplateRef.Schedule.CopyState(oldTemplateRef.Schedule);
 					}
+
+					if (oldTemplateRef.StepStates != null)
+					{
+						newTemplateRef.StepStates = new List<TemplateStepState>(oldTemplateRef.StepStates);
+					}
+
 				}
 				newTemplateRefs.Add(newTemplateRefId, newTemplateRef);
 			}
@@ -220,6 +227,101 @@ namespace Horde.Build.Collections.Impl
 			return newTemplateRefs;
 		}
 
+		/// <inheritdoc/>
+		public async Task<IStream?> TryUpdateTemplateRefAsync(IStream streamInterface, TemplateRefId templateRefId, List<UpdateStepStateRequest>? stepStates = null)
+		{
+			StreamDocument stream = (StreamDocument)streamInterface;
+
+			TemplateRef? oldTemplateRef;
+			TemplateRef? newTemplateRef = null;
+
+			if (!stream.Templates.TryGetValue(templateRefId, out oldTemplateRef))
+			{
+				return null;
+			}
+
+			ITemplate? template = await _templateCollection.GetAsync(oldTemplateRef.Hash);
+			if (template == null)
+			{
+				return null;
+			}
+
+			UpdateDefinitionBuilder<StreamDocument> updateBuilder = Builders<StreamDocument>.Update;
+			List<UpdateDefinition<StreamDocument>> updates = new List<UpdateDefinition<StreamDocument>>();
+
+			Dictionary<TemplateRefId, TemplateRef> newTemplates = new Dictionary<TemplateRefId, TemplateRef>(stream.Templates);
+
+			// clear
+			if (stepStates != null && stepStates.Count == 0)
+			{
+				bool hasUpdates = false;
+				foreach (KeyValuePair<TemplateRefId, TemplateRef> entry in newTemplates)
+				{
+					if (entry.Value?.StepStates != null)
+					{
+						hasUpdates = true;
+						entry.Value.StepStates = null;
+					}
+				}
+
+				if (hasUpdates)
+				{
+					updates.Add(updateBuilder.Set(x => x.Templates, newTemplates));
+				}
+				
+			}
+			else if (stepStates != null)
+			{
+				newTemplateRef = new TemplateRef(template, oldTemplateRef.ShowUgsAlerts, oldTemplateRef.ShowUgsAlerts, oldTemplateRef.NotificationChannel, oldTemplateRef.NotificationChannelFilter, oldTemplateRef.TriageChannel, oldTemplateRef.Schedule, oldTemplateRef.ChainedJobs, null, oldTemplateRef.Acl);
+
+				// get currently valid step states
+				List<TemplateStepState> newStepStates = oldTemplateRef.StepStates?.Where(x => x.QuarantinedByUserId != null || x.PausedByUserId != null).ToList() ?? new List<TemplateStepState>();
+
+				// generate update list
+				foreach (UpdateStepStateRequest updateState in stepStates)
+				{
+					TemplateStepState? newState = newStepStates.Where(x => x.Name == updateState.Name).FirstOrDefault();
+
+					UserId? PausedByUserId = updateState.PausedByUserId != null ? new UserId(updateState.PausedByUserId) : null;
+					UserId? QuarantinedByUserId = updateState.QuarantinedByUserId != null ? new UserId(updateState.QuarantinedByUserId) : null;
+
+					if (newState == null)
+					{
+						// if this is a new state without anything set, ignore it
+						if (PausedByUserId == null && QuarantinedByUserId == null)
+						{
+							continue;
+						}
+
+						newStepStates.Add(new TemplateStepState(updateState.Name, PausedByUserId, QuarantinedByUserId));
+					}
+					else
+					{
+						newState.PausedByUserId = PausedByUserId;
+						newState.QuarantinedByUserId = QuarantinedByUserId;
+					}
+				}
+
+				newTemplateRef.StepStates = newStepStates.Where(x => x.QuarantinedByUserId != null || x.PausedByUserId != null).ToList();
+				if (newTemplateRef.StepStates.Count == 0)
+				{
+					newTemplateRef.StepStates = null;
+				}
+
+				newTemplates[templateRefId] = newTemplateRef;
+				updates.Add(updateBuilder.Set(x => x.Templates, newTemplates));
+
+			}
+
+			if (updates.Count == 0)
+			{
+				return streamInterface;
+			}
+
+			return await TryUpdateStreamAsync(stream, updateBuilder.Combine(updates));
+
+		}
+
 		/// <inheritdoc/>
 		async Task<IStream?> TryCreateAsync(StreamId id, ProjectId projectId, string configRevision, StreamConfig config, DefaultPreflight? defaultPreflight, List<StreamTab> tabs, Dictionary<string, AgentType> agentTypes, Dictionary<string, WorkspaceType> workspaceTypes, Dictionary<TemplateRefId, TemplateRef> templateRefs, Acl? acl)
 		{
diff --git a/Engine/Source/Programs/Horde/Horde.Build/Projects/Impl/StreamService.cs b/Engine/Source/Programs/Horde/Horde.Build/Projects/Impl/StreamService.cs
--- a/Engine/Source/Programs/Horde/Horde.Build/Projects/Impl/StreamService.cs
+++ b/Engine/Source/Programs/Horde/Horde.Build/Projects/Impl/StreamService.cs
@@ -1,4 +1,4 @@
-﻿// Copyright Epic Games, Inc. All Rights Reserved.
+// Copyright Epic Games, Inc. All Rights Reserved.
 
 using System;
 using System.Collections.Generic;
@@ -6,6 +6,7 @@ using System.Linq;
 using System.Security.Claims;
 using System.Threading.Tasks;
 using Horde.Build.Acls;
+using Horde.Build.Api;
 using Horde.Build.Collections;
 using Horde.Build.Models;
 using Horde.Build.Utilities;
@@ -151,6 +152,18 @@ namespace Horde.Build.Services
 			return null;
 		}
 
+		/// <summary>
+		/// Attempts to update a stream template ref
+		/// </summary>
+		/// <param name="streamInterface"></param>
+		/// <param name="templateRefId"></param>
+		/// <param name="stepStates"></param>
+		/// <returns></returns>
+		public Task<IStream?> TryUpdateTemplateRefAsync(IStream streamInterface, TemplateRefId templateRefId, List<UpdateStepStateRequest>? stepStates = null)
+		{
+			return _streams.TryUpdateTemplateRefAsync(streamInterface, templateRefId, stepStates );
+		}
+
 		/// <summary>
 		/// Gets all the available streams for a project
 		/// </summary>
