Horde: Allow specifying any default node annotations on the template and workflow, and allow overriding the issue workflow to use on a per-step basis.

diff --git a/Engine/Source/Programs/Horde/Horde.Build/Api/Issues.cs b/Engine/Source/Programs/Horde/Horde.Build/Api/Issues.cs
--- a/Engine/Source/Programs/Horde/Horde.Build/Api/Issues.cs
+++ b/Engine/Source/Programs/Horde/Horde.Build/Api/Issues.cs
@@ -122,13 +122,12 @@ namespace Horde.Build.Api
 		/// </summary>
 		/// <param name="span">The node to construct from</param>
 		/// <param name="steps">Failing steps for this span</param>
-		/// <param name="workflowId"></param>
-		public GetIssueSpanResponse(IIssueSpan span, List<IIssueStep> steps, WorkflowId? workflowId)
+		public GetIssueSpanResponse(IIssueSpan span, List<IIssueStep> steps)
 		{
 			Id = span.Id.ToString();
 			Name = span.NodeName;
 			TemplateId = span.TemplateRefId.ToString();
-			WorkflowId = workflowId;
+			WorkflowId = span.LastFailure.Annotations.WorkflowId;
 			LastSuccess = (span.LastSuccess != null) ? new GetIssueStepResponse(span.LastSuccess) : null;
 			Steps = steps.ConvertAll(x => new GetIssueStepResponse(x));
 			NextSuccess = (span.NextSuccess != null) ? new GetIssueStepResponse(span.NextSuccess) : null;
@@ -180,13 +179,7 @@ namespace Horde.Build.Api
 				{
 					MaxChange = span.NextSuccess.Change;
 				}
-
-				WorkflowId? workflowId = null;
-				if (stream.Config.TryGetTemplate(span.TemplateRefId, out TemplateRefConfig? templateConfig))
-				{
-					workflowId = templateConfig.Workflow;
-				}
-				Nodes.Add(new GetIssueSpanResponse(span, steps.Where(y => y.SpanId == span.Id).ToList(), workflowId));
+				Nodes.Add(new GetIssueSpanResponse(span, steps.Where(y => y.SpanId == span.Id).ToList()));
 			}
 		}
 	}
diff --git a/Engine/Source/Programs/Horde/Horde.Build/Api/Streams.cs b/Engine/Source/Programs/Horde/Horde.Build/Api/Streams.cs
--- a/Engine/Source/Programs/Horde/Horde.Build/Api/Streams.cs
+++ b/Engine/Source/Programs/Horde/Horde.Build/Api/Streams.cs
@@ -226,9 +226,18 @@ namespace Horde.Build.Api
 		public string? TriageChannel { get; set; }
 
 		/// <summary>
-		/// Issue handling workflow
+		/// Workflow to user for this stream
 		/// </summary>
-		public WorkflowId? Workflow { get; set; }
+		public WorkflowId? WorkflowId
+		{
+			get => Annotations.WorkflowId;
+			set => Annotations.WorkflowId = value;
+		}
+
+		/// <summary>
+		/// Default annotations to apply to nodes in this template
+		/// </summary>
+		public NodeAnnotations Annotations { get; set; } = new NodeAnnotations();
 
 		/// <summary>
 		/// Schedule to execute this template
diff --git a/Engine/Source/Programs/Horde/Horde.Build/Controllers/IssuesController.cs b/Engine/Source/Programs/Horde/Horde.Build/Controllers/IssuesController.cs
--- a/Engine/Source/Programs/Horde/Horde.Build/Controllers/IssuesController.cs
+++ b/Engine/Source/Programs/Horde/Horde.Build/Controllers/IssuesController.cs
@@ -138,12 +138,7 @@ namespace Horde.Build.Controllers
 							// Convert each issue to a response
 							foreach (IIssueSpan span in spansForIssue)
 							{
-								WorkflowId? workflowId = null;
-								if(stream.Config.TryGetTemplate(span.TemplateRefId, out TemplateRefConfig? templateRefConfig))
-								{
-									workflowId = templateRefConfig.Workflow;
-								}
-								spanResponses.Add(new FindIssueSpanResponse(span, workflowId));
+								spanResponses.Add(new FindIssueSpanResponse(span, span.LastFailure.Annotations.WorkflowId));
 							}
 						}
 					
diff --git a/Engine/Source/Programs/Horde/Horde.Build/Issues/IIssueCollection.cs b/Engine/Source/Programs/Horde/Horde.Build/Issues/IIssueCollection.cs
--- a/Engine/Source/Programs/Horde/Horde.Build/Issues/IIssueCollection.cs
+++ b/Engine/Source/Programs/Horde/Horde.Build/Issues/IIssueCollection.cs
@@ -5,6 +5,7 @@ using System.Collections.Generic;
 using System.Linq;
 using System.Threading.Tasks;
 using Horde.Build.Models;
+using Horde.Build.Server;
 using Horde.Build.Utilities;
 using MongoDB.Bson;
 using MongoDB.Driver;
@@ -16,6 +17,7 @@ namespace Horde.Build.Collections
 	using StreamId = StringId<IStream>;
 	using TemplateRefId = StringId<TemplateRef>;
 	using UserId = ObjectId<IUser>;
+	using WorkflowId = StringId<WorkflowConfig>;
 
 	/// <summary>
 	/// Fingerprint for an issue
@@ -233,6 +235,9 @@ namespace Horde.Build.Collections
 		/// <inheritdoc cref="IIssueStep.LogId"/>
 		public LogId? LogId { get; set; }
 
+		/// <inheritdoc cref="IIssueStep.Annotations"/>
+		public NodeAnnotations? Annotations { get; set; }
+
 		/// <inheritdoc cref="IIssueStep.PromoteByDefault"/>
 		public bool PromoteByDefault { get; set; }
 
@@ -247,8 +252,9 @@ namespace Horde.Build.Collections
 		/// <param name="stepId">The step id</param>
 		/// <param name="stepTime">Time that the step started</param>
 		/// <param name="logId">Unique id of the log file for this step</param>
+		/// <param name="annotations">Annotations for this step</param>
 		/// <param name="promoted">Whether this step is promoted</param>
-		public NewIssueStepData(int change, IssueSeverity severity, string jobName, JobId jobId, SubResourceId batchId, SubResourceId stepId, DateTime stepTime, LogId? logId, bool promoted)
+		public NewIssueStepData(int change, IssueSeverity severity, string jobName, JobId jobId, SubResourceId batchId, SubResourceId stepId, DateTime stepTime, LogId? logId, IReadOnlyNodeAnnotations? annotations, bool promoted)
 		{
 			Change = change;
 			Severity = severity;
@@ -258,6 +264,10 @@ namespace Horde.Build.Collections
 			StepId = stepId;
 			StepTime = stepTime;
 			LogId = logId;
+			if (annotations != null)
+			{
+				Annotations = new NodeAnnotations(annotations);
+			}
 			PromoteByDefault = promoted;
 		}
 
@@ -268,9 +278,10 @@ namespace Horde.Build.Collections
 		/// <param name="batch">Batch of the job for the step</param>
 		/// <param name="step">The step being built</param>
 		/// <param name="severity">Severity of the issue in this step</param>
+		/// <param name="annotations">Annotations for this step</param>
 		/// <param name="promoted">Whether this step is promoted</param>
-		public NewIssueStepData(IJob job, IJobStepBatch batch, IJobStep step, IssueSeverity severity, bool promoted)
-			: this(job.Change, severity, job.Name, job.Id, batch.Id, step.Id, step.StartTimeUtc ?? default, step.LogId, promoted)
+		public NewIssueStepData(IJob job, IJobStepBatch batch, IJobStep step, IssueSeverity severity, IReadOnlyNodeAnnotations? annotations, bool promoted)
+			: this(job.Change, severity, job.Name, job.Id, batch.Id, step.Id, step.StartTimeUtc ?? default, step.LogId, annotations, promoted)
 		{
 		}
 
@@ -279,7 +290,7 @@ namespace Horde.Build.Collections
 		/// </summary>
 		/// <param name="jobStepRef">The jobstep reference</param>
 		public NewIssueStepData(IJobStepRef jobStepRef)
-			: this(jobStepRef.Change, IssueSeverity.Unspecified, jobStepRef.JobName, jobStepRef.Id.JobId, jobStepRef.Id.BatchId, jobStepRef.Id.StepId, jobStepRef.StartTimeUtc, jobStepRef.LogId, false)
+			: this(jobStepRef.Change, IssueSeverity.Unspecified, jobStepRef.JobName, jobStepRef.Id.JobId, jobStepRef.Id.BatchId, jobStepRef.Id.StepId, jobStepRef.StartTimeUtc, jobStepRef.LogId, null, false)
 		{
 		}
 	}
diff --git a/Engine/Source/Programs/Horde/Horde.Build/Issues/Impl/IssueReportService.cs b/Engine/Source/Programs/Horde/Horde.Build/Issues/Impl/IssueReportService.cs
--- a/Engine/Source/Programs/Horde/Horde.Build/Issues/Impl/IssueReportService.cs
+++ b/Engine/Source/Programs/Horde/Horde.Build/Issues/Impl/IssueReportService.cs
@@ -5,6 +5,7 @@ using System.Collections.Generic;
 using System.Linq;
 using System.Threading;
 using System.Threading.Tasks;
+using EpicGames.Core;
 using Horde.Build.Collections;
 using Horde.Build.Config;
 using Horde.Build.Models;
@@ -18,23 +19,30 @@ using Microsoft.Extensions.Logging;
 namespace Horde.Build.Issues.Impl
 {
 	using TemplateRefId = StringId<TemplateRef>;
+	using WorkflowId = StringId<WorkflowConfig>;
 
 #pragma warning disable CS1591 // Missing XML comment for publicly visible type or member
+	public class WorkflowStats
+	{
+		public int NumSteps { get; set; }
+		public int NumPassingSteps { get; set; }
+	}
+
 	public class IssueReport
 	{
 		public DateTimeOffset Time { get; }
 		public IStream Stream { get; }
 		public WorkflowConfig Workflow { get; }
-		public int NumSteps { get; set; }
-		public int NumPassingSteps { get; set; }
+		public WorkflowStats WorkflowStats { get; }
 		public List<IIssue> Issues { get; } = new List<IIssue>();
 		public List<IIssueSpan> IssueSpans { get; } = new List<IIssueSpan>();
 
-		public IssueReport(DateTimeOffset time, IStream stream, WorkflowConfig workflow)
+		public IssueReport(DateTimeOffset time, IStream stream, WorkflowConfig workflow, WorkflowStats workflowStats)
 		{
 			Time = time;
 			Stream = stream;
 			Workflow = workflow;
+			WorkflowStats = workflowStats;
 		}
 	}
 #pragma warning restore CS1591 // Missing XML comment for publicly visible type or member
@@ -53,6 +61,7 @@ namespace Horde.Build.Issues.Impl
 		readonly SingletonDocument<IssueReportState> _state;
 		readonly IStreamCollection _streamCollection;
 		readonly IIssueCollection _issueCollection;
+		readonly IGraphCollection _graphCollection;
 		readonly IJobCollection _jobCollection;
 		readonly ConfigCollection _configCollection;
 		readonly INotificationService _notificationService;
@@ -63,11 +72,12 @@ namespace Horde.Build.Issues.Impl
 		/// <summary>
 		/// Constructor
 		/// </summary>
-		public IssueReportService(MongoService mongoService, IStreamCollection streamCollection, IIssueCollection issueCollection, IJobCollection jobCollection, ConfigCollection configCollection, INotificationService notificationService, IClock clock, ILogger<IssueReportService> logger)
+		public IssueReportService(MongoService mongoService, IStreamCollection streamCollection, IIssueCollection issueCollection, IGraphCollection graphCollection, IJobCollection jobCollection, ConfigCollection configCollection, INotificationService notificationService, IClock clock, ILogger<IssueReportService> logger)
 		{
 			_state = new SingletonDocument<IssueReportState>(mongoService);
 			_streamCollection = streamCollection;
 			_issueCollection = issueCollection;
+			_graphCollection = graphCollection;
 			_jobCollection = jobCollection;
 			_configCollection = configCollection;
 			_notificationService = notificationService;
@@ -103,7 +113,6 @@ namespace Horde.Build.Issues.Impl
 				{
 					List<IIssue>? issues = null;
 					List<IIssueSpan>? spans = null;
-					List<IJob>? jobs = null;
 
 					foreach (WorkflowConfig workflow in config.Workflows)
 					{
@@ -129,37 +138,17 @@ namespace Horde.Build.Issues.Impl
 
 						issues ??= await _issueCollection.FindIssuesAsync(streamId: stream.Id);
 						spans ??= await _issueCollection.FindSpansAsync(issueIds: issues.Select(x => x.Id).ToArray());
-						jobs ??= await _jobCollection.FindAsync(streamId: stream.Id, minCreateTime: prevScheduledReportTime);
 
-						HashSet<TemplateRefId> wfTemplateRefIds = new HashSet<TemplateRefId>();
-						foreach (Api.TemplateRefConfig templateRef in config.Templates)
+						Dictionary<WorkflowId, WorkflowStats> workflowIdToStats = await GetWorkflowStatsAsync(stream, prevScheduledReportTime);
+						if (!workflowIdToStats.TryGetValue(workflow.Id, out WorkflowStats? workflowStats))
 						{
-							if (templateRef.Workflow != null && templateRef.Workflow == workflow.Id)
-							{
-								wfTemplateRefIds.Add(templateRef.Id);
-							}
-						}
-
-						IssueReport report = new IssueReport(currentTime, stream, workflow);
-
-						foreach (IJob job in jobs)
-						{
-							if (wfTemplateRefIds.Contains(job.TemplateId))
-							{
-								foreach (IJobStep step in job.Batches.SelectMany(x => x.Steps).Where(x => x.State == JobStepState.Completed))
-								{
-									report.NumSteps++;
-									if (step.Outcome == JobStepOutcome.Success)
-									{
-										report.NumPassingSteps++;
-									}
-								}
-							}
+							workflowStats = new WorkflowStats();
 						}
 
+						IssueReport report = new IssueReport(currentTime, stream, workflow, workflowStats);
 						foreach (IIssueSpan span in spans)
 						{
-							if (span.LastSuccess != null && wfTemplateRefIds.Contains(span.TemplateRefId))
+							if (span.LastSuccess != null && span.LastFailure.Annotations.WorkflowId == workflow.Id)
 							{
 								report.IssueSpans.Add(span);
 							}
@@ -188,6 +177,57 @@ namespace Horde.Build.Issues.Impl
 			}
 		}
 
+		private async Task<Dictionary<WorkflowId, WorkflowStats>> GetWorkflowStatsAsync(IStream stream, DateTime minTime)
+		{
+			List<IJob> jobs = await _jobCollection.FindAsync(streamId: stream.Id, minCreateTime: minTime);
+
+			Dictionary<WorkflowId, WorkflowStats> workflowIdToStats = new Dictionary<WorkflowId, WorkflowStats>();
+			foreach (IGrouping<TemplateRefId, IJob> templateGroup in jobs.GroupBy(x => x.TemplateId))
+			{
+				WorkflowId? templateWorkflowId = null;
+				if (stream.Config.TryGetTemplate(templateGroup.Key, out Api.TemplateRefConfig? templateRefConfig))
+				{
+					templateWorkflowId = templateRefConfig.Annotations.WorkflowId;
+				}
+
+				foreach (IGrouping<ContentHash, IJob> graphGroup in templateGroup.GroupBy(x => x.GraphHash))
+				{
+					IGraph graph = await _graphCollection.GetAsync(graphGroup.Key);
+					foreach (IJob job in graphGroup)
+					{
+						foreach (IJobStepBatch batch in job.Batches)
+						{
+							foreach (IJobStep step in batch.Steps)
+							{
+								if (step.State == JobStepState.Completed)
+								{
+									INode node = graph.Groups[batch.GroupIdx].Nodes[step.NodeIdx];
+									WorkflowId? workflowId = node.Annotations.WorkflowId ?? templateWorkflowId;
+									if (workflowId != null)
+									{
+										WorkflowStats? stats;
+										if (!workflowIdToStats.TryGetValue(workflowId.Value, out stats))
+										{
+											stats = new WorkflowStats();
+											workflowIdToStats.Add(workflowId.Value, stats);
+										}
+
+										stats.NumSteps++;
+										if (step.Outcome == JobStepOutcome.Success)
+										{
+											stats.NumPassingSteps++;
+										}
+									}
+								}
+							}
+						}
+					}
+				}
+			}
+
+			return workflowIdToStats;
+		}
+
 		static DateTime GetLastScheduledReportTime(WorkflowConfig workflow, DateTime currentTime)
 		{
 			if (workflow.ReportTimes.Count > 0)
diff --git a/Engine/Source/Programs/Horde/Horde.Build/Issues/Impl/IssueService.cs b/Engine/Source/Programs/Horde/Horde.Build/Issues/Impl/IssueService.cs
--- a/Engine/Source/Programs/Horde/Horde.Build/Issues/Impl/IssueService.cs
+++ b/Engine/Source/Programs/Horde/Horde.Build/Issues/Impl/IssueService.cs
@@ -32,6 +32,7 @@ namespace Horde.Build.Services.Impl
 	using StreamId = StringId<IStream>;
 	using TemplateRefId = StringId<TemplateRef>;
 	using UserId = ObjectId<IUser>;
+	using WorkflowId = StringId<WorkflowConfig>;
 
 	/// <summary>
 	/// Detailed issue information
@@ -542,23 +543,6 @@ namespace Horde.Build.Services.Impl
 				throw new Exception($"Invalid stream id '{job.StreamId}' on job '{job.Id}'");
 			}
 
-			// Check whether issue creation is enabled for this template type
-			StreamConfig config = await _configCollection.GetConfigAsync<StreamConfig>(stream.ConfigRevision);
-
-			TemplateRefConfig? templateRef;
-			if (!config.TryGetTemplate(job.TemplateId, out templateRef))
-			{
-				return;
-			}
-
-			if (templateRef.Workflow != null)
-			{
-				if(!config.TryGetWorkflow(templateRef.Workflow.Value, out WorkflowConfig? workflow) || !workflow.CreateIssues)
-				{
-					return;
-				}
-			}
-
 			// Find the batch for this event
 			IJobStepBatch? batch;
 			if (!job.TryGetBatch(batchId, out batch))
@@ -591,8 +575,39 @@ namespace Horde.Build.Services.Impl
 
 			scope.Span.SetTag("LogId", step.LogId.ToString());
 
+			// Get all the annotations for this template
+			TemplateRefConfig? templateRef;
+			if (!stream.Config.TryGetTemplate(job.TemplateId, out templateRef))
+			{
+				return;
+			}
+
+			NodeAnnotations annotations = new NodeAnnotations();
+			annotations.Merge(templateRef.Annotations);
+
 			// Get the node associated with this step
 			INode node = graph.Groups[batch.GroupIdx].Nodes[step.NodeIdx];
+			annotations.Merge(graph.Groups[batch.GroupIdx].Nodes[step.NodeIdx].Annotations);
+
+			// Get the workflow for this step
+			WorkflowConfig? workflow = null;
+			WorkflowId? workflowId = annotations.WorkflowId;
+			if (workflowId != null)
+			{
+				if (!stream.Config.TryGetWorkflow(workflowId.Value, out workflow))
+				{
+					_logger.LogWarning("Unable to find workflow {WorkflowId} for {JobId}:{BatchId}:{StepId}", workflowId, job.Id, batchId, stepId);
+					return;
+				}
+				annotations.Merge(workflow.Annotations);
+			}
+
+			// If the workflow disables issue creation, bail out now
+			if (workflow != null && !workflow.CreateIssues)
+			{
+				_logger.LogInformation("Issue creation for step {JobId}:{BatchId}:{StepId} disabled by workflow {WorkflowId}", job.Id, batchId, stepId, workflowId);
+				return;
+			}
 
 			// Gets the events for this step grouped by fingerprint
 			HashSet<NewEventGroup> eventGroups = await GetEventGroupsForStepAsync(job, batch, step, node);
@@ -610,16 +625,16 @@ namespace Horde.Build.Services.Impl
 					// Add the events to existing issues, and create new issues for everything else
 					if (eventGroups.Count > 0)
 					{
-						if (!await AddEventsToExistingSpans(job, batch, step, eventGroups, openSpans, checkedSpanIds, job.PromoteIssuesByDefault))
+						if (!await AddEventsToExistingSpans(job, batch, step, eventGroups, openSpans, checkedSpanIds, annotations, job.PromoteIssuesByDefault))
 						{
 							continue;
 						}
 						if (!SuppressNewIssues)
 						{
-							if (!await AddEventsToNewSpans(stream, job, batch, step, node, openSpans, eventGroups, job.PromoteIssuesByDefault))
-							{
-								continue;
-							}
+						    if (!await AddEventsToNewSpans(stream, job, batch, step, node, openSpans, eventGroups, annotations, job.PromoteIssuesByDefault))
+						    {
+							    continue;
+						    }
 						}
 					}
 
@@ -738,9 +753,10 @@ namespace Horde.Build.Services.Impl
 		/// <param name="newEventGroups">Set of events to try to add</param>
 		/// <param name="openSpans">List of open spans</param>
 		/// <param name="checkedSpanIds">Set of span ids that have been checked</param>
+		/// <param name="annotations">Annotations for this step</param>
 		/// <param name="promoteByDefault"></param>
 		/// <returns>True if the adding completed</returns>
-		async Task<bool> AddEventsToExistingSpans(IJob job, IJobStepBatch batch, IJobStep step, HashSet<NewEventGroup> newEventGroups, List<IIssueSpan> openSpans, HashSet<ObjectId> checkedSpanIds, bool promoteByDefault)
+		async Task<bool> AddEventsToExistingSpans(IJob job, IJobStepBatch batch, IJobStep step, HashSet<NewEventGroup> newEventGroups, List<IIssueSpan> openSpans, HashSet<ObjectId> checkedSpanIds, IReadOnlyNodeAnnotations? annotations, bool promoteByDefault)
 		{
 			for(int spanIdx = 0; spanIdx < openSpans.Count; spanIdx++)
 			{
@@ -752,7 +768,7 @@ namespace Horde.Build.Services.Impl
 					if (matchEventGroups.Count > 0)
 					{
 						// Add the new step data
-						NewIssueStepData newFailure = new NewIssueStepData(job, batch, step, GetIssueSeverity(matchEventGroups.SelectMany(x => x.Events)), promoteByDefault);
+						NewIssueStepData newFailure = new NewIssueStepData(job, batch, step, GetIssueSeverity(matchEventGroups.SelectMany(x => x.Events)), annotations, promoteByDefault);
 						await _issueCollection.AddStepAsync(openSpan.Id, newFailure);
 
 						// Update the span if this changes the current range
@@ -820,9 +836,10 @@ namespace Horde.Build.Services.Impl
 		/// <param name="node">Node run in the step</param>
 		/// <param name="openSpans">List of open spans. New issues will be added to this list.</param>
 		/// <param name="newEventGroups">Set of remaining events</param>
+		/// <param name="annotations"></param>
 		/// <param name="promoteByDefault"></param>
 		/// <returns>True if all events were added</returns>
-		async Task<bool> AddEventsToNewSpans(IStream stream, IJob job, IJobStepBatch batch, IJobStep step, INode node, List<IIssueSpan> openSpans, HashSet<NewEventGroup> newEventGroups, bool promoteByDefault)
+		async Task<bool> AddEventsToNewSpans(IStream stream, IJob job, IJobStepBatch batch, IJobStep step, INode node, List<IIssueSpan> openSpans, HashSet<NewEventGroup> newEventGroups, IReadOnlyNodeAnnotations? annotations, bool promoteByDefault)
 		{
 			while (newEventGroups.Count > 0)
 			{
@@ -844,7 +861,7 @@ namespace Horde.Build.Services.Impl
 				}
 
 				// Get the step data
-				NewIssueStepData stepData = new NewIssueStepData(job, batch, step, GetIssueSeverity(eventGroup.Events), promoteByDefault);
+				NewIssueStepData stepData = new NewIssueStepData(job, batch, step, GetIssueSeverity(eventGroup.Events), annotations, promoteByDefault);
 
 				// Get the span data
 				NewIssueSpanData spanData = new NewIssueSpanData(stream.Id, stream.Name, job.TemplateId, node.Name, eventGroup.Fingerprint, stepData);
@@ -1224,7 +1241,7 @@ namespace Horde.Build.Services.Impl
 			{
 				if (job.Change < span.FirstFailure.Change && (span.LastSuccess == null || job.Change > span.LastSuccess.Change))
 				{
-					NewIssueStepData newLastSuccess = new NewIssueStepData(job.Change, IssueSeverity.Unspecified, job.Name, job.Id, batch.Id, step.Id, step.StartTimeUtc ?? default, step.LogId, false);
+					NewIssueStepData newLastSuccess = new NewIssueStepData(job.Change, IssueSeverity.Unspecified, job.Name, job.Id, batch.Id, step.Id, step.StartTimeUtc ?? default, step.LogId, null, false);
 					List<NewIssueSpanSuspectData> newSuspects = await FindSuspectsForSpanAsync(stream, span.Fingerprint, job.Change + 1, span.FirstFailure.Change);
 
 					if (await _issueCollection.TryUpdateSpanAsync(span, newLastSuccess: newLastSuccess, newSuspects: newSuspects) == null)
@@ -1237,7 +1254,7 @@ namespace Horde.Build.Services.Impl
 				}
 				else if (job.Change > span.LastFailure.Change && (span.NextSuccess == null || job.Change < span.NextSuccess.Change))
 				{
-					NewIssueStepData newNextSucccess = new NewIssueStepData(job.Change, IssueSeverity.Unspecified, job.Name, job.Id, batch.Id, step.Id, step.StartTimeUtc ?? default, step.LogId, false);
+					NewIssueStepData newNextSucccess = new NewIssueStepData(job.Change, IssueSeverity.Unspecified, job.Name, job.Id, batch.Id, step.Id, step.StartTimeUtc ?? default, step.LogId, null, false);
 
 					if (await _issueCollection.TryUpdateSpanAsync(span, newNextSuccess: newNextSucccess) == null)
 					{
diff --git a/Engine/Source/Programs/Horde/Horde.Build/Jobs/IGraph.cs b/Engine/Source/Programs/Horde/Horde.Build/Jobs/IGraph.cs
--- a/Engine/Source/Programs/Horde/Horde.Build/Jobs/IGraph.cs
+++ b/Engine/Source/Programs/Horde/Horde.Build/Jobs/IGraph.cs
@@ -6,7 +6,6 @@ using System.Diagnostics;
 using System.Diagnostics.CodeAnalysis;
 using System.Linq;
 using EpicGames.Core;
-using Horde.Build.Utilities;
 using HordeCommon;
 using MongoDB.Bson.Serialization.Attributes;
 
@@ -65,7 +64,7 @@ namespace Horde.Build.Models
 		/// <summary>
 		/// Annotations for this node
 		/// </summary>
-		public IReadOnlyDictionary<string, string>? Annotations { get; }
+		public IReadOnlyNodeAnnotations Annotations { get; }
 	}
 
 	/// <summary>
@@ -422,7 +421,7 @@ namespace Horde.Build.Models
 		/// <summary>
 		/// Additional user annotations for this node
 		/// </summary>
-		public CaseInsensitiveDictionary<string>? Annotations { get; set; }
+		public NodeAnnotations? Annotations { get; set; }
 
 		/// <summary>
 		/// Constructor
@@ -437,7 +436,7 @@ namespace Horde.Build.Models
 		/// <param name="credentials">Credentials required for this node to run</param>
 		/// <param name="properties">Properties for the node</param>
 		/// <param name="annotations">User annotations for this node</param>
-		public NewNode(string name, List<string>? inputDependencies = null, List<string>? orderDependencies = null, Priority? priority = null, bool? allowRetry = null, bool? runEarly = null, bool? warnings = null, Dictionary<string, string>? credentials = null, Dictionary<string, string>? properties = null, CaseInsensitiveDictionary<string>? annotations = null)
+		public NewNode(string name, List<string>? inputDependencies = null, List<string>? orderDependencies = null, Priority? priority = null, bool? allowRetry = null, bool? runEarly = null, bool? warnings = null, Dictionary<string, string>? credentials = null, Dictionary<string, string>? properties = null, IReadOnlyNodeAnnotations? annotations = null)
 		{
 			Name = name;
 			InputDependencies = inputDependencies;
@@ -448,7 +447,10 @@ namespace Horde.Build.Models
 			Warnings = warnings;
 			Credentials = credentials;
 			Properties = properties;
-			Annotations = annotations;
+			if (annotations != null)
+			{
+				Annotations = new NodeAnnotations(annotations);
+			}
 		}
 	}
 
diff --git a/Engine/Source/Programs/Horde/Horde.Build/Jobs/Impl/GraphCollection.cs b/Engine/Source/Programs/Horde/Horde.Build/Jobs/Impl/GraphCollection.cs
--- a/Engine/Source/Programs/Horde/Horde.Build/Jobs/Impl/GraphCollection.cs
+++ b/Engine/Source/Programs/Horde/Horde.Build/Jobs/Impl/GraphCollection.cs
@@ -52,11 +52,11 @@ namespace Horde.Build.Collections.Impl
 
 			[BsonIgnoreIfNull]
 			[BsonDictionaryOptions(DictionaryRepresentation.ArrayOfDocuments)]
-			public CaseInsensitiveDictionary<string>? Annotations { get; set; }
+			public NodeAnnotations? Annotations { get; set; }
 
 			IReadOnlyDictionary<string, string>? INode.Credentials => Credentials;
 			IReadOnlyDictionary<string, string>? INode.Properties => Properties;
-			IReadOnlyDictionary<string, string>? INode.Annotations => Annotations;
+			IReadOnlyNodeAnnotations INode.Annotations => Annotations ?? NodeAnnotations.Empty;
 
 			[BsonConstructor]
 			private Node()
@@ -66,7 +66,7 @@ namespace Horde.Build.Collections.Impl
 				OrderDependencies = null!;
 			}
 
-			public Node(string name, NodeRef[] inputDependencies, NodeRef[] orderDependencies, Priority priority, bool allowRetry, bool runEarly, bool warnings, Dictionary<string, string>? credentials, Dictionary<string, string>? properties, CaseInsensitiveDictionary<string>? annotations)
+			public Node(string name, NodeRef[] inputDependencies, NodeRef[] orderDependencies, Priority priority, bool allowRetry, bool runEarly, bool warnings, Dictionary<string, string>? credentials, Dictionary<string, string>? properties, IReadOnlyNodeAnnotations? annotations)
 			{
 				Name = name;
 				InputDependencies = inputDependencies;
@@ -77,7 +77,10 @@ namespace Horde.Build.Collections.Impl
 				Warnings = warnings;
 				Credentials = credentials;
 				Properties = properties;
-				Annotations = annotations;
+				if (annotations != null && annotations.Count > 0)
+				{
+					Annotations = new NodeAnnotations(annotations);
+				}
 			}
 		}
 
diff --git a/Engine/Source/Programs/Horde/Horde.Build/Jobs/NodeAnnotations.cs b/Engine/Source/Programs/Horde/Horde.Build/Jobs/NodeAnnotations.cs
new file mode 100644
--- /dev/null
+++ b/Engine/Source/Programs/Horde/Horde.Build/Jobs/NodeAnnotations.cs
@@ -0,0 +1,94 @@
+// Copyright Epic Games, Inc. All Rights Reserved.
+
+using System.Collections.Generic;
+using Horde.Build.Server;
+using Horde.Build.Utilities;
+
+namespace Horde.Build.Models
+{
+	using WorkflowId = StringId<WorkflowConfig>;
+
+	/// <summary>
+	/// Set of annotations for a node
+	/// </summary>
+	public class NodeAnnotations : CaseInsensitiveDictionary<string>, IReadOnlyNodeAnnotations
+	{
+		/// <summary>
+		/// Empty annotation dictionary
+		/// </summary>
+		public static IReadOnlyNodeAnnotations Empty { get; } = new NodeAnnotations();
+
+		/// <summary>
+		/// The issue workflow to use, as defined in the stream configuration file
+		/// </summary>
+		public const string WorkflowKeyName = "Workflow";
+
+		/// <summary>
+		/// Constructor
+		/// </summary>
+		public NodeAnnotations()
+		{
+		}
+
+		/// <summary>
+		/// Constructor
+		/// </summary>
+		/// <param name="annotations"></param>
+		public NodeAnnotations(IReadOnlyDictionary<string, string> annotations)
+		{
+			Merge(annotations);
+		}
+
+		/// <inheritdoc/>
+		public WorkflowId? WorkflowId
+		{
+			get
+			{
+				string? workflowName;
+				if (!TryGetValue(WorkflowKeyName, out workflowName))
+				{
+					return null;
+				}
+				return new WorkflowId(workflowName);
+			}
+			set => SetOrUnset(WorkflowKeyName, value?.ToString());
+		}
+
+		private void SetOrUnset(string key, string? value)
+		{
+			if (value == null)
+			{
+				Remove(key);
+			}
+			else
+			{
+				this[key] = value;
+			}
+		}
+
+		/// <summary>
+		/// Merge in entries from another set of annotation
+		/// </summary>
+		/// <param name="other"></param>
+		public void Merge(IReadOnlyDictionary<string, string> other)
+		{
+			foreach ((string key, string value) in other)
+			{
+				this[key] = value;
+			}
+		}
+	}
+
+	/// <summary>
+	/// Interface which wraps a generic key/value dictionary to provide specific node annotations
+	/// </summary>
+#pragma warning disable CA1710 // Identifiers should have correct suffix
+	public interface IReadOnlyNodeAnnotations : IReadOnlyDictionary<string, string>
+#pragma warning restore CA1710 // Identifiers should have correct suffix
+	{
+		/// <summary>
+		/// Workflow to use for triaging issues from this node
+		/// </summary>
+		WorkflowId? WorkflowId { get; }
+	}
+}
diff --git a/Engine/Source/Programs/Horde/Horde.Build/Notifications/Impl/SlackNotificationSink.cs b/Engine/Source/Programs/Horde/Horde.Build/Notifications/Impl/SlackNotificationSink.cs
--- a/Engine/Source/Programs/Horde/Horde.Build/Notifications/Impl/SlackNotificationSink.cs
+++ b/Engine/Source/Programs/Horde/Horde.Build/Notifications/Impl/SlackNotificationSink.cs
@@ -39,6 +39,7 @@ namespace Horde.Build.Notifications.Impl
 	using JobId = ObjectId<IJob>;
 	using LogId = ObjectId<ILogFile>;
 	using UserId = ObjectId<IUser>;
+	using WorkflowId = StringId<WorkflowConfig>;
 
 	/// <summary>
 	/// Maintains a connection to Slack, in order to receive socket-mode notifications of user interactions
@@ -624,16 +625,14 @@ namespace Horde.Build.Notifications.Impl
 			if (details.Spans.Count > 0)
 			{
 				IIssueSpan span = details.Spans[0];
-				IStream? stream = await _streamService.GetStreamAsync(span.StreamId);
-				if(stream != null)
+
+				WorkflowId? workflowId = span.LastFailure.Annotations.WorkflowId;
+				if (workflowId != null)
 				{
-					StreamConfig config = await _configCollection.GetConfigAsync<StreamConfig>(stream.ConfigRevision);
-					if (config.TryGetTemplate(span.TemplateRefId, out TemplateRefConfig? templateRef))
+					IStream? stream = await _streamService.GetStreamAsync(span.StreamId);
+					if(stream != null && stream.Config.TryGetWorkflow(workflowId.Value, out WorkflowConfig? workflow))
 					{
-						if (templateRef.Workflow != null && config.TryGetWorkflow(templateRef.Workflow.Value, out WorkflowConfig? workflow))
-						{
-							await SendTriageMessageAsync(issue, span, workflow);
-						}
+						await SendTriageMessageAsync(issue, span, workflow);
 					}
 				}
 			}
@@ -1135,10 +1134,10 @@ namespace Horde.Build.Notifications.Impl
 					body.Append($"*: {issue.Summary} [{FormatReadableTimeSpan(report.Time - issue.CreatedAt)}] - {status}");
 				}
 
-				if (report.NumSteps > 0)
+				if (report.WorkflowStats.NumSteps > 0)
 				{
-					double totalPct = (report.NumPassingSteps * 100.0) / report.NumSteps;
-					body.Append($"\n\n{report.NumPassingSteps:n0} of {report.NumSteps:n0} (*{totalPct:0.0}%*) build steps succeeded since last status update.");
+					double totalPct = (report.WorkflowStats.NumPassingSteps * 100.0) / report.WorkflowStats.NumSteps;
+					body.Append($"\n\n{report.WorkflowStats.NumPassingSteps:n0} of {report.WorkflowStats.NumSteps:n0} (*{totalPct:0.0}%*) build steps succeeded since last status update.");
 				}
 
 				await SendMessageAsync(report.Workflow.ReportChannel, text: body.ToString(), withEnvironment: false);
diff --git a/Engine/Source/Programs/Horde/Horde.Build/Server/Config.cs b/Engine/Source/Programs/Horde/Horde.Build/Server/Config.cs
--- a/Engine/Source/Programs/Horde/Horde.Build/Server/Config.cs
+++ b/Engine/Source/Programs/Horde/Horde.Build/Server/Config.cs
@@ -346,6 +346,11 @@ namespace Horde.Build.Server
 		/// Maximum number of people to mention on a triage thread
 		/// </summary>
 		public int MaxMentions { get; set; } = 5;
+
+		/// <summary>
+		/// Additional node annotations implicit in this workflow
+		/// </summary>
+		public NodeAnnotations Annotations { get; set; } = new NodeAnnotations();
 	}
 
 	/// <summary>
diff --git a/Engine/Source/Programs/Horde/Horde.Build/Services/RpcService.cs b/Engine/Source/Programs/Horde/Horde.Build/Services/RpcService.cs
--- a/Engine/Source/Programs/Horde/Horde.Build/Services/RpcService.cs
+++ b/Engine/Source/Programs/Horde/Horde.Build/Services/RpcService.cs
@@ -730,7 +730,7 @@ namespace Horde.Build.Services
 				List<NewNode> newNodes = new List<NewNode>();
 				foreach (CreateNodeRequest node in group.Nodes)
 				{
-					NewNode newNode = new NewNode(node.Name, node.InputDependencies.ToList(), node.OrderDependencies.ToList(), node.Priority, node.AllowRetry, node.RunEarly, node.Warnings, new Dictionary<string, string>(node.Credentials), new Dictionary<string, string>(node.Properties), new CaseInsensitiveDictionary<string>(node.Annotations));
+					NewNode newNode = new NewNode(node.Name, node.InputDependencies.ToList(), node.OrderDependencies.ToList(), node.Priority, node.AllowRetry, node.RunEarly, node.Warnings, new Dictionary<string, string>(node.Credentials), new Dictionary<string, string>(node.Properties), new NodeAnnotations(node.Annotations));
 					newNodes.Add(newNode);
 				}
 				newGroups.Add(new NewGroup(group.AgentType, newNodes));
