Horde: Missing files.

diff --git a/Engine/Source/Programs/Horde/Horde.Build.Tests/IssueServiceTests.cs b/Engine/Source/Programs/Horde/Horde.Build.Tests/IssueServiceTests.cs
--- a/Engine/Source/Programs/Horde/Horde.Build.Tests/IssueServiceTests.cs
+++ b/Engine/Source/Programs/Horde/Horde.Build.Tests/IssueServiceTests.cs
@@ -169,8 +169,9 @@ namespace Horde.Build.Tests
 
 		public static INode MockNode(string name)
 		{
-			Mock<INode> node = new Mock<INode>();
+			Mock<INode> node = new Mock<INode>(MockBehavior.Strict);
 			node.SetupGet(x => x.Name).Returns(name);
+			node.SetupGet(x => x.Annotations).Returns(NodeAnnotations.Empty);
 			return node.Object;
 		}
 
diff --git a/Engine/Source/Programs/Horde/Horde.Build/Issues/IIssueStep.cs b/Engine/Source/Programs/Horde/Horde.Build/Issues/IIssueStep.cs
--- a/Engine/Source/Programs/Horde/Horde.Build/Issues/IIssueStep.cs
+++ b/Engine/Source/Programs/Horde/Horde.Build/Issues/IIssueStep.cs
@@ -1,6 +1,7 @@
 // Copyright Epic Games, Inc. All Rights Reserved.
 
 using System;
+using Horde.Build.Server;
 using Horde.Build.Utilities;
 using MongoDB.Bson;
 
@@ -8,6 +9,7 @@ namespace Horde.Build.Models
 {
 	using JobId = ObjectId<IJob>;
 	using LogId = ObjectId<ILogFile>;
+	using WorkflowId = StringId<WorkflowConfig>;
 
 	/// <summary>
 	/// Identifies a particular changelist and job that contributes to a span
@@ -54,6 +56,11 @@ namespace Horde.Build.Models
 		/// </summary>
 		public DateTime StepTime { get; }
 
+		/// <summary>
+		/// Annotations for this step
+		/// </summary>
+		public IReadOnlyNodeAnnotations Annotations { get; }
+
 		/// <summary>
 		/// Whether to promote spans including this step by default
 		/// </summary>
diff --git a/Engine/Source/Programs/Horde/Horde.Build/Issues/Impl/IssueCollection.cs b/Engine/Source/Programs/Horde/Horde.Build/Issues/Impl/IssueCollection.cs
--- a/Engine/Source/Programs/Horde/Horde.Build/Issues/Impl/IssueCollection.cs
+++ b/Engine/Source/Programs/Horde/Horde.Build/Issues/Impl/IssueCollection.cs
@@ -22,6 +22,7 @@ namespace Horde.Build.Collections.Impl
 	using StreamId = StringId<IStream>;
 	using TemplateRefId = StringId<TemplateRef>;
 	using UserId = ObjectId<IUser>;
+	using WorkflowId = StringId<WorkflowConfig>;
 
 	class IssueCollection : IIssueCollection
 	{
@@ -357,11 +358,15 @@ namespace Horde.Build.Collections.Impl
 
 			public LogId? LogId { get; set; }
 
+			[BsonIgnoreIfNull]
+			public NodeAnnotations? Annotations { get; set; }
+
 			public bool? PromoteByDefault { get; set; }
 
 			[BsonElement("NotifySuspects"), BsonIgnoreIfDefault(false)]
 			public bool NotifySuspectsDeprecated { get; set; }
 
+			IReadOnlyNodeAnnotations IIssueStep.Annotations => Annotations ?? NodeAnnotations.Empty;
 			bool IIssueStep.PromoteByDefault => PromoteByDefault ?? NotifySuspectsDeprecated;
 
 			[BsonConstructor]
@@ -382,6 +387,7 @@ namespace Horde.Build.Collections.Impl
 				StepId = stepData.StepId;
 				StepTime = stepData.StepTime;
 				LogId = stepData.LogId;
+				Annotations = stepData.Annotations;
 				PromoteByDefault = stepData.PromoteByDefault;
 			}
 		}
