Horde: Adding job step error for paused steps
#rnx
#skipci

diff --git a/Engine/Source/Programs/Horde/Horde.Build.Tests/JobTaskSourceTests.cs b/Engine/Source/Programs/Horde/Horde.Build.Tests/JobTaskSourceTests.cs
--- a/Engine/Source/Programs/Horde/Horde.Build.Tests/JobTaskSourceTests.cs
+++ b/Engine/Source/Programs/Horde/Horde.Build.Tests/JobTaskSourceTests.cs
@@ -154,6 +154,7 @@ namespace Horde.Build.Tests
 			Assert.AreEqual(batch.Steps.Count, 4);
 			Assert.AreEqual(batch.Steps[0].State, JobStepState.Ready);
 			Assert.AreEqual(batch.Steps[1].State, JobStepState.Skipped);
+			Assert.AreEqual(batch.Steps[1].Error, JobStepError.Paused);
 			Assert.AreEqual(batch.Steps[2].State, JobStepState.Skipped);
 			Assert.AreEqual(batch.Steps[3].State, JobStepState.Waiting);
 
diff --git a/Engine/Source/Programs/Horde/Horde.Build/Jobs/IJob.cs b/Engine/Source/Programs/Horde/Horde.Build/Jobs/IJob.cs
--- a/Engine/Source/Programs/Horde/Horde.Build/Jobs/IJob.cs
+++ b/Engine/Source/Programs/Horde/Horde.Build/Jobs/IJob.cs
@@ -160,6 +160,11 @@ namespace Horde.Build.Models
 		/// Step did not complete in the required amount of time
 		/// </summary>
 		TimedOut,
+
+		/// <summary>
+		/// Step is in is paused state so was skipped
+		/// </summary>
+		Paused
 	}
 
 	/// <summary>
diff --git a/Engine/Source/Programs/Horde/Horde.Build/Jobs/Impl/JobTaskSource.cs b/Engine/Source/Programs/Horde/Horde.Build/Jobs/Impl/JobTaskSource.cs
--- a/Engine/Source/Programs/Horde/Horde.Build/Jobs/Impl/JobTaskSource.cs
+++ b/Engine/Source/Programs/Horde/Horde.Build/Jobs/Impl/JobTaskSource.cs
@@ -451,7 +451,7 @@ namespace Horde.Build.Tasks.Impl
 										if (step != null)
 										{
 											JobId jobId = newJob.Id;
-											newJob = await _jobs.TryUpdateStepAsync(newJob, graph, batch.Id, step.Id, JobStepState.Skipped);
+											newJob = await _jobs.TryUpdateStepAsync(newJob, graph, batch.Id, step.Id, JobStepState.Skipped, newError: JobStepError.Paused);
 											if (newJob == null)
 											{
 												_logger.LogError("Job {JobId} failed to update step {StepName} pause state", jobId, state.Name);
diff --git a/Engine/Source/Programs/Horde/HordeDashboard/src/backend/Api.ts b/Engine/Source/Programs/Horde/HordeDashboard/src/backend/Api.ts
--- a/Engine/Source/Programs/Horde/HordeDashboard/src/backend/Api.ts
+++ b/Engine/Source/Programs/Horde/HordeDashboard/src/backend/Api.ts
@@ -72,6 +72,17 @@ export enum JobStepOutcome {
 	Unspecified = "Unspecified"
 }
 
+/// Systemic error codes for a job failing
+export enum JobStepError
+{
+	/// No systemic error
+	None,
+	/// Step did not complete in the required amount of time
+	TimedOut,
+	/// Step is in paused state so was skipped
+	Paused
+}
+
 // The state of a particular run
 export enum JobStepBatchState {
 
@@ -1048,6 +1059,9 @@ export type GetStepResponse = {
 	/**Current outcome of the jobstep */
 	outcome: JobStepOutcome;
 
+	/**Error describing additional context for why a step failed to complete*/
+	error: JobStepError;
+
 	/** If the step has been requested to abort	*/
 	abortRequested?: boolean;
 
