Horde: Allow issue creation to be disabled on a per-step and per-template level, via node annotations.

diff --git a/Engine/Source/Programs/Horde/Horde.Build/Issues/Impl/IssueService.cs b/Engine/Source/Programs/Horde/Horde.Build/Issues/Impl/IssueService.cs
--- a/Engine/Source/Programs/Horde/Horde.Build/Issues/Impl/IssueService.cs
+++ b/Engine/Source/Programs/Horde/Horde.Build/Issues/Impl/IssueService.cs
@@ -603,9 +603,9 @@ namespace Horde.Build.Services.Impl
 			}
 
 			// If the workflow disables issue creation, bail out now
-			if (workflow != null && !workflow.CreateIssues)
+			if (annotations.CreateIssues ?? true)
 			{
-				_logger.LogInformation("Issue creation for step {JobId}:{BatchId}:{StepId} disabled by workflow {WorkflowId}", job.Id, batchId, stepId, workflowId);
+				_logger.LogInformation("Issue creation for step {JobId}:{BatchId}:{StepId} disabled", job.Id, batchId, stepId);
 				return;
 			}
 
diff --git a/Engine/Source/Programs/Horde/Horde.Build/Jobs/NodeAnnotations.cs b/Engine/Source/Programs/Horde/Horde.Build/Jobs/NodeAnnotations.cs
--- a/Engine/Source/Programs/Horde/Horde.Build/Jobs/NodeAnnotations.cs
+++ b/Engine/Source/Programs/Horde/Horde.Build/Jobs/NodeAnnotations.cs
@@ -1,5 +1,6 @@
 // Copyright Epic Games, Inc. All Rights Reserved.
 
+using System;
 using System.Collections.Generic;
 using Horde.Build.Server;
 using Horde.Build.Utilities;
@@ -8,6 +9,24 @@ namespace Horde.Build.Models
 {
 	using WorkflowId = StringId<WorkflowConfig>;
 
+	/// <summary>
+	/// Interface which wraps a generic key/value dictionary to provide specific node annotations
+	/// </summary>
+#pragma warning disable CA1710 // Identifiers should have correct suffix
+	public interface IReadOnlyNodeAnnotations : IReadOnlyDictionary<string, string>
+#pragma warning restore CA1710 // Identifiers should have correct suffix
+	{
+		/// <summary>
+		/// Workflow to use for triaging issues from this node
+		/// </summary>
+		WorkflowId? WorkflowId { get; }
+
+		/// <summary>
+		/// Whether to create issues for this node
+		/// </summary>
+		bool? CreateIssues { get; }
+	}
+
 	/// <summary>
 	/// Set of annotations for a node
 	/// </summary>
@@ -18,11 +37,12 @@ namespace Horde.Build.Models
 		/// </summary>
 		public static IReadOnlyNodeAnnotations Empty { get; } = new NodeAnnotations();
 
-		/// <summary>
-		/// The issue workflow to use, as defined in the stream configuration file
-		/// </summary>
+		/// <inheritdoc cref="IReadOnlyNodeAnnotations.WorkflowId"/>
 		public const string WorkflowKeyName = "Workflow";
 
+		/// <inheritdoc cref="IReadOnlyNodeAnnotations.CreateIssues"/>
+		public const string CreateIssuesKeyName = "CreateIssues";
+
 		/// <summary>
 		/// Constructor
 		/// </summary>
@@ -51,10 +71,52 @@ namespace Horde.Build.Models
 				}
 				return new WorkflowId(workflowName);
 			}
-			set => SetOrUnset(WorkflowKeyName, value?.ToString());
+			set => SetValue(WorkflowKeyName, value?.ToString());
+		}
+
+		/// <inheritdoc/>
+		public bool? CreateIssues
+		{
+			get => GetBoolValue(CreateIssuesKeyName);
+			set => SetBoolValue(CreateIssuesKeyName, value);
+		}
+
+		private bool? GetBoolValue(string key)
+		{
+			string? value = GetValue(key);
+			if (value != null)
+			{
+				if (value.Equals("0", StringComparison.Ordinal) || value.Equals("false", StringComparison.OrdinalIgnoreCase))
+				{
+					return false;
+				}
+				if (value.Equals("1", StringComparison.Ordinal) || value.Equals("true", StringComparison.OrdinalIgnoreCase))
+				{
+					return true;
+				}
+			}
+			return null;
 		}
 
-		private void SetOrUnset(string key, string? value)
+		private void SetBoolValue(string key, bool? value)
+		{
+			if (value == null)
+			{
+				Remove(key);
+			}
+			else
+			{
+				SetValue(key, value.Value ? "1" : "0");
+			}
+		}
+
+		private string? GetValue(string key)
+		{
+			TryGetValue(key, out string? value);
+			return value;
+		}
+
+		private void SetValue(string key, string? value)
 		{
 			if (value == null)
 			{
@@ -78,17 +140,4 @@ namespace Horde.Build.Models
 			}
 		}
 	}
-
-	/// <summary>
-	/// Interface which wraps a generic key/value dictionary to provide specific node annotations
-	/// </summary>
-#pragma warning disable CA1710 // Identifiers should have correct suffix
-	public interface IReadOnlyNodeAnnotations : IReadOnlyDictionary<string, string>
-#pragma warning restore CA1710 // Identifiers should have correct suffix
-	{
-		/// <summary>
-		/// Workflow to use for triaging issues from this node
-		/// </summary>
-		WorkflowId? WorkflowId { get; }
-	}
 }
diff --git a/Engine/Source/Programs/Horde/Horde.Build/Server/Config.cs b/Engine/Source/Programs/Horde/Horde.Build/Server/Config.cs
--- a/Engine/Source/Programs/Horde/Horde.Build/Server/Config.cs
+++ b/Engine/Source/Programs/Horde/Horde.Build/Server/Config.cs
@@ -307,11 +307,6 @@ namespace Horde.Build.Server
 		/// </summary>
 		public WorkflowId Id { get; set; } = WorkflowId.Empty;
 
-		/// <summary>
-		/// Whether to allow creating issues for these jobs
-		/// </summary>
-		public bool CreateIssues { get; set; } = true;
-
 		/// <summary>
 		/// Times of day at which to send a report
 		/// </summary>
