Horde: Fix inverted condition check causing test failures.

diff --git a/Engine/Source/Programs/Horde/Horde.Build/Issues/Impl/IssueService.cs b/Engine/Source/Programs/Horde/Horde.Build/Issues/Impl/IssueService.cs
--- a/Engine/Source/Programs/Horde/Horde.Build/Issues/Impl/IssueService.cs
+++ b/Engine/Source/Programs/Horde/Horde.Build/Issues/Impl/IssueService.cs
@@ -603,7 +603,7 @@ namespace Horde.Build.Services.Impl
 			}
 
 			// If the workflow disables issue creation, bail out now
-			if (annotations.CreateIssues ?? true)
+			if (!(annotations.CreateIssues ?? true))
 			{
 				_logger.LogInformation("Issue creation for step {JobId}:{BatchId}:{StepId} disabled", job.Id, batchId, stepId);
 				return;
