Horde: Fix typo, and force an incremental
#rnx

diff --git a/Engine/Source/Programs/Horde/Horde.Build/Jira/Impl/JiraService.cs b/Engine/Source/Programs/Horde/Horde.Build/Jira/Impl/JiraService.cs
--- a/Engine/Source/Programs/Horde/Horde.Build/Jira/Impl/JiraService.cs
+++ b/Engine/Source/Programs/Horde/Horde.Build/Jira/Impl/JiraService.cs
@@ -16,7 +16,7 @@ using Polly.Extensions.Http;
 namespace Horde.Build.Services
 {
 	/// <summary>
-	/// Jira service functinality
+	/// Jira service functionality
 	/// </summary>
 	public sealed class JiraService : IJiraService, IDisposable
 	{
