Horde: Add annotations to control whether submitters get notified for a newly created issue.

diff --git a/Engine/Source/Programs/Horde/Horde.Build/Issues/Impl/IssueCollection.cs b/Engine/Source/Programs/Horde/Horde.Build/Issues/Impl/IssueCollection.cs
--- a/Engine/Source/Programs/Horde/Horde.Build/Issues/Impl/IssueCollection.cs
+++ b/Engine/Source/Programs/Horde/Horde.Build/Issues/Impl/IssueCollection.cs
@@ -22,7 +22,6 @@ namespace Horde.Build.Collections.Impl
 	using StreamId = StringId<IStream>;
 	using TemplateRefId = StringId<TemplateRef>;
 	using UserId = ObjectId<IUser>;
-	using WorkflowId = StringId<WorkflowConfig>;
 
 	class IssueCollection : IIssueCollection
 	{
@@ -982,9 +981,12 @@ namespace Horde.Build.Collections.Impl
 				newPromoted = newSpans.Any(x => ((IIssueSpan)x).PromoteByDefault);
 			}
 
+			// Figure out if we can auto-assign an owner
+			bool canAutoAssign = newPromoted || newSpans.Any(x => x.LastFailure.Annotations?.AutoAssign ?? false);
+
 			// Find the default owner
 			UserId? newDefaultOwnerId = null;
-			if (newPromoted && newSuspectImpls.Count > 0)
+			if (canAutoAssign && newSuspectImpls.Count > 0)
 			{
 				UserId possibleOwnerId = newSuspectImpls[0].AuthorId;
 				if (newSuspectImpls.All(x => x.AuthorId == possibleOwnerId) && newSuspectImpls.Any(x => x.DeclinedAt == null))
diff --git a/Engine/Source/Programs/Horde/Horde.Build/Jobs/NodeAnnotations.cs b/Engine/Source/Programs/Horde/Horde.Build/Jobs/NodeAnnotations.cs
--- a/Engine/Source/Programs/Horde/Horde.Build/Jobs/NodeAnnotations.cs
+++ b/Engine/Source/Programs/Horde/Horde.Build/Jobs/NodeAnnotations.cs
@@ -25,6 +25,16 @@ namespace Horde.Build.Models
 		/// Whether to create issues for this node
 		/// </summary>
 		bool? CreateIssues { get; }
+
+		/// <summary>
+		/// Whether to automatically assign issues that could only be caused by one user, or have a well defined correlation with a modified file.
+		/// </summary>
+		bool? AutoAssign { get; }
+
+		/// <summary>
+		/// Whether to notify all submitters between a build suceeding and failing, allowing them to step forward and take ownership of an issue.
+		/// </summary>
+		bool? NotifySubmitters { get; }
 	}
 
 	/// <summary>
@@ -43,6 +53,12 @@ namespace Horde.Build.Models
 		/// <inheritdoc cref="IReadOnlyNodeAnnotations.CreateIssues"/>
 		public const string CreateIssuesKeyName = "CreateIssues";
 
+		/// <inheritdoc cref="IReadOnlyNodeAnnotations.AutoAssign"/>
+		public const string AutoAssignKeyName = "AutoAssign";
+
+		/// <inheritdoc cref="IReadOnlyNodeAnnotations.NotifySubmitters"/>
+		public const string NotifySubmittersKeyName = "NotifySubmitters";
+
 		/// <summary>
 		/// Constructor
 		/// </summary>
@@ -81,6 +97,20 @@ namespace Horde.Build.Models
 			set => SetBoolValue(CreateIssuesKeyName, value);
 		}
 
+		/// <inheritdoc/>
+		public bool? AutoAssign
+		{
+			get => GetBoolValue(AutoAssignKeyName);
+			set => SetBoolValue(AutoAssignKeyName, value);
+		}
+
+		/// <inheritdoc/>
+		public bool? NotifySubmitters
+		{
+			get => GetBoolValue(NotifySubmittersKeyName);
+			set => SetBoolValue(NotifySubmittersKeyName, value);
+		}
+
 		private bool? GetBoolValue(string key)
 		{
 			string? value = GetValue(key);
diff --git a/Engine/Source/Programs/Horde/Horde.Build/Notifications/Impl/SlackNotificationSink.cs b/Engine/Source/Programs/Horde/Horde.Build/Notifications/Impl/SlackNotificationSink.cs
--- a/Engine/Source/Programs/Horde/Horde.Build/Notifications/Impl/SlackNotificationSink.cs
+++ b/Engine/Source/Programs/Horde/Horde.Build/Notifications/Impl/SlackNotificationSink.cs
@@ -597,8 +597,10 @@ namespace Horde.Build.Notifications.Impl
 		{
 			IIssueDetails details = await _issueService.GetIssueDetailsAsync(issue);
 
+			bool notifySuspects = issue.Promoted || details.Spans.Any(x => x.LastFailure.Annotations.NotifySubmitters ?? false);
+
 			HashSet<UserId> userIds = new HashSet<UserId>();
-			if (issue.Promoted)
+			if (notifySuspects)
 			{
 				userIds.UnionWith(details.Suspects.Select(x => x.AuthorId));
 			}
