Horde: Add a setting to forcibly assign issues to a particular owner.

diff --git a/Engine/Source/Programs/Horde/Horde.Build/Issues/Impl/IssueCollection.cs b/Engine/Source/Programs/Horde/Horde.Build/Issues/Impl/IssueCollection.cs
--- a/Engine/Source/Programs/Horde/Horde.Build/Issues/Impl/IssueCollection.cs
+++ b/Engine/Source/Programs/Horde/Horde.Build/Issues/Impl/IssueCollection.cs
@@ -392,6 +392,7 @@ namespace Horde.Build.Collections.Impl
 		}
 
 		readonly RedisService _redisService;
+		readonly IUserCollection _userCollection;
 		readonly ISingletonDocument<IssueLedger> _ledgerSingleton;
 		readonly IMongoCollection<Issue> _issues;
 		readonly IMongoCollection<IssueSpan> _issueSpans;
@@ -400,9 +401,10 @@ namespace Horde.Build.Collections.Impl
 		readonly IAuditLog<int> _auditLog;
 		readonly ILogger _logger;
 
-		public IssueCollection(MongoService mongoService, RedisService redisService, IAuditLogFactory<int> auditLogFactory, ILogger<IssueCollection> logger)
+		public IssueCollection(MongoService mongoService, RedisService redisService, IUserCollection userCollection, IAuditLogFactory<int> auditLogFactory, ILogger<IssueCollection> logger)
 		{
 			_redisService = redisService;
+			_userCollection = userCollection;
 			_logger = logger;
 
 			_ledgerSingleton = new SingletonDocument<IssueLedger>(mongoService);
@@ -981,17 +983,28 @@ namespace Horde.Build.Collections.Impl
 				newPromoted = newSpans.Any(x => ((IIssueSpan)x).PromoteByDefault);
 			}
 
-			// Figure out if we can auto-assign an owner
-			bool canAutoAssign = newPromoted || newSpans.Any(x => x.LastFailure.Annotations?.AutoAssign ?? false);
-
 			// Find the default owner
 			UserId? newDefaultOwnerId = null;
-			if (canAutoAssign && newSuspectImpls.Count > 0)
+			string? autoAssignToUser = newSpans.Select(x => x.LastFailure.Annotations?.AutoAssignToUser).Where(x => x != null).FirstOrDefault();
+			if (autoAssignToUser != null)
+			{
+				IUser? user = await _userCollection.FindUserByLoginAsync(autoAssignToUser);
+				if(user != null)
+				{
+					newDefaultOwnerId = user.Id;
+				}
+			}
+			else
 			{
-				UserId possibleOwnerId = newSuspectImpls[0].AuthorId;
-				if (newSuspectImpls.All(x => x.AuthorId == possibleOwnerId) && newSuspectImpls.Any(x => x.DeclinedAt == null))
+				// Figure out if we can auto-assign an owner
+				bool canAutoAssign = newPromoted || newSpans.Any(x => x.LastFailure.Annotations?.AutoAssign ?? false);
+				if (canAutoAssign && newSuspectImpls.Count > 0)
 				{
-					newDefaultOwnerId = possibleOwnerId;
+					UserId possibleOwnerId = newSuspectImpls[0].AuthorId;
+					if (newSuspectImpls.All(x => x.AuthorId == possibleOwnerId) && newSuspectImpls.Any(x => x.DeclinedAt == null))
+					{
+						newDefaultOwnerId = possibleOwnerId;
+					}
 				}
 			}
 
diff --git a/Engine/Source/Programs/Horde/Horde.Build/Jobs/NodeAnnotations.cs b/Engine/Source/Programs/Horde/Horde.Build/Jobs/NodeAnnotations.cs
--- a/Engine/Source/Programs/Horde/Horde.Build/Jobs/NodeAnnotations.cs
+++ b/Engine/Source/Programs/Horde/Horde.Build/Jobs/NodeAnnotations.cs
@@ -31,6 +31,11 @@ namespace Horde.Build.Models
 		/// </summary>
 		bool? AutoAssign { get; }
 
+		/// <summary>
+		/// Automatically assign any issues to the given user
+		/// </summary>
+		string? AutoAssignToUser { get; }
+
 		/// <summary>
 		/// Whether to notify all submitters between a build suceeding and failing, allowing them to step forward and take ownership of an issue.
 		/// </summary>
@@ -56,6 +61,9 @@ namespace Horde.Build.Models
 		/// <inheritdoc cref="IReadOnlyNodeAnnotations.AutoAssign"/>
 		public const string AutoAssignKeyName = "AutoAssign";
 
+		/// <inheritdoc cref="IReadOnlyNodeAnnotations.AutoAssignToUser"/>
+		public const string AutoAssignToUserKeyName = "AutoAssignToUser";
+
 		/// <inheritdoc cref="IReadOnlyNodeAnnotations.NotifySubmitters"/>
 		public const string NotifySubmittersKeyName = "NotifySubmitters";
 
@@ -87,7 +95,7 @@ namespace Horde.Build.Models
 				}
 				return new WorkflowId(workflowName);
 			}
-			set => SetValue(WorkflowKeyName, value?.ToString());
+			set => SetStringValue(WorkflowKeyName, value?.ToString());
 		}
 
 		/// <inheritdoc/>
@@ -104,6 +112,13 @@ namespace Horde.Build.Models
 			set => SetBoolValue(AutoAssignKeyName, value);
 		}
 
+		/// <inheritdoc/>
+		public string? AutoAssignToUser
+		{
+			get => GetStringValue(AutoAssignToUserKeyName);
+			set => SetStringValue(AutoAssignToUserKeyName, value);
+		}
+
 		/// <inheritdoc/>
 		public bool? NotifySubmitters
 		{
@@ -113,7 +128,7 @@ namespace Horde.Build.Models
 
 		private bool? GetBoolValue(string key)
 		{
-			string? value = GetValue(key);
+			string? value = GetStringValue(key);
 			if (value != null)
 			{
 				if (value.Equals("0", StringComparison.Ordinal) || value.Equals("false", StringComparison.OrdinalIgnoreCase))
@@ -136,17 +151,17 @@ namespace Horde.Build.Models
 			}
 			else
 			{
-				SetValue(key, value.Value ? "1" : "0");
+				SetStringValue(key, value.Value ? "1" : "0");
 			}
 		}
 
-		private string? GetValue(string key)
+		private string? GetStringValue(string key)
 		{
 			TryGetValue(key, out string? value);
 			return value;
 		}
 
-		private void SetValue(string key, string? value)
+		private void SetStringValue(string key, string? value)
 		{
 			if (value == null)
 			{
