Horde: Fix slack messages no longer replying in threads.

diff --git a/Engine/Source/Programs/Horde/Horde.Build/Notifications/Impl/SlackNotificationSink.cs b/Engine/Source/Programs/Horde/Horde.Build/Notifications/Impl/SlackNotificationSink.cs
--- a/Engine/Source/Programs/Horde/Horde.Build/Notifications/Impl/SlackNotificationSink.cs
+++ b/Engine/Source/Programs/Horde/Horde.Build/Notifications/Impl/SlackNotificationSink.cs
@@ -1559,7 +1559,11 @@ namespace Horde.Build.Notifications.Impl
 				{
 					_logger.LogWarning("Missing 'ts' or 'channel' field on slack response");
 				}
-				await SetMessageTimestampAsync(state.Id, response.Channel ?? String.Empty, response.Ts ?? String.Empty);
+
+				state.Channel = response.Channel ?? String.Empty;
+				state.Ts = response.Ts ?? String.Empty;
+
+				await SetMessageTimestampAsync(state.Id, state.Channel, state.Ts);
 			}
 			else if (!String.IsNullOrEmpty(state.Ts))
 			{
