Horde: Prevent config data being saved to database. It is already persisted via the config collection, and saving it to BSON causes exceptions due to inheritance in the tab list.

diff --git a/Engine/Source/Programs/Horde/Horde.Build/Projects/Impl/StreamCollection.cs b/Engine/Source/Programs/Horde/Horde.Build/Projects/Impl/StreamCollection.cs
--- a/Engine/Source/Programs/Horde/Horde.Build/Projects/Impl/StreamCollection.cs
+++ b/Engine/Source/Programs/Horde/Horde.Build/Projects/Impl/StreamCollection.cs
@@ -1,4 +1,4 @@
-// Copyright Epic Games, Inc. All Rights Reserved.
+﻿// Copyright Epic Games, Inc. All Rights Reserved.
 
 using System;
 using System.Collections.Generic;
@@ -47,6 +47,8 @@ namespace Horde.Build.Collections.Impl
 
 			public string? ClusterName { get; set; }
 			public string ConfigRevision { get; set; } = String.Empty;
+
+			[BsonIgnore]
 			public StreamConfig? Config { get; set; }
 
 			public int Order { get; set; } = DefaultOrder;
