Horde: Do not initialize PerforceService when running on Apple silicon; we do not currently have native arm64 libraries.

diff --git a/Engine/Source/Programs/Horde/Horde.Build/Models/Globals.cs b/Engine/Source/Programs/Horde/Horde.Build/Models/Globals.cs
--- a/Engine/Source/Programs/Horde/Horde.Build/Models/Globals.cs
+++ b/Engine/Source/Programs/Horde/Horde.Build/Models/Globals.cs
@@ -1,10 +1,11 @@
-// Copyright Epic Games, Inc. All Rights Reserved.
+﻿// Copyright Epic Games, Inc. All Rights Reserved.
 
 using System;
 using System.Collections.Generic;
 using System.ComponentModel.DataAnnotations;
 using System.Linq;
 using System.Security.Cryptography;
+using EpicGames.Perforce;
 using Horde.Build.Acls;
 using Horde.Build.Server;
 using Horde.Build.Utilities;
@@ -372,7 +373,7 @@ namespace Horde.Build.Models
 		static List<PerforceCluster> GetDefaultClusters()
 		{
 			PerforceServer server = new PerforceServer();
-			server.ServerAndPort = Perforce.P4.P4Server.Get("P4PORT") ?? "perforce:1666";
+			server.ServerAndPort = PerforceSettings.Default.ServerAndPort;
 
 			PerforceCluster cluster = new PerforceCluster();
 			cluster.Name = "Default";
diff --git a/Engine/Source/Programs/Horde/Horde.Build/Perforce/Impl/PerforceService.cs b/Engine/Source/Programs/Horde/Horde.Build/Perforce/Impl/PerforceService.cs
--- a/Engine/Source/Programs/Horde/Horde.Build/Perforce/Impl/PerforceService.cs
+++ b/Engine/Source/Programs/Horde/Horde.Build/Perforce/Impl/PerforceService.cs
@@ -1,4 +1,4 @@
-// Copyright Epic Games, Inc. All Rights Reserved.
+﻿// Copyright Epic Games, Inc. All Rights Reserved.
 
 using System;
 using System.Collections.Generic;
@@ -63,10 +63,11 @@ namespace Horde.Build.Services
 		readonly PerforceLoadBalancer _loadBalancer;
 		readonly LazyCachedValue<Task<Globals>> _cachedGlobals;
 		readonly ILogger _logger;
+		readonly bool _haveNativePerforceApi;
 
 		// Useful overrides for local debugging with read-only data
-		string? _perforceServerOverride;
-		string? _perforceUserOverride;
+		readonly string? _perforceServerOverride;
+		readonly string? _perforceUserOverride;
 
 		/// <summary>
 		/// Object used for controlling access to the access user tickets
@@ -103,9 +104,13 @@ namespace Horde.Build.Services
 			_logger = logger;
 
 			_logBridgeDelegate = new P4.P4CallBacks.LogMessageDelegate(LogBridgeMessage);
-			P4.P4Debugging.SetBridgeLogFunction(_logBridgeDelegate);
 
-			P4.LogFile.SetLoggingFunction(LogPerforce);
+			_haveNativePerforceApi = !(RuntimeInformation.IsOSPlatform(OSPlatform.OSX) && RuntimeInformation.OSArchitecture == Architecture.Arm64);
+			if (_haveNativePerforceApi)
+			{
+				P4.P4Debugging.SetBridgeLogFunction(_logBridgeDelegate);
+				P4.LogFile.SetLoggingFunction(LogPerforce);
+			}
 
 			if(settings.Value.UseLocalPerforceEnv)
 			{
