Horde: Upgrade gRPC to latest stable version for Apple silicon support.

diff --git a/Engine/Source/Programs/Horde/Horde.Build/Horde.Build.csproj b/Engine/Source/Programs/Horde/Horde.Build/Horde.Build.csproj
--- a/Engine/Source/Programs/Horde/Horde.Build/Horde.Build.csproj
+++ b/Engine/Source/Programs/Horde/Horde.Build/Horde.Build.csproj
@@ -25,10 +25,10 @@
     <PackageReference Include="AWSSDK.SecurityToken" Version="3.7.1.69" />
     <PackageReference Include="Datadog.Trace.OpenTracing" Version="1.18.0" />
     <PackageReference Include="DogStatsD-CSharp-Client" Version="6.0.0" />
-    <PackageReference Include="Grpc.AspNetCore" Version="2.37.0" />
-    <PackageReference Include="Grpc.AspNetCore.Server.Reflection" Version="2.37.0" />
-    <PackageReference Include="Grpc.Net.Client" Version="2.37.0" />
-    <PackageReference Include="Grpc.Tools" Version="2.37.0">
+    <PackageReference Include="Grpc.AspNetCore" Version="2.45.0" />
+    <PackageReference Include="Grpc.AspNetCore.Server.Reflection" Version="2.45.0" />
+    <PackageReference Include="Grpc.Net.Client" Version="2.45.0" />
+    <PackageReference Include="Grpc.Tools" Version="2.45.0">
       <PrivateAssets>all</PrivateAssets>
       <IncludeAssets>runtime; build; native; contentfiles; analyzers; buildtransitive</IncludeAssets>
     </PackageReference>
