Horde: Update gRPC packages for Horde.Agent to latest version for Apple silicon support.

diff --git a/Engine/Source/Programs/Horde/Horde.Agent/Horde.Agent.csproj b/Engine/Source/Programs/Horde/Horde.Agent/Horde.Agent.csproj
--- a/Engine/Source/Programs/Horde/Horde.Agent/Horde.Agent.csproj
+++ b/Engine/Source/Programs/Horde/Horde.Agent/Horde.Agent.csproj
@@ -25,9 +25,9 @@
     <PackageReference Include="AWSSDK.Core" Version="3.5.1.49" />
     <PackageReference Include="AWSSDK.EC2" Version="3.5.1.2" />
     <PackageReference Include="Datadog.Trace.OpenTracing" Version="1.18.0" />
-    <PackageReference Include="Google.Protobuf" Version="3.12.3" />
-    <PackageReference Include="Grpc.Net.Client" Version="2.29.0" />
-    <PackageReference Include="Grpc.Tools" Version="2.30.0">
+    <PackageReference Include="Google.Protobuf" Version="3.20.1" />
+    <PackageReference Include="Grpc.Net.Client" Version="2.45.0" />
+    <PackageReference Include="Grpc.Tools" Version="2.45.0">
       <PrivateAssets>all</PrivateAssets>
       <IncludeAssets>runtime; build; native; contentfiles; analyzers; buildtransitive</IncludeAssets>
     </PackageReference>
diff --git a/Engine/Source/Programs/Horde/Horde.Build.Tests/RpcServiceTest.cs b/Engine/Source/Programs/Horde/Horde.Build.Tests/RpcServiceTest.cs
--- a/Engine/Source/Programs/Horde/Horde.Build.Tests/RpcServiceTest.cs
+++ b/Engine/Source/Programs/Horde/Horde.Build.Tests/RpcServiceTest.cs
@@ -1,4 +1,4 @@
-// Copyright Epic Games, Inc. All Rights Reserved.
+﻿// Copyright Epic Games, Inc. All Rights Reserved.
 
 using System;
 using System.Collections.Generic;
@@ -126,7 +126,7 @@ namespace Horde.Build.Tests
 				throw new NotImplementedException();
 			}
 
-			protected override ContextPropagationToken CreatePropagationTokenCore(ContextPropagationOptions options)
+			protected override ContextPropagationToken CreatePropagationTokenCore(ContextPropagationOptions? options)
 			{
 				throw new NotImplementedException();
 			}
@@ -139,7 +139,7 @@ namespace Horde.Build.Tests
 			protected override CancellationToken CancellationTokenCore { get; } = new CancellationToken();
 			protected override Metadata ResponseTrailersCore { get; } = null!;
 			protected override Status StatusCore { get; set; }
-			protected override WriteOptions WriteOptionsCore { get; set; } = null!;
+			protected override WriteOptions? WriteOptionsCore { get; set; } = null!;
 			protected override AuthContext AuthContextCore { get; } = null!;
 		}
 		
@@ -154,32 +154,32 @@ namespace Horde.Build.Tests
 				_serverCallContext = serverCallContext;
 			}
 
-			public override TResponse BlockingUnaryCall<TRequest, TResponse>(Method<TRequest, TResponse> method, string host, CallOptions options, TRequest request)
+			public override TResponse BlockingUnaryCall<TRequest, TResponse>(Method<TRequest, TResponse> method, string? host, CallOptions options, TRequest request)
 			{
 				throw new NotImplementedException("Blocking calls are not supported! Method " + method.FullName);
 			}
 
-			public override AsyncUnaryCall<TResponse> AsyncUnaryCall<TRequest, TResponse>(Method<TRequest, TResponse> method, string host, CallOptions options, TRequest request)
+			public override AsyncUnaryCall<TResponse> AsyncUnaryCall<TRequest, TResponse>(Method<TRequest, TResponse> method, string? host, CallOptions options, TRequest request)
 			{
 				MethodInfo methodInfo = GetMethod(method.Name);
 				Task<TResponse> res = (methodInfo.Invoke(_rpcService, new object[] {request, _serverCallContext}) as Task<TResponse>)!;
 				return new AsyncUnaryCall<TResponse>(res, null!, null!, null!, null!, null!);
 			}
 
-			public override AsyncServerStreamingCall<TResponse> AsyncServerStreamingCall<TRequest, TResponse>(Method<TRequest, TResponse> method, string host, CallOptions options,
+			public override AsyncServerStreamingCall<TResponse> AsyncServerStreamingCall<TRequest, TResponse>(Method<TRequest, TResponse> method, string? host, CallOptions options,
 				TRequest request)
 			{
 				Console.WriteLine($"RpcServiceInvoker.AsyncServerStreamingCall(method={method.FullName} request={request})");
 				throw new NotImplementedException();
 			}
 
-			public override AsyncClientStreamingCall<TRequest, TResponse> AsyncClientStreamingCall<TRequest, TResponse>(Method<TRequest, TResponse> method, string host, CallOptions options)
+			public override AsyncClientStreamingCall<TRequest, TResponse> AsyncClientStreamingCall<TRequest, TResponse>(Method<TRequest, TResponse> method, string? host, CallOptions options)
 			{
 				Console.WriteLine($"RpcServiceInvoker.AsyncClientStreamingCall(method={method.FullName})");
 				throw new NotImplementedException();
 			}
 
-			public override AsyncDuplexStreamingCall<TRequest, TResponse> AsyncDuplexStreamingCall<TRequest, TResponse>(Method<TRequest, TResponse> method, string host, CallOptions options)
+			public override AsyncDuplexStreamingCall<TRequest, TResponse> AsyncDuplexStreamingCall<TRequest, TResponse>(Method<TRequest, TResponse> method, string? host, CallOptions options)
 			{
 				Console.WriteLine($"RpcServiceInvoker.AsyncDuplexStreamingCall(method={method.FullName})");
 
diff --git a/Engine/Source/Programs/Shared/EpicGames.Perforce/NativePerforceConnection.cs b/Engine/Source/Programs/Shared/EpicGames.Perforce/NativePerforceConnection.cs
--- a/Engine/Source/Programs/Shared/EpicGames.Perforce/NativePerforceConnection.cs
+++ b/Engine/Source/Programs/Shared/EpicGames.Perforce/NativePerforceConnection.cs
@@ -326,6 +326,10 @@ namespace EpicGames.Perforce
 			return connection;
 		}
 
+		/// <summary>
+		/// Check whether the native client is supported on the current platform
+		/// </summary>
+		/// <returns></returns>
 		public static bool IsSupported() => !(RuntimeInformation.IsOSPlatform(OSPlatform.OSX) && RuntimeInformation.OSArchitecture == Architecture.Arm64);
 
 		void GetNextWriteBuffer(NativeWriteBuffer nativeWriteBuffer, int minSize)
