Horde: Issue modal improvements
#skipci
#rnx

diff --git a/Engine/Source/Programs/Horde/HordeDashboard/src/backend/Dashboard.ts b/Engine/Source/Programs/Horde/HordeDashboard/src/backend/Dashboard.ts
--- a/Engine/Source/Programs/Horde/HordeDashboard/src/backend/Dashboard.ts
+++ b/Engine/Source/Programs/Horde/HordeDashboard/src/backend/Dashboard.ts
@@ -456,6 +456,9 @@ export class Dashboard {
         return this.data.id !== "";
     }
 
+    get user(): GetUserResponse {
+        return this.data;
+    }
 
     private async postPreferences(reload?:boolean): Promise<boolean> {
 
diff --git a/Engine/Source/Programs/Horde/HordeDashboard/src/index.css b/Engine/Source/Programs/Horde/HordeDashboard/src/index.css
--- a/Engine/Source/Programs/Horde/HordeDashboard/src/index.css
+++ b/Engine/Source/Programs/Horde/HordeDashboard/src/index.css
@@ -80,4 +80,19 @@ div[data-automation-key="ViewLogColumn"] {
   padding: 0
 }
 
+@keyframes red-pulse {
+  0% {
+    background-color: #CC0000;
+    border: 1px solid #CC0000;
+  }
+  50% {
+    background-color: #FF0000;
+    border: 1px solid #FF0000;
+  }
+  100% {
+    background-color: #CC0000;
+    border: 1px solid #CC0000;
+  }
+
+}
 
