Use memory backed files if the content we need to replicate is small enough to fit into a buffer

diff --git a/Engine/Source/Programs/Horde/HordeStorage/Horde.Storage/Implementation/Replication/RefsReplicator.cs b/Engine/Source/Programs/Horde/HordeStorage/Horde.Storage/Implementation/Replication/RefsReplicator.cs
--- a/Engine/Source/Programs/Horde/HordeStorage/Horde.Storage/Implementation/Replication/RefsReplicator.cs
+++ b/Engine/Source/Programs/Horde/HordeStorage/Horde.Storage/Implementation/Replication/RefsReplicator.cs
@@ -511,7 +511,10 @@ namespace Horde.Storage.Implementation
                     }
 
                     await using Stream s = await blobResponse.Content.ReadAsStreamAsync(cancellationToken);
-                    using FilesystemBufferedPayload payload = await FilesystemBufferedPayload.Create(s);
+                    long? contentLength = blobResponse.Content.Headers.ContentLength;
+                            
+                    using IBufferedPayload payload = contentLength is null or > int.MaxValue ? await FilesystemBufferedPayload.Create(s) : await MemoryBufferedPayload.Create(s);
+
                     await _blobService.PutObject(ns, payload, blobToReplicate);
                 });
             }
diff --git a/Engine/Source/Programs/Horde/HordeStorage/Horde.Storage/Implementation/Replication/Replicator.cs b/Engine/Source/Programs/Horde/HordeStorage/Horde.Storage/Implementation/Replication/Replicator.cs
--- a/Engine/Source/Programs/Horde/HordeStorage/Horde.Storage/Implementation/Replication/Replicator.cs
+++ b/Engine/Source/Programs/Horde/HordeStorage/Horde.Storage/Implementation/Replication/Replicator.cs
@@ -384,8 +384,10 @@ namespace Horde.Storage.Implementation
 
                         if (response.IsSuccessStatusCode)
                         {
+                            long? contentLength = response.Content.Headers.ContentLength;
                             await using Stream responseStream = await response.Content.ReadAsStreamAsync(replicationToken);
-                            using FilesystemBufferedPayload payload = await FilesystemBufferedPayload.Create(responseStream);
+                            
+                            using IBufferedPayload payload = contentLength is null or > int.MaxValue ? await FilesystemBufferedPayload.Create(responseStream) : await MemoryBufferedPayload.Create(responseStream);
 
                             {
                                 await using Stream s = payload.GetStream();
