Horde: Fix issue with link event matcher not recognizing clang link errors with test case
#rnx
#skipci

diff --git a/Engine/Source/Programs/Horde/Horde.Agent.Tests/LogParserTests.cs b/Engine/Source/Programs/Horde/Horde.Agent.Tests/LogParserTests.cs
--- a/Engine/Source/Programs/Horde/Horde.Agent.Tests/LogParserTests.cs
+++ b/Engine/Source/Programs/Horde/Horde.Agent.Tests/LogParserTests.cs
@@ -484,6 +484,28 @@ namespace Horde.Agent.Tests
 				Assert.AreEqual(1, logEvents.Count);
 				Assert.AreEqual(LogLevel.Error, logEvents[0].Level);
 			}
+
+			{
+				string[] lines =
+				{
+					@"tool 19.50.0.12 (rel,tool,19.500 @527452 x64) D:\Workspaces\AutoSDK\HostWin64\9.508.001\9.500\host_tools\bin\tool.exe",
+					@"Link : error: L0039: reference to undefined symbol `UProjectPlayControllerComponent::HandleMoveToolSpawnedActor(ATestGameCreativeMoveTool*, AActor*, bool)' in file ""D:\Workspaces\TestGame\TestGame\Intermediate\Build\Win64\TestGame\Development\ProjectPlayGameRuntime\Module.ProjectPlayGameRuntime.gen.cpp.o""\",
+					@"Link : error: L0039: reference to undefined symbol `UProjectPlayControllerComponent::HandleWeaponEquipped(ATestGameWeapon*, ATestGameWeapon*)' in file ""D:\Workspaces\TestGame\TestGame\Intermediate\Build\Win64\TestGameClient\Development\ProjectPlayGameRuntime\Module.ProjectPlayGameRuntime.gen.cpp.o""\",
+					@"Module.ProjectPlayGameRuntime.gen.cpp : error: L0039: reference to undefined symbol `UProjectPlayControllerComponent::HandleMoveToolSpawnedActor(ATestGameMoveTool*, AActor*, bool)' in file ""D:\Workspaces\TestGame\TestGame\Intermediate\Build\Win64\TestGameClient\Development\ProjectPlayGameRuntime\Module.ProjectPlayGameRuntime.gen.cpp.o""\",
+					@"Module.ProjectPlayGameRuntime.gen.cpp : error: L0039: reference to undefined symbol `UProjectPlayControllerComponent::HandleWeaponEquipped(ATestGameWeapon*, ATestGameWeapon*)' in file ""D:\Workspaces\TestGame\TestGame\Intermediate\Build\Win64\TestGame\Development\ProjectPlayGameRuntime\Module.ProjectPlayGameRuntime.gen.cpp.o""\",
+					@"orbis-clang: error: linker command failed with exit code 1 (use -v to see invocation)"
+				};
+
+				List<LogEvent> logEvents = Parse(String.Join("\n", lines), new DirectoryReference("[REDACTED_USER_HOME]/Build/++UE5/Sync"));
+				Assert.AreEqual(5, logEvents.Count);
+				CheckEventGroup(logEvents.Slice(0, 5), 1, 5, LogLevel.Error, KnownLogEvents.Linker_UndefinedSymbol);
+
+				Assert.AreEqual("UProjectPlayControllerComponent::HandleMoveToolSpawnedActor(ATestGameCreativeMoveTool*, AActor*, bool)", logEvents[0].Properties!["symbol"].ToString());
+				Assert.AreEqual("UProjectPlayControllerComponent::HandleWeaponEquipped(ATestGameWeapon*, ATestGameWeapon*)", logEvents[1].Properties!["symbol"].ToString());
+				Assert.AreEqual("UProjectPlayControllerComponent::HandleMoveToolSpawnedActor(ATestGameMoveTool*, AActor*, bool)", logEvents[2].Properties!["symbol"].ToString());
+				Assert.AreEqual("UProjectPlayControllerComponent::HandleWeaponEquipped(ATestGameWeapon*, ATestGameWeapon*)", logEvents[3].Properties!["symbol"].ToString());
+
+			}
 		}
 
 		[TestMethod]
diff --git a/Engine/Source/Programs/Horde/Horde.Agent/Parser/Matchers/LinkEventMatcher.cs b/Engine/Source/Programs/Horde/Horde.Agent/Parser/Matchers/LinkEventMatcher.cs
--- a/Engine/Source/Programs/Horde/Horde.Agent/Parser/Matchers/LinkEventMatcher.cs
+++ b/Engine/Source/Programs/Horde/Horde.Agent/Parser/Matchers/LinkEventMatcher.cs
@@ -18,7 +18,7 @@ namespace Horde.Agent.Parser.Matchers
 			int lineCount = 0;
 			bool isError = false;
 			bool isWarning = false;
-			while (cursor.IsMatch(lineCount, @": undefined reference to |undefined symbol|^\s*(ld|ld.lld):|^[^:]*[^a-zA-Z][a-z]?ld: |^\s*>>>"))
+			while (cursor.IsMatch(lineCount, @": undefined reference to |undefined symbol|^\s*(ld|ld.lld|(.*clang)):|^[^:]*[^a-zA-Z][a-z]?ld: |^\s*>>>"))
 			{
 				isError |= cursor.IsMatch("error:");
 				isWarning |= cursor.IsMatch("warning:");
