Horde: Tighten up the clang linker error matcher
#rnx
#skipci

diff --git a/Engine/Source/Programs/Horde/Horde.Agent/Parser/Matchers/LinkEventMatcher.cs b/Engine/Source/Programs/Horde/Horde.Agent/Parser/Matchers/LinkEventMatcher.cs
--- a/Engine/Source/Programs/Horde/Horde.Agent/Parser/Matchers/LinkEventMatcher.cs
+++ b/Engine/Source/Programs/Horde/Horde.Agent/Parser/Matchers/LinkEventMatcher.cs
@@ -18,7 +18,7 @@ namespace Horde.Agent.Parser.Matchers
 			int lineCount = 0;
 			bool isError = false;
 			bool isWarning = false;
-			while (cursor.IsMatch(lineCount, @": undefined reference to |undefined symbol|^\s*(ld|ld.lld|(.*clang)):|^[^:]*[^a-zA-Z][a-z]?ld: |^\s*>>>"))
+			while (cursor.IsMatch(lineCount, @": undefined reference to |undefined symbol|^\s*(ld|ld.lld):|.*clang:.*linker command failed|^[^:]*[^a-zA-Z][a-z]?ld: |^\s*>>>"))
 			{
 				isError |= cursor.IsMatch("error:");
 				isWarning |= cursor.IsMatch("warning:");
