Horde: Add notices for both downtime and admin set alerts
#rnx
#skipci

diff --git a/Engine/Source/Programs/Horde/Horde.Build/Api/Notices.cs b/Engine/Source/Programs/Horde/Horde.Build/Api/Notices.cs
--- a/Engine/Source/Programs/Horde/Horde.Build/Api/Notices.cs
+++ b/Engine/Source/Programs/Horde/Horde.Build/Api/Notices.cs
@@ -72,12 +72,12 @@ namespace Horde.Build.Api
 	}
 
 	/// <summary>
-	/// Notice informtation
+	/// Notice information
 	/// </summary>
 	public class GetNoticeResponse
 	{
 		/// <summary>
-		/// Object id of notice
+		/// The id of the notice to update
 		/// </summary>
 		public string? Id { get; set; }
 
@@ -92,34 +92,32 @@ namespace Horde.Build.Api
 		public DateTime? FinishTime { get; set; }
 
 		/// <summary>
-		/// User id who created the notice, otherwise null if a system message
+		/// Whether this notice is for scheduled downtime
 		/// </summary>
-		public GetThinUserInfoResponse? CreatedByUser { get; set; }
+		public bool ScheduledDowntime { get; set; } = false;
 
 		/// <summary>
-		/// Message to display
+		/// Whether the notice is currently active
 		/// </summary>
-		[Required]
-		public string Message { get; set; } = String.Empty;
+		public bool Active { get; set; }
 
 		/// <summary>
-		/// Constructor
+		/// Message to display
 		/// </summary>
-		public GetNoticeResponse()
-		{
+		public string? Message { get; set; }
 
-		}
+		/// <summary>
+		/// User id who created the notice, otherwise null if a system message
+		/// </summary>
+		public GetThinUserInfoResponse? CreatedByUser { get; set; }
 
 		/// <summary>
 		/// Constructor
 		/// </summary>
-		public GetNoticeResponse(INotice notice, GetThinUserInfoResponse? createByUser = null)
+		public GetNoticeResponse()
 		{
-			Id = notice.Id.ToString();
-			StartTime = notice.StartTime;
-			FinishTime = notice.FinishTime;
-			Message = notice.Message;
-			CreatedByUser = createByUser;
+
 		}
+
 	}
 }
diff --git a/Engine/Source/Programs/Horde/Horde.Build/Controllers/NoticesController.cs b/Engine/Source/Programs/Horde/Horde.Build/Controllers/NoticesController.cs
--- a/Engine/Source/Programs/Horde/Horde.Build/Controllers/NoticesController.cs
+++ b/Engine/Source/Programs/Horde/Horde.Build/Controllers/NoticesController.cs
@@ -2,7 +2,6 @@
 
 using System;
 using System.Collections.Generic;
-using System.Globalization;
 using System.Threading.Tasks;
 using Horde.Build.Acls;
 using Horde.Build.Api;
@@ -14,7 +13,7 @@ using Microsoft.AspNetCore.Authorization;
 using Microsoft.AspNetCore.Mvc;
 using Microsoft.Extensions.Options;
 using MongoDB.Bson;
-using TimeZoneConverter;
+using Horde.Build.Server;
 
 namespace Horde.Build.Controllers
 {
@@ -27,24 +26,29 @@ namespace Horde.Build.Controllers
 	public class NoticesController : ControllerBase
 	{
 		/// <summary>
-		/// 
+		/// The acl service singleton
 		/// </summary>
-		readonly NoticeService _noticeService;
+		readonly AclService _aclService;
 
 		/// <summary>
-		/// The acl service singleton
+		/// Reference to the user collection
 		/// </summary>
-		readonly AclService _aclService;
+		private readonly IUserCollection _userCollection;
 
 		/// <summary>
-		/// Reference to the globals singleton
+		/// Notice service
 		/// </summary>
-		readonly ISingletonDocument<Globals> _globals;
+		private readonly NoticeService _noticeService;
 
 		/// <summary>
-		/// Reference to the user collection
+		/// cached globals
 		/// </summary>
-		private readonly IUserCollection _userCollection;
+		readonly LazyCachedValue<Task<Globals>> _cachedGlobals;
+
+		/// <summary>
+		/// cached notices
+		/// </summary>
+		readonly LazyCachedValue<Task<List<INotice>>> _cachedNotices;
 
 		/// <summary>
 		/// Server settings
@@ -55,17 +59,18 @@ namespace Horde.Build.Controllers
 		/// Constructor
 		/// </summary>
 		/// <param name="settings">The server settings</param>
+		/// <param name="mongoService">The mongo service singleton</param>
 		/// <param name="noticeService">The notice service singleton</param>
 		/// <param name="aclService">The acl service singleton</param>
 		/// <param name="userCollection">The user collection singleton</param>
-		/// <param name="globals">The global singleton</param>
-		public NoticesController(IOptionsMonitor<ServerSettings> settings, NoticeService noticeService, AclService aclService, IUserCollection userCollection, ISingletonDocument<Globals> globals)
-		{
+		public NoticesController(IOptionsMonitor<ServerSettings> settings, MongoService mongoService, NoticeService noticeService, AclService aclService, IUserCollection userCollection)
+		{			
 			_settings = settings;
-			_noticeService = noticeService;
 			_aclService = aclService;
 			_userCollection = userCollection;
-			_globals = globals;
+			_noticeService = noticeService;
+			_cachedGlobals = new LazyCachedValue<Task<Globals>>(() => mongoService.GetGlobalsAsync(), TimeSpan.FromSeconds(30.0));
+			_cachedNotices = new LazyCachedValue<Task<List<INotice>>>(() => noticeService.GetNoticesAsync(), TimeSpan.FromMinutes(1));
 		}
 
 		/// <summary>
@@ -126,50 +131,22 @@ namespace Horde.Build.Controllers
 		/// <returns>The status messages</returns>
 		[HttpGet("/api/v1/notices")]
 		public async Task<List<GetNoticeResponse>> GetNoticesAsync()
-		{			
-			string? scheduleTimeZone = _settings.CurrentValue.ScheduleTimeZone;
-			TimeZoneInfo timeZone = (scheduleTimeZone == null) ? TimeZoneInfo.Local : TZConvert.GetTimeZoneInfo(scheduleTimeZone);
-			
-			Globals globalsValue = await _globals.GetAsync();
+		{
 			List<GetNoticeResponse> messages = new List<GetNoticeResponse>();
 
-			// Add system downtime notices
-			for (int idx = 0; idx < globalsValue.ScheduledDowntime.Count; idx++)
-			{
-				ScheduledDowntime downtime = globalsValue.ScheduledDowntime[idx];
-
-				// @todo: handle this, though likely will have a custom message
-				if (downtime.Frequency == ScheduledDowntimeFrequency.Once)
-				{
-					continue;
-				}
+			Globals globals = await _cachedGlobals.GetCached();
 
-				DateTimeOffset now = TimeZoneInfo.ConvertTime(DateTimeOffset.Now, _settings.CurrentValue.TimeZoneInfo);
-				
-				// @todo: this time and formatting should be on client
-				if (downtime.IsActive(now))
-				{
-					(DateTimeOffset startTime, DateTimeOffset finishTime) = downtime.GetNext(now);
-					string finishTimeString = String.Format(CultureInfo.CurrentCulture, "{0:t}", TimeZoneInfo.ConvertTime(finishTime, timeZone));
-					messages.Add(new GetNoticeResponse() { StartTime = startTime.UtcDateTime, FinishTime = finishTime.UtcDateTime, Message = $"Horde is currently in scheduled downtime. Jobs will resume execution at {finishTimeString} {timeZone.Id}." });
-				}
-				else 
-				{
-					// add a week to get the actual next scheduled time
-					if (downtime.Frequency == ScheduledDowntimeFrequency.Weekly)
-					{
-						now += TimeSpan.FromDays(7);
-					}
-
-					(DateTimeOffset startTime, DateTimeOffset finishTime) = downtime.GetNext(now);
-					string startTimeString = String.Format(CultureInfo.CurrentCulture, "{0:g}", startTime);
-					string finishTimeString = String.Format(CultureInfo.CurrentCulture, "{0:g}", finishTime);
-					messages.Add(new GetNoticeResponse() { StartTime = startTime.UtcDateTime, FinishTime = finishTime.UtcDateTime, Message = $"Downtime is scheduled from {startTimeString} to {finishTimeString} {timeZone.Id}. No jobs will run during this time." });
-				}
+			DateTimeOffset now = TimeZoneInfo.ConvertTime(DateTimeOffset.Now, _settings.CurrentValue.TimeZoneInfo);
+						
+			foreach (ScheduledDowntime schedule in globals.ScheduledDowntime)
+			{
+				DateTimeOffset start = TimeZoneInfo.ConvertTimeBySystemTimeZoneId(schedule.GetNext(now).StartTime, "UTC");
+				DateTimeOffset finish = TimeZoneInfo.ConvertTimeBySystemTimeZoneId(schedule.GetNext(now).FinishTime, "UTC");
+				messages.Add(new GetNoticeResponse() { ScheduledDowntime = true, StartTime = new DateTime(start.Ticks, DateTimeKind.Utc), FinishTime = new DateTime(finish.Ticks, DateTimeKind.Utc), Active = schedule.IsActive(now) });
 			}
 
 			// Add user notices
-			List<INotice> notices = await _noticeService.GetNoticesAsync();
+			List<INotice> notices = await _cachedNotices.GetCached();
 
 			for (int i = 0; i < notices.Count; i++)
 			{
@@ -179,12 +156,13 @@ namespace Horde.Build.Controllers
 				if (notice.UserId != null)
 				{
 					userInfo = new GetThinUserInfoResponse(await _userCollection.GetCachedUserAsync(notice.UserId));
-				}								
+				}
 
-				messages.Add(new GetNoticeResponse(notice, userInfo));
+				messages.Add(new GetNoticeResponse() { Id = notice.Id.ToString(), Active = true, Message = notice.Message, CreatedByUser = userInfo});
 			}
 
 			return messages;
+			
 		}
 	}
 }
diff --git a/Engine/Source/Programs/Horde/Horde.Build/Notices/INotice.cs b/Engine/Source/Programs/Horde/Horde.Build/Notices/INotice.cs
--- a/Engine/Source/Programs/Horde/Horde.Build/Notices/INotice.cs
+++ b/Engine/Source/Programs/Horde/Horde.Build/Notices/INotice.cs
@@ -15,12 +15,12 @@ namespace Horde.Build.Models
 	public interface INotice
 	{
 		/// <summary>
-		/// Unique id for this reservation
+		/// Unique id for this notice
 		/// </summary>
 		public ObjectId Id { get; }
 
 		/// <summary>
-		/// User id for who created this notice, null for system created notices
+		/// User id of who created this notice, null for system created notices
 		/// </summary>
 		[BsonIgnoreIfNull]
 		public UserId? UserId { get; }
diff --git a/Engine/Source/Programs/Horde/Horde.Build/Notices/Impl/NoticeCollection.cs b/Engine/Source/Programs/Horde/Horde.Build/Notices/Impl/NoticeCollection.cs
--- a/Engine/Source/Programs/Horde/Horde.Build/Notices/Impl/NoticeCollection.cs
+++ b/Engine/Source/Programs/Horde/Horde.Build/Notices/Impl/NoticeCollection.cs
@@ -1,15 +1,15 @@
 // Copyright Epic Games, Inc. All Rights Reserved.
 
+using MongoDB.Bson.Serialization.Attributes;
+using System.Threading.Tasks;
+using MongoDB.Driver;
 using System;
 using System.Collections.Generic;
 using System.Linq;
-using System.Threading.Tasks;
 using Horde.Build.Models;
 using Horde.Build.Server;
 using Horde.Build.Utilities;
 using MongoDB.Bson;
-using MongoDB.Bson.Serialization.Attributes;
-using MongoDB.Driver;
 
 namespace Horde.Build.Collections.Impl
 {
diff --git a/Engine/Source/Programs/Horde/HordeDashboard/src/backend/Api.ts b/Engine/Source/Programs/Horde/HordeDashboard/src/backend/Api.ts
--- a/Engine/Source/Programs/Horde/HordeDashboard/src/backend/Api.ts
+++ b/Engine/Source/Programs/Horde/HordeDashboard/src/backend/Api.ts
@@ -3859,13 +3859,77 @@ export type GetAgentSoftwareChannelResponse = {
 /// Jira issue information
 export type GetJiraIssueResponse = {
 
+	/** The jira key */
 	key: string;
+
+	/** Link to the issue on jira */
 	jiraLink: string;
+
+	/** The current status of the issue */
 	statusName?: string;
+
+	/** The current priority of the issue */
 	priorityName?: string;
+
+	/** The resolution of the issue */
 	resolutionName?: string;
+
+	/** Jira username of assignee */
 	assigneeName?: string;
+
+	/** Jira display name of assignee */
 	assigneeDisplayName?: string;
+
+	/** Email address of assignee */
 	assigneeEmailAddress?: string;
 
 }
+
+/// Create a notice which will display on the dashboard
+export type CreateNoticeRequest = {
+		
+	/**  Message to display	*/
+	message: string;	
+}
+
+/// Parameters required to update a notice
+export type UpdateNoticeRequest = {
+
+	/** The id of the notice to update */
+	id: string;
+
+	/** Start time to display this message */
+	startTime?: Date | string;
+
+	/** Finish time to display this message */
+	finishTime?: Date | String;
+
+	/** Message to display */
+	message?: string;
+}
+
+
+export type GetNoticeResponse = {
+
+	/** The id of the notice for user created notices */
+	id?: string;
+
+	/** Start time to display this message */
+	startTime?: Date | String;
+
+	/** Finish time to display this message */
+	finishTime?: Date | String;
+
+	/** Whether this notice is for scheduled downtime */
+	scheduledDowntime: boolean;
+
+	/** Whether the notice is currently active */
+	active: boolean;
+
+	/** Message to display */
+	message?: string;
+
+	/** User id who created the notice, otherwise null if a system message */
+	createdByUser?: GetThinUserInfoResponse;
+
+}
diff --git a/Engine/Source/Programs/Horde/HordeDashboard/src/backend/Backend.ts b/Engine/Source/Programs/Horde/HordeDashboard/src/backend/Backend.ts
--- a/Engine/Source/Programs/Horde/HordeDashboard/src/backend/Backend.ts
+++ b/Engine/Source/Programs/Horde/HordeDashboard/src/backend/Backend.ts
@@ -3,7 +3,7 @@
 import { configure } from "mobx";
 import { isNumber } from 'util';
 import templateCache from '../backend/TemplateCache';
-import { AgentData, ArtifactData, AuditLogEntry, AuditLogQuery, BatchUpdatePoolRequest, ChangeSummaryData, CreateDeviceRequest, CreateDeviceResponse, CreateJobRequest, CreateJobResponse, CreatePoolRequest, CreateSoftwareResponse, CreateSubscriptionRequest, CreateSubscriptionResponse, DashboardPreference, EventData, FindIssueResponse, FindJobTimingsResponse, GetAgentSoftwareChannelResponse, GetArtifactZipRequest, GetDevicePlatformResponse, GetDevicePoolResponse, GetDeviceReservationResponse, GetDeviceResponse, GetGraphResponse, GetIssueStreamResponse, GetJiraIssueResponse, GetJobsTabResponse, GetJobStepRefResponse, GetJobStepTraceResponse, GetJobTimingResponse, GetLogEventResponse, GetNotificationResponse, GetPerforceServerStatusResponse, GetServerInfoResponse, GetServerSettingsResponse, GetSoftwareResponse, GetSubscriptionResponse, GetUserResponse, GetUtilizationTelemetryResponse, GlobalConfig, IssueData, IssueQuery, IssueQueryV2, JobData, JobQuery, JobsTabColumnType, JobStepOutcome, JobStreamQuery, JobTimingsQuery, LeaseData, LogData, LogLineData, PoolData, ProjectData, ScheduleData, ScheduleQuery, SearchLogFileResponse, ServerUpdateResponse, SessionData, StreamData, TabType, TemplateData, TestData, UpdateAgentRequest, UpdateDeviceRequest, UpdateGlobalConfigRequest, UpdateIssueRequest, UpdateJobRequest, UpdateLeaseRequest, UpdateNotificationsRequest, UpdatePoolRequest, UpdateServerSettingsRequest, UpdateStepRequest, UpdateStepResponse, UpdateUserRequest, UsersQuery } from './Api';
+import { AgentData, ArtifactData, AuditLogEntry, AuditLogQuery, BatchUpdatePoolRequest, ChangeSummaryData, CreateDeviceRequest, CreateDeviceResponse, CreateJobRequest, CreateJobResponse, CreateNoticeRequest, CreatePoolRequest, CreateSoftwareResponse, CreateSubscriptionRequest, CreateSubscriptionResponse, DashboardPreference, EventData, FindIssueResponse, FindJobTimingsResponse, GetAgentSoftwareChannelResponse, GetArtifactZipRequest, GetDevicePlatformResponse, GetDevicePoolResponse, GetDeviceReservationResponse, GetDeviceResponse, GetGraphResponse, GetIssueStreamResponse, GetJiraIssueResponse, GetJobsTabResponse, GetJobStepRefResponse, GetJobStepTraceResponse, GetJobTimingResponse, GetLogEventResponse, GetNoticeResponse, GetNotificationResponse, GetPerforceServerStatusResponse, GetServerInfoResponse, GetServerSettingsResponse, GetSoftwareResponse, GetSubscriptionResponse, GetUserResponse, GetUtilizationTelemetryResponse, GlobalConfig, IssueData, IssueQuery, IssueQueryV2, JobData, JobQuery, JobsTabColumnType, JobStepOutcome, JobStreamQuery, JobTimingsQuery, LeaseData, LogData, LogLineData, PoolData, ProjectData, ScheduleData, ScheduleQuery, SearchLogFileResponse, ServerUpdateResponse, SessionData, StreamData, TabType, TemplateData, TestData, UpdateAgentRequest, UpdateDeviceRequest, UpdateGlobalConfigRequest, UpdateIssueRequest, UpdateJobRequest, UpdateLeaseRequest, UpdateNoticeRequest, UpdateNotificationsRequest, UpdatePoolRequest, UpdateServerSettingsRequest, UpdateStepRequest, UpdateStepResponse, UpdateUserRequest, UsersQuery } from './Api';
 import dashboard from './Dashboard';
 import { ChallengeStatus, Fetch } from './Fetch';
 import graphCache, { GraphQuery } from './GraphCache';
@@ -1418,6 +1418,53 @@ export class Backend {
         });
     }
 
+    // create a new notice
+    createNotice(request:CreateNoticeRequest): Promise<void> {
+        return new Promise<void>((resolve, reject) => {
+            this.backend.post(`/api/v1/notices`, request).then(() => {
+                resolve();
+            }).catch(reason => {
+                reject(reason);
+            });
+        });
+    }
+
+    // get all notices
+    getNotices(): Promise<GetNoticeResponse[]> {
+        return new Promise<GetNoticeResponse[]>((resolve, reject) => {
+            this.backend.get(`/api/v1/notices`).then((value) => {
+                resolve((value.data as GetNoticeResponse[]).map(notice => {
+                    notice.startTime = notice.startTime ? new Date(notice.startTime as string) : undefined;
+                    notice.finishTime = notice.finishTime ? new Date(notice.finishTime as string) : undefined;
+                    return notice;
+                }));
+            }).catch(reason => {
+                reject(reason);
+            });
+        });
+    }
+
+    // update a notice
+    updateNotice(request: UpdateNoticeRequest):Promise<void> {
+        return new Promise<void>((resolve, reject) => {
+            this.backend.put(`/api/v1/notices`, request).then(() => {
+                resolve();
+            }).catch(reason => {
+                reject(reason);
+            });
+        });
+    }
+
+    // delete a notice
+    deleteNotice(id:string):Promise<void> {
+        return new Promise<void>((resolve, reject) => {
+            this.backend.delete(`/api/v1/notices/${id}`).then(() => {
+                resolve();
+            }).catch(reason => {
+                reject(reason);
+            });
+        });
+    }
 
     private async update() {
 
diff --git a/Engine/Source/Programs/Horde/HordeDashboard/src/backend/Dashboard.ts b/Engine/Source/Programs/Horde/HordeDashboard/src/backend/Dashboard.ts
--- a/Engine/Source/Programs/Horde/HordeDashboard/src/backend/Dashboard.ts
+++ b/Engine/Source/Programs/Horde/HordeDashboard/src/backend/Dashboard.ts
@@ -1,9 +1,9 @@
 // Copyright Epic Games, Inc. All Rights Reserved.
 
+import { getTheme } from "@fluentui/react";
 import { action, observable } from 'mobx';
 import backend from '.';
-import { UserClaim, DashboardPreference, GetUserResponse } from './Api';
-import { getTheme } from "@fluentui/react";
+import { DashboardPreference, GetUserResponse, UserClaim } from './Api';
 
 const theme = getTheme();
 
@@ -366,7 +366,7 @@ export class Dashboard {
 
             if (this.updating) {
                 clearTimeout(this.updateTimeoutId);
-                this.updateTimeoutId = setTimeout(() => { this.update(); }, 4000);
+                this.updateTimeoutId = setTimeout(() => { this.update(); }, this.pollMS);
                 return;
             }
 
@@ -376,9 +376,11 @@ export class Dashboard {
 
                 const cancelId = this.cancelId++;
 
-                const response = await backend.getCurrentUser();
+                const responses = await Promise.all([backend.getCurrentUser()]);
+
+                const response = responses[0] as GetUserResponse;
 
-                // check for canceled during graph request
+                // check for canceled during request
                 if (!this.canceled.has(cancelId)) {
 
                     this.data = response;
@@ -415,7 +417,7 @@ export class Dashboard {
         } finally {
             this.updating = false;
             clearTimeout(this.updateTimeoutId);
-            this.updateTimeoutId = setTimeout(() => { this.update(); }, 4000);
+            this.updateTimeoutId = setTimeout(() => { this.update(); }, this.pollMS);
         }
     }
 
@@ -489,8 +491,7 @@ export class Dashboard {
         return success;
 
     }
-
-
+    
     requestLogout = false;
 
     serverSettingsChanged: boolean = false;
@@ -506,6 +507,8 @@ export class Dashboard {
     private canceled = new Set<number>();
     private cancelId = 0;
 
+    private pollMS = 4 * 1000;
+
 }
 
 const dashboard = new Dashboard();
diff --git a/Engine/Source/Programs/Horde/HordeDashboard/src/backend/Notices.ts b/Engine/Source/Programs/Horde/HordeDashboard/src/backend/Notices.ts
new file mode 100644
--- /dev/null
+++ b/Engine/Source/Programs/Horde/HordeDashboard/src/backend/Notices.ts
@@ -0,0 +1,83 @@
+import moment from "moment";
+import backend from ".";
+import { displayTimeZone } from "../base/utilities/timeUtils";
+import { GetNoticeResponse } from "./Api";
+import dashboard from "./Dashboard";
+import { PollBase } from "./PollBase";
+
+
+class Notices extends PollBase {
+
+    constructor(pollTime = 1000) {
+        super(pollTime);
+    }
+
+    clear() {
+        super.clear();
+    }
+
+    async poll(): Promise<void> {
+
+
+        try {
+
+            if (!dashboard.available) {
+                return;
+            }
+
+            this.notices = await backend.getNotices();
+
+            this.pollTime = 60 * 1000;
+
+            this.setUpdated();
+
+        } catch (err) {
+
+        }
+
+    }
+
+    get alertText(): string | undefined {
+        
+        // first one created by a user
+        let notice = this.notices.find(n => (n.active && !!n.createdByUser));
+        if (!notice) {
+            notice = this.notices.find(n => n.active);
+        }
+
+        if (!notice) {
+            return undefined;
+        }
+
+        if (notice.scheduledDowntime) {
+
+            console.log(notice.finishTime?.toString());
+
+            let time = "";
+            if (notice.finishTime) {
+                time = moment(notice.finishTime as Date).tz(displayTimeZone()).format(dashboard.display24HourClock ? "MMM Do, HH:mm z" : "MMM Do, h:mma z");
+                return `Horde is currently in scheduled downtime. Jobs will resume execution at ${time}`;
+            }
+                    
+            return `Horde is currently in scheduled downtime`
+        }
+
+        return notice?.message;
+
+    }
+
+
+    get allNotices(): GetNoticeResponse[] {
+
+        return this.notices ?? [];
+    
+    }
+
+    private notices: GetNoticeResponse[] = [];
+
+}
+
+const notices = new Notices();
+notices.start();
+
+export default notices;
diff --git a/Engine/Source/Programs/Horde/HordeDashboard/src/backend/PollBase.ts b/Engine/Source/Programs/Horde/HordeDashboard/src/backend/PollBase.ts
--- a/Engine/Source/Programs/Horde/HordeDashboard/src/backend/PollBase.ts
+++ b/Engine/Source/Programs/Horde/HordeDashboard/src/backend/PollBase.ts
@@ -87,7 +87,7 @@ export abstract class PollBase {
 
     private timeoutId: any;
 
-    private pollTime: number;
+    protected pollTime: number;
 
     private canceled = new Set<number>();
     private cancelId = 0;
