Horde: Remove some debug logging for notices
#rnx
#skipci

diff --git a/Engine/Source/Programs/Horde/HordeDashboard/src/backend/Notices.ts b/Engine/Source/Programs/Horde/HordeDashboard/src/backend/Notices.ts
--- a/Engine/Source/Programs/Horde/HordeDashboard/src/backend/Notices.ts
+++ b/Engine/Source/Programs/Horde/HordeDashboard/src/backend/Notices.ts
@@ -49,9 +49,7 @@ class Notices extends PollBase {
             return undefined;
         }
 
-        if (notice.scheduledDowntime) {
-
-            console.log(notice.finishTime?.toString());
+        if (notice.scheduledDowntime) {            
 
             let time = "";
             if (notice.finishTime) {
