Horde: Add more logging to catch bad stream config lookups

diff --git a/Engine/Source/Programs/Horde/Horde.Build/Projects/Impl/StreamCollection.cs b/Engine/Source/Programs/Horde/Horde.Build/Projects/Impl/StreamCollection.cs
--- a/Engine/Source/Programs/Horde/Horde.Build/Projects/Impl/StreamCollection.cs
+++ b/Engine/Source/Programs/Horde/Horde.Build/Projects/Impl/StreamCollection.cs
@@ -13,6 +13,7 @@ using Horde.Build.Models;
 using Horde.Build.Server;
 using Horde.Build.Utilities;
 using HordeCommon;
+using Microsoft.Extensions.Logging;
 using MongoDB.Bson.Serialization.Attributes;
 using MongoDB.Driver;
 
@@ -111,6 +112,7 @@ namespace Horde.Build.Collections.Impl
 		readonly ConfigCollection _configCollection;
 		readonly IClock _clock;
 		readonly ITemplateCollection _templateCollection;
+		readonly ILogger<StreamCollection> _logger;
 
 		/// <summary>
 		/// Constructor
@@ -119,17 +121,27 @@ namespace Horde.Build.Collections.Impl
 		/// <param name="configCollection"></param>
 		/// <param name="clock"></param>
 		/// <param name="templateCollection"></param>
-		public StreamCollection(MongoService mongoService, ConfigCollection configCollection, IClock clock, ITemplateCollection templateCollection)
+		/// <param name="logger"></param>
+		public StreamCollection(MongoService mongoService, ConfigCollection configCollection, IClock clock, ITemplateCollection templateCollection, ILogger<StreamCollection> logger)
 		{
 			_streams = mongoService.GetCollection<StreamDocument>("Streams");
 			_configCollection = configCollection;
 			_clock = clock;
 			_templateCollection = templateCollection;
+			_logger = logger;
 		}
 
 		async Task PostLoadAsync(StreamDocument stream)
 		{
-			stream.Config = await _configCollection.GetConfigAsync<StreamConfig>(stream.ConfigRevision);
+			try
+			{
+				stream.Config = await _configCollection.GetConfigAsync<StreamConfig>(stream.ConfigRevision);
+			}
+			catch (Exception)
+			{
+				_logger.LogError("Unable to get stream config for {streamId}/{streamName} at {revision}", stream.Id, stream.Name, stream.ConfigRevision);
+				throw;
+			}
 		}
 
 		/// <inheritdoc/>
