Horde: Guard issue endpoint for stream config errors
#rnx
#skipci

diff --git a/Engine/Source/Programs/Horde/Horde.Build/Controllers/IssuesController.cs b/Engine/Source/Programs/Horde/Horde.Build/Controllers/IssuesController.cs
--- a/Engine/Source/Programs/Horde/Horde.Build/Controllers/IssuesController.cs
+++ b/Engine/Source/Programs/Horde/Horde.Build/Controllers/IssuesController.cs
@@ -16,6 +16,7 @@ using HordeCommon;
 using Microsoft.AspNetCore.Authorization;
 using Microsoft.AspNetCore.Mvc;
 using MongoDB.Bson;
+using Microsoft.Extensions.Logging;
 
 namespace Horde.Build.Controllers
 {
@@ -33,26 +34,26 @@ namespace Horde.Build.Controllers
 	[Route("[controller]")]
 	public class IssuesController : HordeControllerBase
 	{
-		private readonly ConfigCollection _configCollection;
 		private readonly IIssueCollection _issueCollection;
 		private readonly IIssueService _issueService;
 		private readonly JobService _jobService;
 		private readonly StreamService _streamService;
 		private readonly IUserCollection _userCollection;
 		private readonly ILogFileService _logFileService;
+		private readonly ILogger<IssuesController> _logger;
 
 		/// <summary>
 		/// Constructor
 		/// </summary>
-		public IssuesController(ConfigCollection configCollection, IIssueCollection issueCollection, IIssueService issueService, JobService jobService, StreamService streamService, IUserCollection userCollection, ILogFileService logFileService)
-		{
-			_configCollection = configCollection;
+		public IssuesController(ILogger<IssuesController> logger, IIssueCollection issueCollection, IIssueService issueService, JobService jobService, StreamService streamService, IUserCollection userCollection, ILogFileService logFileService)
+		{			
 			_issueCollection = issueCollection;
 			_issueService = issueService;
 			_jobService = jobService;
 			_streamService = streamService;
 			_userCollection = userCollection;
 			_logFileService = logFileService;
+			_logger = logger;
 		}
 
 		/// <summary>
@@ -307,8 +308,15 @@ namespace Horde.Build.Controllers
 			List<GetIssueAffectedStreamResponse> affectedStreams = new List<GetIssueAffectedStreamResponse>();
 			foreach (IGrouping<StreamId, IIssueSpan> streamSpans in details.Spans.GroupBy(x => x.StreamId))
 			{
-				IStream? stream = await _streamService.GetCachedStream(streamSpans.Key);
-				affectedStreams.Add(new GetIssueAffectedStreamResponse(details, stream, streamSpans));
+				try
+				{
+					IStream? stream = await _streamService.GetCachedStream(streamSpans.Key);
+					affectedStreams.Add(new GetIssueAffectedStreamResponse(details, stream, streamSpans));
+				}
+				catch 
+				{
+					_logger.LogError("Unable to get {StreamId} for span key", streamSpans.Key);
+				}
 			}
 			return new GetIssueResponse(details, affectedStreams, showDesktopAlerts);
 		}
