Horde: Back out regex that takes too long on 2 million characters
#skipci
#rnx

diff --git a/Engine/Source/Programs/Horde/Horde.Agent.Tests/LogParserTests.cs b/Engine/Source/Programs/Horde/Horde.Agent.Tests/LogParserTests.cs
--- a/Engine/Source/Programs/Horde/Horde.Agent.Tests/LogParserTests.cs
+++ b/Engine/Source/Programs/Horde/Horde.Agent.Tests/LogParserTests.cs
@@ -485,6 +485,7 @@ namespace Horde.Agent.Tests
 				Assert.AreEqual(LogLevel.Error, logEvents[0].Level);
 			}
 
+			/* Disabled until have faster regex, can operate on 2 million character lines :/
 			{
 				string[] lines =
 				{
@@ -506,6 +507,7 @@ namespace Horde.Agent.Tests
 				Assert.AreEqual("UProjectPlayControllerComponent::HandleWeaponEquipped(ATestGameWeapon*, ATestGameWeapon*)", logEvents[3].Properties!["symbol"].ToString());
 
 			}
+			*/
 		}
 
 		[TestMethod]
diff --git a/Engine/Source/Programs/Horde/Horde.Agent/Parser/Matchers/LinkEventMatcher.cs b/Engine/Source/Programs/Horde/Horde.Agent/Parser/Matchers/LinkEventMatcher.cs
--- a/Engine/Source/Programs/Horde/Horde.Agent/Parser/Matchers/LinkEventMatcher.cs
+++ b/Engine/Source/Programs/Horde/Horde.Agent/Parser/Matchers/LinkEventMatcher.cs
@@ -18,7 +18,7 @@ namespace Horde.Agent.Parser.Matchers
 			int lineCount = 0;
 			bool isError = false;
 			bool isWarning = false;
-			while (cursor.IsMatch(lineCount, @": undefined reference to |undefined symbol|^\s*(ld|ld.lld):|.*clang:.*linker command failed|^[^:]*[^a-zA-Z][a-z]?ld: |^\s*>>>"))
+			while (cursor.IsMatch(lineCount, @": undefined reference to |undefined symbol|^\s*(ld|ld.lld):|^[^:]*[^a-zA-Z][a-z]?ld: |^\s*>>>"))
 			{
 				isError |= cursor.IsMatch("error:");
 				isWarning |= cursor.IsMatch("warning:");
