Horde: Add initial support for AWS Lambda execution

Includes the AWS side of things - the process that will get invoked inside the Lambda, which in turn runs the ComputeTask executor.

diff --git a/Engine/Source/Programs/Horde/Horde.Agent/Commands/Compute/AwsLambdaClient.cs b/Engine/Source/Programs/Horde/Horde.Agent/Commands/Compute/AwsLambdaClient.cs
new file mode 100644
--- /dev/null
+++ b/Engine/Source/Programs/Horde/Horde.Agent/Commands/Compute/AwsLambdaClient.cs
@@ -0,0 +1,208 @@
+// Copyright Epic Games, Inc. All Rights Reserved.
+
+using System;
+using System.Collections.Generic;
+using System.Linq;
+using System.Net;
+using System.Net.Http;
+using System.Text.Json;
+using System.Threading;
+using System.Threading.Tasks;
+using Microsoft.Extensions.Logging;
+
+namespace Horde.Agent.Commands.Compute
+{
+	/// <summary>
+	/// Class representing the HTTP response from next invocation request
+	/// </summary>
+	public class NextInvocationResponse
+	{
+		/// <summary>
+		/// Lambda request ID
+		/// </summary>
+		public readonly string requestId;
+		
+		/// <summary>
+		/// Name of the Lambda function being invoked
+		/// </summary>
+		public readonly string invokedFunctionArn;
+		
+		/// <summary>
+		/// Deadline in milliseconds for completing the function invocation
+		/// </summary>
+		public readonly long deadlineMs;
+		
+		/// <summary>
+		/// Data
+		/// </summary>
+		public readonly byte[] data;
+
+		/// <summary>
+		/// Constructor
+		/// </summary>
+		/// <param name="requestId"></param>
+		/// <param name="invokedFunctionArn"></param>
+		/// <param name="deadlineMs"></param>
+		/// <param name="data"></param>
+		public NextInvocationResponse(string requestId, string invokedFunctionArn, long deadlineMs, byte[] data)
+		{
+			this.requestId = requestId;
+			this.invokedFunctionArn = invokedFunctionArn;
+			this.deadlineMs = deadlineMs;
+			this.data = data;
+		}
+	}
+
+	/// <summary>
+	/// Class representing an exception when communicating with the AWS Lambda Runtime API
+	/// </summary>
+	public class AwsLambdaClientException : Exception
+	{
+		/// <summary>
+		/// Whether the exception is fatal and requires termination of the process (to adhere to Lambda specs)
+		/// </summary>
+		public readonly bool isFatal;
+		
+		/// <summary>
+		/// HTTP status code in response from AWS Lambda API
+		/// </summary>
+		public readonly HttpStatusCode? statusCode;
+
+		/// <summary>
+		/// Constructor
+		/// </summary>
+		/// <param name="message"></param>
+		/// <param name="statusCode"></param>
+		/// <param name="isFatal"></param>
+		public AwsLambdaClientException(string message, bool isFatal = false, HttpStatusCode? statusCode = null) : base(message)
+		{
+			this.isFatal = isFatal;
+			this.statusCode = statusCode;
+		}
+	}
+
+	/// <summary>
+	/// Client interfacing with the AWS Lambda runtime API
+	/// </summary>
+	class AwsLambdaClient
+	{
+		private const string EnvVarRuntimeApi = "AWS_LAMBDA_RUNTIME_API";
+		private const string HeaderNameAwsRequestId = "Lambda-Runtime-Aws-Request-Id";
+		private const string HeaderNameDeadlineMs = "Lambda-Runtime-Deadline-Ms";
+		private const string HeaderNameInvokedFunctionArn = "Lambda-Runtime-Invoked-Function-Arn";
+		private const string HeaderNameFuncErrorType = "Lambda-Runtime-Function-Error-Type";
+
+		private readonly ILogger<AwsLambdaClient> _logger;
+		private readonly string _hostPort;
+
+		public AwsLambdaClient(string hostPort, ILogger<AwsLambdaClient> logger)
+		{
+			_hostPort = hostPort;
+			_logger = logger;
+		}
+
+		public static AwsLambdaClient InitFromEnv(ILogger<AwsLambdaClient> logger)
+		{
+			string? runtimeApi = Environment.GetEnvironmentVariable(EnvVarRuntimeApi);
+			if (runtimeApi == null)
+			{
+				throw new ArgumentException($"Env var {EnvVarRuntimeApi} is not set");
+			}
+
+			return new AwsLambdaClient(runtimeApi, logger);
+		}
+		
+		private string GetNextInvocationUrl() => $"http://{_hostPort}/2018-06-01/runtime/invocation/next";
+		private string GetInvocationResponseUrl(string awsRequestId) => $"http://{_hostPort}/2018-06-01/runtime/invocation/{awsRequestId}/response";
+		private string GetInitErrorUrl() => $"http://{_hostPort}/runtime/init/error";
+		private string GetInvocationErrorUrl(string awsRequestId) => $"http://{_hostPort}/2018-06-01/runtime/invocation/{awsRequestId}/error";
+
+		private HttpClient GetHttpClient()
+		{
+			return new HttpClient();
+		}
+		
+		public async Task<NextInvocationResponse> GetNextInvocationAsync(CancellationToken cancellationToken)
+		{
+			using HttpClient client = GetHttpClient();
+			client.Timeout = TimeSpan.FromHours(1); // Waiting for the next invocation can block for a while
+			using HttpRequestMessage request = new HttpRequestMessage(HttpMethod.Get, GetNextInvocationUrl());
+
+			HttpResponseMessage response = await client.SendAsync(request, cancellationToken);
+
+			if (response.StatusCode != HttpStatusCode.OK)
+			{
+				throw new AwsLambdaClientException("Failed getting next invocation", false, response.StatusCode);
+			}
+			
+			if (!response.Headers.Contains(HeaderNameAwsRequestId))
+			{
+				throw new AwsLambdaClientException($"Missing response header {HeaderNameAwsRequestId}");
+			}
+			
+			if (!response.Headers.Contains(HeaderNameDeadlineMs))
+			{
+				throw new AwsLambdaClientException($"Missing response header {HeaderNameDeadlineMs}");
+			}
+			
+			if (!response.Headers.Contains(HeaderNameInvokedFunctionArn))
+			{
+				throw new AwsLambdaClientException($"Missing response header {HeaderNameInvokedFunctionArn}");
+			}
+
+			string requestId = response.Headers.GetValues(HeaderNameAwsRequestId).First();
+			long deadlineMs = Convert.ToInt64(response.Headers.GetValues(HeaderNameDeadlineMs).First());
+			string invokedFunctionArn = response.Headers.GetValues(HeaderNameInvokedFunctionArn).First();
+
+			byte[] data = await response.Content.ReadAsByteArrayAsync();
+			return new NextInvocationResponse(requestId, invokedFunctionArn, deadlineMs, data);
+		}
+		
+		public async Task SendInvocationResponseAsync(string awsRequestId, ReadOnlyMemory<byte> data)
+		{
+			using HttpClient client = GetHttpClient();
+			client.Timeout = TimeSpan.FromSeconds(30);
+			using HttpRequestMessage request = new HttpRequestMessage(HttpMethod.Post, GetInvocationResponseUrl(awsRequestId));
+			request.Content = new ByteArrayContent(data.ToArray());
+			
+			HttpResponseMessage response = await client.SendAsync(request);
+			if (response.StatusCode != HttpStatusCode.Accepted)
+			{
+				_logger.LogError("Failed sending invocation response! RequestId={awsRequestId} DataLen={dataLen}", awsRequestId, data.Length);
+				throw new AwsLambdaClientException("Failed sending invocation response!");
+			}
+		}
+		
+		public Task SendInitErrorAsync(string errorType, string errorMessage, List<string>? stackTrace = null)
+		{
+			return SendErrorAsync(GetInitErrorUrl(), errorType, errorMessage, stackTrace);
+		}
+		
+		public Task SendInvocationErrorAsync(string awsRequestId, string errorType, string errorMessage, List<string>? stackTrace = null)
+		{
+			return SendErrorAsync(GetInvocationErrorUrl(awsRequestId), errorType, errorMessage, stackTrace);
+		}
+		
+		public async Task SendErrorAsync(string url, string errorType, string errorMessage, List<string>? stackTrace)
+		{
+			using HttpClient client = new HttpClient();
+			client.Timeout = TimeSpan.FromHours(1);
+			using HttpRequestMessage request = new HttpRequestMessage(HttpMethod.Post, url);
+			request.Headers.Add(HeaderNameFuncErrorType, errorType);
+			string bodyStr = JsonSerializer.Serialize(new
+			{
+				errorMessage = errorMessage,
+				errorType = errorType,
+				stackTrace = stackTrace != null ? stackTrace.ToArray() : Array.Empty<string>()
+			});
+			request.Content = new StringContent(bodyStr);
+			
+			HttpResponseMessage response = await client.SendAsync(request);
+			if (response.StatusCode != HttpStatusCode.Accepted)
+			{
+				bool isFatal = response.StatusCode == HttpStatusCode.InternalServerError;
+				throw new AwsLambdaClientException($"Failed sending initialization error! Status code {response.StatusCode}", isFatal, response.StatusCode);
+			}
+		}
+	}
+}
diff --git a/Engine/Source/Programs/Horde/Horde.Agent/Commands/Compute/AwsLambdaCommand.cs b/Engine/Source/Programs/Horde/Horde.Agent/Commands/Compute/AwsLambdaCommand.cs
new file mode 100644
--- /dev/null
+++ b/Engine/Source/Programs/Horde/Horde.Agent/Commands/Compute/AwsLambdaCommand.cs
@@ -0,0 +1,252 @@
+// Copyright Epic Games, Inc. All Rights Reserved.
+
+using System;
+using System.IO;
+using System.Net.Http;
+using System.Text;
+using System.Threading;
+using System.Threading.Tasks;
+using Amazon.SimpleSystemsManagement;
+using Amazon.SimpleSystemsManagement.Model;
+using EpicGames.Core;
+using EpicGames.Horde.Auth;
+using EpicGames.Horde.Storage;
+using EpicGames.Horde.Storage.Impl;
+using Google.Protobuf;
+using Horde.Agent.Commands.Compute;
+using Horde.Agent.Execution;
+using HordeCommon.Rpc.Tasks;
+using Microsoft.Extensions.Logging;
+using Command = EpicGames.Core.Command;
+
+namespace Horde.Agent.Commands
+{
+	/// <summary>
+	/// Listens continuously for new Lambda invocations and executes them
+	/// </summary>
+	class AwsLambdaRunner
+	{
+		// TODO: Rename to AwsLambdaInvoker or Executor?
+		
+		private readonly AwsLambdaClient _lambdaClient;
+		private readonly IStorageClient _storageClient;
+		private readonly ILogger<AwsLambdaRunner> _logger;
+		private readonly ComputeTaskExecutor _executor;
+
+		public AwsLambdaRunner(AwsLambdaClient lambdaClient, IStorageClient storageClient, ILogger<AwsLambdaRunner> logger)
+		{
+			_lambdaClient = lambdaClient;
+			_storageClient = storageClient;
+			_logger = logger;
+			_executor = new ComputeTaskExecutor(_storageClient, _logger);
+		}
+
+		public async Task<int> ListenForInvocationsAsync(CancellationToken cancellationToken)
+		{
+			while (!cancellationToken.IsCancellationRequested)
+			{
+				NextInvocationResponse? nextInvocationResponse = null;
+				try
+				{
+					_logger.LogDebug("Waiting for next invocation...");
+
+					// Block and wait for the next invocation (can be many minutes)
+					nextInvocationResponse = await _lambdaClient.GetNextInvocationAsync(cancellationToken);
+
+					AwsLambdaComputeTaskMessage computeTaskMessage;
+					// Detect if protobuf is binary or text-encoded JSON. The JSON format is used during testing/debugging.
+					bool isProtobufCompatibleJson = nextInvocationResponse.data.Length >= 1 && nextInvocationResponse.data[0] == '{';
+					if (isProtobufCompatibleJson)
+					{
+						computeTaskMessage = AwsLambdaComputeTaskMessage.Parser.ParseJson(Encoding.UTF8.GetString(nextInvocationResponse.data));
+					}
+					else
+					{
+						computeTaskMessage = AwsLambdaComputeTaskMessage.Parser.ParseFrom(nextInvocationResponse.data);
+					}
+
+					_logger.LogDebug("Invocation received. Lambda request ID {requestId}", nextInvocationResponse.requestId);
+
+					// Lease ID is not relevant when executing under AWS Lambda, but one is generated to satisfy the ComputeTaskExecutor.
+					string leaseId = "aws-lambda-" + Guid.NewGuid();
+					string workingDir = Path.GetTempPath();
+					
+					DirectoryReference leaseDir = DirectoryReference.Combine(new DirectoryReference(workingDir), "Compute", leaseId);
+					DirectoryReference.CreateDirectory(leaseDir);
+
+					ComputeTaskMessage task = computeTaskMessage.Task;
+					string msg = "clusterId={clusterId}\n" +
+					             "channelId={channelId}\n" +
+					             "NamespaceId={NamespaceId}\n" +
+					             "DispatchedMs={DispatchedMs}\n" +
+					             "InputBucketId={InputBucketId}\n" +
+					             "OutputBucketId={OutputBucketId}\n" +
+					             "RequirementsHash={RequirementsHash}\n" +
+					             "TaskRefId={TaskRefId}\n" +
+					             "QueuedAt={QueuedAt}\n";
+					_logger.LogDebug(msg, task.ClusterId, task.ChannelId, task.NamespaceId,
+						task.DispatchedMs, task.InputBucketId, task.OutputBucketId, task.RequirementsHash,
+						task.TaskRefId, task.QueuedAt);
+					
+					_logger.LogDebug("Lease dir {leaseDir}", leaseDir.ToString());
+					ComputeTaskResultMessage result;
+					ComputeTaskExecutor executor = new ComputeTaskExecutor(_storageClient, _logger);
+					try
+					{
+						result = await executor.ExecuteAsync(leaseId, computeTaskMessage.Task, leaseDir, cancellationToken);
+					}
+					finally
+					{
+						try
+						{
+							DirectoryReference.Delete(leaseDir, true);
+						}
+						catch
+						{
+							// Ignore any exceptions
+						}
+					}
+					
+					_logger.LogDebug("Execution completed. outcome={outcome} resultRefId={resultRefId} detail={detail}", result.Outcome, result.ResultRefId, result.Detail);
+
+					AwsLambdaComputeTaskResultMessage taskResultMessage = new AwsLambdaComputeTaskResultMessage { TaskResult = result };
+					byte[] responseData;
+					if (isProtobufCompatibleJson)
+					{
+						string jsonStr = JsonFormatter.Default.Format(taskResultMessage);
+						responseData = Encoding.UTF8.GetBytes(jsonStr);
+					}
+					else
+					{
+						responseData = taskResultMessage.ToByteArray();
+					}
+
+					await _lambdaClient.SendInvocationResponseAsync(nextInvocationResponse.requestId, responseData);
+				}
+				catch (Exception e)
+				{
+					_logger.LogError(e, "Exception invoking Lambda function from request");
+					
+					if (nextInvocationResponse != null)
+					{
+						try
+						{
+							await _lambdaClient.SendInvocationErrorAsync(nextInvocationResponse.requestId, "general", e.Message);
+						}
+						catch (AwsLambdaClientException sendErrorException)
+						{
+							_logger.LogError(sendErrorException, "Bad response when sending invocation error. isFatal={isFatal}", sendErrorException.isFatal);
+							return 1;
+						}
+					}
+					
+					await Task.Delay(500, cancellationToken); // Cool down before trying another next invocation call
+				}
+			}
+
+			return 0;
+		}
+	}
+	
+	/// <summary>
+	/// Run agent inside an AWS Lambda function, wrapping the execution of an arbitrary task
+	/// </summary>
+	[Command("Compute", "AwsLambda", "Listen for AWS Lambda invocations")]
+	class AwsLambdaCommand : Command
+	{
+		/// <inheritdoc/>
+		public override async Task<int> ExecuteAsync(ILogger logger)
+		{
+			LogLevel logLevel;
+			// TODO: Add tracing
+
+			string logLevelStr = Environment.GetEnvironmentVariable("UE_HORDE_LOG_LEVEL") ?? "debug";
+
+			if (Enum.TryParse(logLevelStr, true, out LogLevel logEventLevel))
+			{
+				logLevel = logEventLevel;
+			}
+			else
+			{
+				logger.LogError($"Unable to parse log level: {logLevelStr}");
+				return 1;
+			}
+
+			ILoggerFactory loggerFactory = LoggerFactory.Create(builder => builder.AddConsole().SetMinimumLevel(logLevel));
+			ILogger<AwsLambdaClient> lambdaClientLogger = loggerFactory.CreateLogger<AwsLambdaClient>();
+			AwsLambdaClient lambdaClient = AwsLambdaClient.InitFromEnv(lambdaClientLogger);
+			AwsLambdaRunner lambdaRunner;
+			CancellationTokenSource cts = new CancellationTokenSource();
+			
+			try
+			{
+				logger.LogInformation("Initializing AWS Lambda executor...");
+
+				string oAuthClientSecret = await GetOAuthClientSecretAsync(cts.Token);
+				HttpServiceClientOptions clientOptions = CreateHttpServiceClientOptionsFromEnv(oAuthClientSecret);
+				HttpClient httpClient = new HttpClient();
+				OAuthHandlerFactory oAuthHandlerFactory = new OAuthHandlerFactory(httpClient);
+				OAuthHandler<HttpStorageClient> oAuthHandler = oAuthHandlerFactory.Create<HttpStorageClient>(clientOptions);
+				HttpClient client = new HttpClient(oAuthHandler);
+				client.BaseAddress = new Uri(clientOptions.Url);
+				IStorageClient storageClientManual = new HttpStorageClient(client);
+
+				AppDomain.CurrentDomain.ProcessExit += (s, e) => 
+				{
+					logger.LogInformation("ProcessExit triggered. Shutting down...");
+					cts.Cancel();
+				};
+
+				ILogger<AwsLambdaRunner> lambdaRunnerLogger = loggerFactory.CreateLogger<AwsLambdaRunner>();
+				lambdaRunner = new AwsLambdaRunner(lambdaClient, storageClientManual, lambdaRunnerLogger);
+			}
+			catch (Exception e)
+			{
+				logger.LogError(e, "Error initializing");
+				try
+				{
+					await lambdaClient.SendInitErrorAsync("general", e.Message);
+				}
+				catch (AwsLambdaClientException initException)
+				{
+					logger.LogError(initException, "Bad response when sending init error. isFatal={isFatal}", initException.isFatal);
+				}
+				
+				return 1;
+			}
+			
+			return await lambdaRunner.ListenForInvocationsAsync(cts.Token);
+		}
+
+		private static async Task<string> GetOAuthClientSecretAsync(CancellationToken cancellationToken)
+		{
+			string oAuthClientSecret = GetEnvVar("UE_HORDE_STORAGE_OAUTH_CLIENT_SECRET_ARN");
+			AmazonSimpleSystemsManagementClient ssmClient = new AmazonSimpleSystemsManagementClient();
+			GetParameterResponse paramRes = await ssmClient.GetParameterAsync(new GetParameterRequest { Name = oAuthClientSecret }, cancellationToken);
+			return paramRes.Parameter.Value;
+		}
+
+		private static HttpServiceClientOptions CreateHttpServiceClientOptionsFromEnv(string clientSecret)
+		{
+			return new HttpServiceClientOptions
+			{
+				Url = GetEnvVar("UE_HORDE_STORAGE_URL"),
+				AuthUrl = GetEnvVar("UE_HORDE_STORAGE_OAUTH_URL"),
+				GrantType = GetEnvVar("UE_HORDE_STORAGE_OAUTH_GRANT_TYPE"),
+				ClientId = GetEnvVar("UE_HORDE_STORAGE_OAUTH_CLIENT_ID"),
+				ClientSecret = clientSecret,
+				Scope = GetEnvVar("UE_HORDE_STORAGE_OAUTH_SCOPE"),
+			};
+		}
+
+		private static string GetEnvVar(string name)
+		{
+			string? value = Environment.GetEnvironmentVariable(name);
+			if (value == null)
+			{
+				throw new ArgumentException($"Missing env var {name}");
+			}
+			return value;
+		}
+	}
+}
diff --git a/Engine/Source/Programs/Horde/Horde.Agent/Horde.Agent.csproj b/Engine/Source/Programs/Horde/Horde.Agent/Horde.Agent.csproj
--- a/Engine/Source/Programs/Horde/Horde.Agent/Horde.Agent.csproj
+++ b/Engine/Source/Programs/Horde/Horde.Agent/Horde.Agent.csproj
@@ -22,8 +22,9 @@
   </ItemGroup>
 
   <ItemGroup>
-    <PackageReference Include="AWSSDK.Core" Version="3.5.1.49" />
-    <PackageReference Include="AWSSDK.EC2" Version="3.5.1.2" />
+    <PackageReference Include="AWSSDK.Core" Version="3.7.3.19" />
+    <PackageReference Include="AWSSDK.EC2" Version="3.7.4" />
+    <PackageReference Include="AWSSDK.SimpleSystemsManagement" Version="3.7.3.19" />
     <PackageReference Include="Datadog.Trace.OpenTracing" Version="1.18.0" />
     <PackageReference Include="Google.Protobuf" Version="3.20.1" />
     <PackageReference Include="Grpc.Net.Client" Version="2.45.0" />
@@ -120,6 +121,7 @@
     <Protobuf Include="..\Horde.Build\Protos\horde\messages\agent_requirements.proto" ProtoRoot="..\Horde.Build\Protos" Access="Internal" />
     <Protobuf Include="..\Horde.Build\Protos\horde\messages\agent_workspace.proto" ProtoRoot="..\Horde.Build\Protos" Access="Internal" />
     <Protobuf Include="..\Horde.Build\Protos\horde\messages\lease.proto" ProtoRoot="..\Horde.Build\Protos" Access="Internal" />
+    <Protobuf Include="..\Horde.Build\Protos\horde\tasks\aws_lambda_compute_task.proto" ProtoRoot="..\Horde.Build\Protos" Access="Internal" />
     <Protobuf Include="..\Horde.Build\Protos\horde\tasks\compute_task.proto" ProtoRoot="..\Horde.Build\Protos" Access="Internal" />
     <Protobuf Include="..\Horde.Build\Protos\horde\tasks\conform_task.proto" ProtoRoot="..\Horde.Build\Protos" Access="Internal" />
     <Protobuf Include="..\Horde.Build\Protos\horde\tasks\job_task.proto" ProtoRoot="..\Horde.Build\Protos" Access="Internal" />
