Horde: Fix linker event matcher with faster regex
#rnx

diff --git a/Engine/Source/Programs/Horde/Horde.Agent.Tests/LogParserTests.cs b/Engine/Source/Programs/Horde/Horde.Agent.Tests/LogParserTests.cs
--- a/Engine/Source/Programs/Horde/Horde.Agent.Tests/LogParserTests.cs
+++ b/Engine/Source/Programs/Horde/Horde.Agent.Tests/LogParserTests.cs
@@ -485,8 +485,8 @@ namespace Horde.Agent.Tests
 				Assert.AreEqual(LogLevel.Error, logEvents[0].Level);
 			}
 
-			/* Disabled until have faster regex, can operate on 2 million character lines :/
 			{
+				
 				string[] lines =
 				{
 					@"tool 19.50.0.12 (rel,tool,19.500 @527452 x64) D:\Workspaces\AutoSDK\HostWin64\9.508.001\9.500\host_tools\bin\tool.exe",
@@ -507,7 +507,6 @@ namespace Horde.Agent.Tests
 				Assert.AreEqual("UProjectPlayControllerComponent::HandleWeaponEquipped(ATestGameWeapon*, ATestGameWeapon*)", logEvents[3].Properties!["symbol"].ToString());
 
 			}
-			*/
 		}
 
 		[TestMethod]
diff --git a/Engine/Source/Programs/Horde/Horde.Agent/Parser/Matchers/LinkEventMatcher.cs b/Engine/Source/Programs/Horde/Horde.Agent/Parser/Matchers/LinkEventMatcher.cs
--- a/Engine/Source/Programs/Horde/Horde.Agent/Parser/Matchers/LinkEventMatcher.cs
+++ b/Engine/Source/Programs/Horde/Horde.Agent/Parser/Matchers/LinkEventMatcher.cs
@@ -18,7 +18,7 @@ namespace Horde.Agent.Parser.Matchers
 			int lineCount = 0;
 			bool isError = false;
 			bool isWarning = false;
-			while (cursor.IsMatch(lineCount, @": undefined reference to |undefined symbol|^\s*(ld|ld.lld):|^[^:]*[^a-zA-Z][a-z]?ld: |^\s*>>>"))
+			while (cursor.IsMatch(lineCount, @": undefined reference to |undefined symbol|linker command failed|^\s*(ld|ld.lld):|^[^:]*[^a-zA-Z][a-z]?ld: |^\s*>>>"))
 			{
 				isError |= cursor.IsMatch("error:");
 				isWarning |= cursor.IsMatch("warning:");
