Horde Storage - Updated editorconfig to more closely match what we run for Horde and updated code to fix all static analysis errors

diff --git a/Engine/Source/Programs/Horde/HordeStorage/.editorconfig b/Engine/Source/Programs/Horde/HordeStorage/.editorconfig
--- a/Engine/Source/Programs/Horde/HordeStorage/.editorconfig
+++ b/Engine/Source/Programs/Horde/HordeStorage/.editorconfig
@@ -1,237 +1,29 @@
-
-[*.{asax,ascx,aspx,cs,cshtml,css,htm,html,js,jsx,master,razor,skin,ts,tsx,vb,xaml,xamlx,xoml}]
-indent_style=space
-indent_size=4
-tab_width=4
-insert_final_newline=true
-
-[*.{appxmanifest,build,config,csproj,dbml,discomap,dtd,json,jsproj,lsproj,njsproj,nuspec,proj,props,resjson,resw,resx,StyleCop,targets,tasks,vbproj,xml,xsd}]
-indent_style=space
-indent_size=2
-tab_width=2
-insert_final_newline=true
-
-[*]
-
-# Microsoft .NET properties
-csharp_new_line_before_members_in_object_initializers=false
-csharp_preferred_modifier_order=public, private, protected, internal, new, abstract, virtual, sealed, override, static, readonly, extern, unsafe, volatile, async:suggestion
-csharp_style_var_elsewhere=false:hint
-csharp_style_var_for_built_in_types=false:hint
-csharp_style_var_when_type_is_apparent=false:hint
-dotnet_style_predefined_type_for_locals_parameters_members=true:hint
-dotnet_style_predefined_type_for_member_access=true:hint
-dotnet_style_qualification_for_event=false:hint
-dotnet_style_qualification_for_field=false:hint
-dotnet_style_qualification_for_method=false:hint
-dotnet_style_qualification_for_property=false:hint
-dotnet_style_require_accessibility_modifiers=for_non_interface_members:hint
+[*.cs]
 
 # ReSharper properties
-resharper_accessor_owner_body=accessors_with_block_body
-resharper_arguments_other=named
-resharper_arguments_skip_single=true
-resharper_braces_for_for=required
-resharper_braces_for_foreach=required
-resharper_braces_for_ifelse=required
-resharper_braces_for_while=required
-
-# ReSharper inspection severities
-resharper_check_namespace_highlighting=hint
-resharper_redundant_base_qualifier_highlighting=warning
-
-# Override inherited settings
-
-# CA1309: Use ordinal string comparison
-dotnet_diagnostic.CA1309.severity = none
-
-# CA1823: Avoid unused private fields
-dotnet_diagnostic.CA1823.severity = none
-
-# CA1052: Static holder types should be Static or NotInheritable
-dotnet_diagnostic.CA1052.severity = none
-
-# CA1310: Specify StringComparison for correctness
-dotnet_diagnostic.CA1310.severity = none
-
-# CA1847: Use char literal for a single character lookup
-dotnet_diagnostic.CA1847.severity = none
-
-# CA1001: Types that own disposable fields should be disposable
-dotnet_diagnostic.CA1001.severity = none
-
-# CA1721: Property names should not match get methods
-dotnet_diagnostic.CA1721.severity = none
-
-# CA1018: Mark attributes with AttributeUsageAttribute
-dotnet_diagnostic.CA1018.severity = none
-
-# CA1716: Identifiers should not match keywords
-dotnet_diagnostic.CA1716.severity = none
+resharper_csharp_indent_style = space
 
-# CA1707: Identifiers should not contain underscores
-dotnet_diagnostic.CA1707.severity = none
+# All Rules Enabled as build warnings
+# Description: All Rules are enabled as build warnings. Rules with IsEnabledByDefault = false are force enabled as build warnings.
 
-# CA1812: Avoid uninstantiated internal classes
-dotnet_diagnostic.CA1812.severity = none
+# Code files
+[*.{cs,vb}]
 
-# CA1063: Implement IDisposable Correctly
-dotnet_diagnostic.CA1063.severity = none
+dotnet_style_predefined_type_for_locals_parameters_members = true
+dotnet_style_predefined_type_for_member_access = true
 
 # CA1819: Properties should not return arrays
-dotnet_diagnostic.CA1819.severity = none
-
-# CA2234: Pass system uri objects instead of strings
-dotnet_diagnostic.CA2234.severity = none
-
-# CA2208: Instantiate argument exceptions correctly
-dotnet_diagnostic.CA2208.severity = none
-
-# CA2008: Do not create tasks without passing a TaskScheduler
-dotnet_diagnostic.CA2008.severity = none
-
-# CA1307: Specify StringComparison for clarity
-dotnet_diagnostic.CA1307.severity = none
-
-# CA1051: Do not declare visible instance fields
-dotnet_diagnostic.CA1051.severity = none
-
-# CA2000: Dispose objects before losing scope
-dotnet_diagnostic.CA2000.severity = none
-
-# CA1034: Nested types should not be visible
-dotnet_diagnostic.CA1034.severity = none
-
-# CA2227: Collection properties should be read only
-dotnet_diagnostic.CA2227.severity = none
-
-# CA1712: Do not prefix enum values with type name
-dotnet_diagnostic.CA1712.severity = none
-
-# CA1816: Dispose methods should call SuppressFinalize
-dotnet_diagnostic.CA1816.severity = none
-
-# CA1064: Exceptions should be public
-dotnet_diagnostic.CA1064.severity = none
-
-# CA1056: URI-like properties should not be strings
-dotnet_diagnostic.CA1056.severity = none
-
-# CA2211: Non-constant fields should not be visible
-dotnet_diagnostic.CA2211.severity = none
-
-# CA1028: Enum Storage should be Int32
-dotnet_diagnostic.CA1028.severity = none
-
-# CA1836: Prefer IsEmpty over Count
-dotnet_diagnostic.CA1836.severity = none
-
-# CA1822: Mark members as static
-dotnet_diagnostic.CA1822.severity = none
-
-# CA2016: Forward the 'CancellationToken' parameter to methods
-dotnet_diagnostic.CA2016.severity = none
-
-# CA2213: Disposable fields should be disposed
-dotnet_diagnostic.CA2213.severity = none
-
-# CA1813: Avoid unsealed attributes
-dotnet_diagnostic.CA1813.severity = none
-
-# CA1843: Do not use 'WaitAll' with a single task
-dotnet_diagnostic.CA1843.severity = none
-
-# CA5398: Avoid hardcoded SslProtocols values
-dotnet_diagnostic.CA5398.severity = none
-
-# CA1012: Abstract types should not have public constructors
-dotnet_diagnostic.CA1012.severity = none
-
-# CA1308: Normalize strings to uppercase
-dotnet_diagnostic.CA1308.severity = none
-
-# CA1849: Call async methods when in an async method
-dotnet_diagnostic.CA1849.severity = none
-
-# CA1054: URI-like parameters should not be strings
-dotnet_diagnostic.CA1054.severity = none
-
-# CA2215: Dispose methods should call base class dispose
-dotnet_diagnostic.CA2215.severity = none
-
-# CA1825: Avoid zero-length array allocations
-dotnet_diagnostic.CA1825.severity = none
-
-# CA1050: Declare types in namespaces
-dotnet_diagnostic.CA1050.severity = none
-
-# CA1710: Identifiers should have correct suffix
-dotnet_diagnostic.CA1710.severity = none
-
-# CA1027: Mark enums with FlagsAttribute
-dotnet_diagnostic.CA1027.severity = none
-
-# CA1008: Enums should have zero value
-dotnet_diagnostic.CA1008.severity = none
-
-# CA1019: Define accessors for attribute arguments
-dotnet_diagnostic.CA1019.severity = none
+# Reason: Arrays is a good way to indicate that the collection should not be added to that is still easily serializable
+dotnet_diagnostic.ca1819.severity = none
 
 # CA5350: Do Not Use Weak Cryptographic Algorithms
-dotnet_diagnostic.CA5350.severity = none
-
-# CA1000: Do not declare static members on generic types
-dotnet_diagnostic.CA1000.severity = none
-
-# IDE0011: Add braces
-csharp_prefer_braces = when_multiline:none
-dotnet_diagnostic.IDE0011.severity = none
-
-# dotnet_style_allow_multiple_blank_lines_experimental
-dotnet_diagnostic.IDE2000.severity = none
-
-# csharp_style_allow_blank_lines_between_consecutive_braces_experimental
-dotnet_diagnostic.IDE2002.severity = none
-
-# IDE0005: Using directive is unnecessary
-dotnet_diagnostic.IDE0005.severity = none
-
-# IDE0040: Add accessibility modifiers
-dotnet_diagnostic.IDE0040.severity = none
-
-# IDE0060: Remove unused parameter
-dotnet_diagnostic.IDE0060.severity = none
+# Reason: Use we md5 but only has a hash not for encryption
+dotnet_diagnostic.ca5350.severity = none
 
-# IDE0059: Unnecessary assignment to a value
-dotnet_diagnostic.IDE0059.severity = none
+# CA5350: Do Not Use Broken Cryptographic Algorithms
+# Reason: Use we md5 but only has a hash not for encryption
+dotnet_diagnostic.ca5351.severity = none
 
-# IDE0052: Remove unread private member
-dotnet_diagnostic.IDE0052.severity = none
-
-# csharp_style_allow_embedded_statements_on_same_line_experimental
-dotnet_diagnostic.IDE2001.severity = none
-
-csharp_style_expression_bodied_properties = false:none
-csharp_style_expression_bodied_accessors = false:none
-csharp_prefer_static_local_function = false:none
-csharp_style_prefer_not_pattern = false:none
-
-dotnet_naming_rule.non_private_readonly_fields_should_be_pascal_case.severity = none
-dotnet_naming_rule.non_private_static_fields_should_be_pascal_case.severity = none
-dotnet_naming_rule.constants_should_be_pascal_case.severity = none
-dotnet_naming_rule.instance_fields_should_be_camel_case.severity = none
-dotnet_naming_rule.locals_should_be_camel_case.severity = none
-dotnet_naming_rule.static_fields_should_be_camel_case.severity = none
-dotnet_naming_rule.members_should_be_pascal_case.severity = none
-
-dotnet_style_readonly_field = false
-csharp_prefer_simple_default_expression = false:none
-dotnet_style_prefer_is_null_check_over_reference_equality_method = false:none
-csharp_style_pattern_matching_over_is_with_cast_check = false:none
-dotnet_style_prefer_compound_assignment = false:none
-
-# csharp_style_allow_blank_line_after_colon_in_constructor_initializer_experimental
-dotnet_diagnostic.IDE2004.severity = none
-
-# Do not generate build warning for broken cryptographic algorithms (MD5 in this case) we do not use it for encryption but as a hash
-dotnet_diagnostic.CA5351.severity = suggestion
+# CA1716: Identifiers should not match keywords
+# Reason: We frequently use verbs as method names, like GET as well as terms like namespace which collide with built in keywords
+dotnet_diagnostic.CA1716.severity = none
\ No newline at end of file
diff --git a/Engine/Source/Programs/Horde/HordeStorage/Callisto/CallistoStartup.cs b/Engine/Source/Programs/Horde/HordeStorage/Callisto/CallistoStartup.cs
--- a/Engine/Source/Programs/Horde/HordeStorage/Callisto/CallistoStartup.cs
+++ b/Engine/Source/Programs/Horde/HordeStorage/Callisto/CallistoStartup.cs
@@ -7,7 +7,6 @@ using System.IO;
 using Callisto.Implementation;
 using Jupiter;
 using Microsoft.AspNetCore.Authorization;
-using Microsoft.AspNetCore.Hosting;
 using Microsoft.Extensions.Configuration;
 using Microsoft.Extensions.DependencyInjection;
 using Microsoft.Extensions.Options;
@@ -17,7 +16,7 @@ namespace Callisto
     // ReSharper disable once ClassNeverInstantiated.Global
     public class CallistoStartup : BaseStartup
     {
-        public CallistoStartup(IConfiguration configuration, IWebHostEnvironment environment) : base(configuration, environment)
+        public CallistoStartup(IConfiguration configuration) : base(configuration)
         {
         }
 
@@ -92,8 +91,8 @@ namespace Callisto
         [Required]
         public string TransactionLogRoot
         {
-            get { return Environment.ExpandEnvironmentVariables(_transactionLogRoot); }
-            set { _transactionLogRoot = value; }
+            get => Environment.ExpandEnvironmentVariables(_transactionLogRoot);
+            set => _transactionLogRoot = value;
         }
 
         public bool VerifySerialization { get; set; } = false;
diff --git a/Engine/Source/Programs/Horde/HordeStorage/Callisto/Implementation/FileTransactionLog.cs b/Engine/Source/Programs/Horde/HordeStorage/Callisto/Implementation/FileTransactionLog.cs
--- a/Engine/Source/Programs/Horde/HordeStorage/Callisto/Implementation/FileTransactionLog.cs
+++ b/Engine/Source/Programs/Horde/HordeStorage/Callisto/Implementation/FileTransactionLog.cs
@@ -5,7 +5,6 @@ using System.Collections.Generic;
 using System.Diagnostics;
 using System.IO;
 using System.Linq;
-using System.Reflection.Metadata;
 using System.Security.Cryptography;
 using System.Text;
 using System.Threading.Tasks;
@@ -64,12 +63,7 @@ namespace Callisto.Implementation
             _transactionLogHandle = new FileInfo(Path.Combine(root.FullName, ns.ToString(), LogFilename));
         }
 
-
-        internal FileTransactionLogIndex IndexFile
-        {
-            get { return _index; }
-        }
-
+        internal FileTransactionLogIndex IndexFile => _index;
 
         public async Task<long> NewTransaction(AddTransactionEvent @event)
         {
@@ -92,7 +86,9 @@ namespace Callisto.Implementation
             span.OperationName = "WRITE";
 
             if (_index.Version != 1)
+            {
                 throw new InvalidOperationException($"Version mismatch of transaction log, found {_index.Version} but expected version 1");
+            }
 
             byte[] eventBuffer = SerializeToBuffer(eventRecord);
             int eventLength = eventBuffer.Length;
@@ -100,7 +96,6 @@ namespace Callisto.Implementation
             using MD5 md5 = MD5.Create();
             byte[] md5Hash = md5.ComputeHash(eventBuffer, 0, eventLength);
 
-
             long offset = await Task.Run(() =>
             {
                 lock (_commitLock)
@@ -111,7 +106,9 @@ namespace Callisto.Implementation
                     long position = fs.Position;
                     
                     if (_settings?.CurrentValue?.VerifySerialization ?? false)
+                    {
                         VerifySerializationOfEvent(position, eventRecord.Op!, eventBuffer, md5Hash, eventLength, md5);
+                    }
 
                     OpHeader header = new OpHeader(position, md5Hash, eventLength);
                     header.Serialize(fs);
@@ -123,7 +120,7 @@ namespace Callisto.Implementation
             return offset;
         }
 
-        private byte[] SerializeToBuffer(OpRecord eventRecord)
+        private static byte[] SerializeToBuffer(OpRecord eventRecord)
         {
             using MemoryStream memoryStream = new MemoryStream();
             eventRecord.Serialize(memoryStream);
@@ -182,7 +179,9 @@ namespace Callisto.Implementation
             span.OperationName = "GET";
 
             if (_index.Version != 1)
+            {
                 throw new InvalidOperationException($"Version mismatch of transaction log, found {_index.Version} but expected version 1");
+            }
 
             OpRecord[] events = new OpRecord[count];
 
@@ -195,7 +194,9 @@ namespace Callisto.Implementation
 
             _transactionLogHandle.Refresh();
             if (!_transactionLogHandle.Exists)
+            {
                 throw new FileNotFoundException("Transaction log is empty");
+            }
 
             // we set a max amount of events to skip to make sure we actually move the current event counter forward for the replicator and do not have to start from the beginning again
             const int maxEventsSkippedDueToLocation = 1000;
@@ -218,7 +219,9 @@ namespace Callisto.Implementation
 
                 // if we would read outside the stream when reading the next header we stop
                 if (fs.Position + OpHeader.HeaderLength > fs.Length)
+                {
                     break;
+                }
 
                 OpHeader opHeader = OpHeader.Deserialize(fs);
 
@@ -236,9 +239,16 @@ namespace Callisto.Implementation
                 offset = fs.Position;
 
                 if (offset - eventPosition != opHeader.Length)
+                {
                     throw new Exception($"Read event size {offset - eventPosition} did not match expected size from header {opHeader.Length}");
-                using MD5 md5 = MD5.Create();
-                byte[] md5Hash = md5.ComputeHash(eventBuf, 0, opHeader.Length);
+                }
+
+                byte[] md5Hash;
+                using (MD5 md5 = MD5.Create())
+                {
+                    md5Hash = md5.ComputeHash(eventBuf, 0, opHeader.Length);
+                }
+
                 if (!opHeader.Md5hash.SequenceEqual(md5Hash))
                 {
                     string headerMd5 = StringUtils.FormatAsHexString(opHeader.Md5hash);
@@ -247,8 +257,13 @@ namespace Callisto.Implementation
                     throw new EventHashMismatchException(opHeader.Offset, headerMd5, calculatedMd5);
                 }
 
-                await using MemoryStream eventStream = new MemoryStream(eventBuf);
-                OpRecord op = OpRecord.Deserialize(eventStream);
+                OpRecord op;
+                {
+#pragma warning disable CA2000 // Dispose objects before losing scope -- object is disposed so false positive
+                    await using MemoryStream eventStream = new MemoryStream(eventBuf);
+#pragma warning restore CA2000 // Dispose objects before losing scope
+                    op = OpRecord.Deserialize(eventStream);
+                }
 
                 if (siteFilter != 0)
                 {
@@ -305,11 +320,7 @@ namespace Callisto.Implementation
             // Length of transaction event
             public int Length;
 
-
-            public static int HeaderLength
-            {
-                get { return sizeof(long) + 16 /* MD5 hash */ + sizeof(int); }
-            }
+            public static int HeaderLength => sizeof(long) + 16 /* MD5 hash */ + sizeof(int);
 
             public OpHeader(in long offset, byte[] md5Hash, in int length)
             {
@@ -379,7 +390,7 @@ namespace Callisto.Implementation
                         baseOp = RemoveOp.Deserialize(stream);
                         break;
                     default:
-                        throw new ArgumentOutOfRangeException();
+                        throw new NotImplementedException($"Op type {op} not supported");
                 }
 
                 return new OpRecord(op, baseOp);
@@ -392,7 +403,6 @@ namespace Callisto.Implementation
             }
         }
 
-
         internal abstract class BaseOp
         {
             protected BaseOp(string name, string bucket, ulong locations)
@@ -432,7 +442,10 @@ namespace Callisto.Implementation
             {
                 List<string> enumerable = eventLocations.ToList();
                 if (!enumerable.Any())
+                {
                     throw new Exception("No locations set, at least the location of the current instance has to be specified");
+                }
+
                 ulong flags = 0;
                 foreach (int indexToSet in enumerable.Select(index.ToIndex))
                 {
@@ -464,7 +477,6 @@ namespace Callisto.Implementation
             protected abstract void DoSerialize(Stream stream);
         }
 
-
         internal class RemoveOp : BaseOp
         {
             private RemoveOp(string name, string bucket, ulong locations) : base(name, bucket, locations)
@@ -492,7 +504,7 @@ namespace Callisto.Implementation
                 // nothing extra to serialize
             }
 
-            public new static BaseOp Deserialize(Stream stream)
+            public static new BaseOp Deserialize(Stream stream)
             {
                 (string name, string bucket, ulong locations) = BaseOp.Deserialize(stream);
                 return new RemoveOp(name, bucket, locations);
@@ -511,17 +523,16 @@ namespace Callisto.Implementation
 
             public List<BlobIdentifier> Blobs { get; set; } = null!;
 
-            public byte[] Metadata { get; set; } = new byte[0];
-
+            public byte[] Metadata { get; set; } = Array.Empty<byte>();
 
             public static OpRecord FromEvent(FileTransactionLogIndex index, AddTransactionEvent @event)
             {
                 // create a bson representation of the metadata
-                byte[]? metadataBuffer = new byte[0];
+                byte[] metadataBuffer = Array.Empty<byte>();
                 if (@event.Metadata != null)
                 {
                     using MemoryStream ms = new MemoryStream();
-                    BsonDataWriter writer = new BsonDataWriter(ms);
+                    using BsonDataWriter writer = new BsonDataWriter(ms);
                     JsonSerializer serializer = JsonSerializer.CreateDefault();
                     serializer.Serialize(writer, @event.Metadata);
 
@@ -542,7 +553,7 @@ namespace Callisto.Implementation
                 if (Metadata != null && Metadata.Length != 0)
                 {
                     using MemoryStream ms = new MemoryStream(Metadata);
-                    BsonDataReader reader = new BsonDataReader(ms);
+                    using BsonDataReader reader = new BsonDataReader(ms);
                     JsonSerializer serializer = JsonSerializer.CreateDefault();
                     try
                     {
@@ -581,7 +592,7 @@ namespace Callisto.Implementation
                 writer.Write(Metadata, 0, Metadata.Length);
             }
 
-            public new static BaseOp Deserialize(Stream stream)
+            public static new BaseOp Deserialize(Stream stream)
             {
                 (string name, string bucket, ulong locations) = BaseOp.Deserialize(stream);
 
diff --git a/Engine/Source/Programs/Horde/HordeStorage/Callisto/Implementation/FileTransactionLogIndex.cs b/Engine/Source/Programs/Horde/HordeStorage/Callisto/Implementation/FileTransactionLogIndex.cs
--- a/Engine/Source/Programs/Horde/HordeStorage/Callisto/Implementation/FileTransactionLogIndex.cs
+++ b/Engine/Source/Programs/Horde/HordeStorage/Callisto/Implementation/FileTransactionLogIndex.cs
@@ -23,16 +23,11 @@ namespace Callisto.Implementation
         [JsonProperty]
         public Guid LogGeneration { get; set; }
 
-
         [JsonProperty]
         public int Version { get; set; }
 
-
         [JsonIgnore]
-        public List<string> Locations
-        {
-            get { return _locationToIdentifier.Keys.ToList(); }
-        }
+        public List<string> Locations => _locationToIdentifier.Keys.ToList();
 
         // ReSharper disable once UnusedMember.Local
         public FileTransactionLogIndex()
diff --git a/Engine/Source/Programs/Horde/HordeStorage/Callisto/Implementation/FileTransactionLogProxy.cs b/Engine/Source/Programs/Horde/HordeStorage/Callisto/Implementation/FileTransactionLogProxy.cs
--- a/Engine/Source/Programs/Horde/HordeStorage/Callisto/Implementation/FileTransactionLogProxy.cs
+++ b/Engine/Source/Programs/Horde/HordeStorage/Callisto/Implementation/FileTransactionLogProxy.cs
@@ -5,7 +5,6 @@ using System.Collections.Concurrent;
 using System.IO;
 using System.Linq;
 using EpicGames.Horde.Storage;
-using Jupiter.Implementation;
 using Microsoft.Extensions.Options;
 
 namespace Callisto.Implementation
diff --git a/Engine/Source/Programs/Horde/HordeStorage/Callisto/Implementation/ITransactionLog.cs b/Engine/Source/Programs/Horde/HordeStorage/Callisto/Implementation/ITransactionLog.cs
--- a/Engine/Source/Programs/Horde/HordeStorage/Callisto/Implementation/ITransactionLog.cs
+++ b/Engine/Source/Programs/Horde/HordeStorage/Callisto/Implementation/ITransactionLog.cs
@@ -19,7 +19,6 @@ namespace Callisto.Implementation
         TransactionLogDescription Describe();
     }
 
-
     public interface ITransactionLogs
     {
         NamespaceId[] GetNamespaces();
@@ -35,8 +34,8 @@ namespace Callisto.Implementation
             KnownSites = knownSites;
         }
 
-        public Guid Generation { get; set; }
-        public List<string> KnownSites { get; set; }
+        public Guid Generation { get; }
+        public List<string> KnownSites { get; }
     }
 
     public class TransactionEvents
diff --git a/Engine/Source/Programs/Horde/HordeStorage/Callisto/Tests/CallisoFunctionalTests/CallistoTests.cs b/Engine/Source/Programs/Horde/HordeStorage/Callisto/Tests/CallisoFunctionalTests/CallistoTests.cs
--- a/Engine/Source/Programs/Horde/HordeStorage/Callisto/Tests/CallisoFunctionalTests/CallistoTests.cs
+++ b/Engine/Source/Programs/Horde/HordeStorage/Callisto/Tests/CallisoFunctionalTests/CallistoTests.cs
@@ -1,8 +1,6 @@
 // Copyright Epic Games, Inc. All Rights Reserved.
 
 using System;
-using System.Collections.Generic;
-using System.Diagnostics;
 using System.IO;
 using System.Linq;
 using System.Net.Http;
@@ -107,7 +105,7 @@ namespace Callisto.FunctionalTests
         {
             const int startOffset = 0;
             HttpResponseMessage response =
-                await _httpClient!.GetAsync($"api/v1/t/{TestExistingNamespace}/{startOffset}");
+                await _httpClient!.GetAsync(new Uri($"api/v1/t/{TestExistingNamespace}/{startOffset}", UriKind.Relative));
             response.EnsureSuccessStatusCode();
             dynamic ter = await response.Content.ReadAsAsync<dynamic>();
             Assert.IsNotNull(ter.generation);
@@ -129,8 +127,7 @@ namespace Callisto.FunctionalTests
         public async Task GetMismatchedOffsetEvent()
         {
             const int startOffset = 50; // this is not a valid offset, so we should get the object at offset 70
-            HttpResponseMessage response =
-                await _httpClient!.GetAsync($"api/v1/t/{TestExistingNamespace}/{startOffset}");
+            HttpResponseMessage response = await _httpClient!.GetAsync(new Uri($"api/v1/t/{TestExistingNamespace}/{startOffset}", UriKind.Relative));
             response.EnsureSuccessStatusCode();
             dynamic ter = await response.Content.ReadAsAsync<dynamic>();
             Assert.IsNotNull(ter.generation);
@@ -148,13 +145,12 @@ namespace Callisto.FunctionalTests
             Assert.AreEqual(new BlobIdentifier("0BEEC7B5EA3F0FDBC95D0DD47F3C5BC275DA8A34"), actual: tEvent.Blobs.First());
         }
 
-
         [TestMethod]    
         public async Task GetNonExistentEvent()
         {
             // fetch the event after the second event, this will not exist
             const int startOffset = 77;
-            HttpResponseMessage response = await _httpClient!.GetAsync($"api/v1/t/{TestExistingNamespace}/{startOffset}");
+            HttpResponseMessage response = await _httpClient!.GetAsync(new Uri($"api/v1/t/{TestExistingNamespace}/{startOffset}", UriKind.Relative));
             response.EnsureSuccessStatusCode();
             dynamic ter = await response.Content.ReadAsAsync<dynamic>();
             Assert.IsNotNull(ter.generation);
@@ -162,7 +158,6 @@ namespace Callisto.FunctionalTests
             Assert.AreEqual(0, events.Count);
         }
 
-
         [TestMethod]
         public async Task AddNewEvent()
         {
@@ -255,7 +250,7 @@ namespace Callisto.FunctionalTests
                 };
 
                 HttpResponseMessage response =
-                    await _httpClient!.PostAsJsonAsync(requestUri: $"api/v1/t/{FullFlowNamespace}", request);
+                    await _httpClient!.PostAsJsonAsync(new Uri($"api/v1/t/{FullFlowNamespace}", UriKind.Relative), request);
                 response.EnsureSuccessStatusCode();
                 NewTransactionResponse r = await response.Content.ReadAsAsync<NewTransactionResponse>();
                 long offset = r.Offset;
@@ -265,7 +260,7 @@ namespace Callisto.FunctionalTests
             {
                 const int startOffset = 0;
                 HttpResponseMessage response =
-                    await _httpClient!.GetAsync($"api/v1/t/{FullFlowNamespace}/{startOffset}");
+                    await _httpClient!.GetAsync(new Uri($"api/v1/t/{FullFlowNamespace}/{startOffset}", UriKind.Relative));
                 response.EnsureSuccessStatusCode();
                 dynamic ter = await response.Content.ReadAsAsync<dynamic>();
                 JArray events = ter.events;
diff --git a/Engine/Source/Programs/Horde/HordeStorage/Callisto/Tests/CallisoUnitTests/TransactionLogTests.cs b/Engine/Source/Programs/Horde/HordeStorage/Callisto/Tests/CallisoUnitTests/TransactionLogTests.cs
--- a/Engine/Source/Programs/Horde/HordeStorage/Callisto/Tests/CallisoUnitTests/TransactionLogTests.cs
+++ b/Engine/Source/Programs/Horde/HordeStorage/Callisto/Tests/CallisoUnitTests/TransactionLogTests.cs
@@ -1,8 +1,6 @@
 // Copyright Epic Games, Inc. All Rights Reserved.
 
-using System;
 using System.Collections.Generic;
-using System.Diagnostics;
 using System.IO;
 using System.Linq;
 using Microsoft.VisualStudio.TestTools.UnitTesting;
@@ -173,7 +171,6 @@ namespace Calliso.UnitTests
             }
         }
 
-
         [TestMethod]
         public async Task FileTransactionLogSiteFilter()
         {
@@ -221,7 +218,6 @@ namespace Calliso.UnitTests
             Assert.AreEqual(arnEvent.Bucket, transactionEvent.Bucket);
         }
 
-
         [TestMethod]
         public async Task FileTransactionLogMultipleSites()
         {
diff --git a/Engine/Source/Programs/Horde/HordeStorage/Horde.Storage/Controllers/AdminController.cs b/Engine/Source/Programs/Horde/HordeStorage/Horde.Storage/Controllers/AdminController.cs
--- a/Engine/Source/Programs/Horde/HordeStorage/Horde.Storage/Controllers/AdminController.cs
+++ b/Engine/Source/Programs/Horde/HordeStorage/Horde.Storage/Controllers/AdminController.cs
@@ -24,14 +24,12 @@ namespace Horde.Storage.Controllers
     {
         private readonly LastAccessService _lastAccessService;
         private readonly IRefCleanup _refCleanup;
-        private readonly BlobCleanupService _blobCleanupService;
         private readonly IConfiguration _configuration;
 
-        public AdminController(LastAccessService lastAccessService, IRefCleanup refCleanup, BlobCleanupService blobCleanupService, IConfiguration configuration)
+        public AdminController(LastAccessService lastAccessService, IRefCleanup refCleanup, IConfiguration configuration)
         {
             _lastAccessService = lastAccessService;
             _refCleanup = refCleanup;
-            _blobCleanupService = blobCleanupService;
             _configuration = configuration;
         }
 
@@ -105,7 +103,9 @@ namespace Horde.Storage.Controllers
             {
                 Dictionary<string, object> values = ResolveSection(section);
                 if (values.Count != 0)
+                {
                     settings.Add(section.Key, values);
+                }
             }
 
             return new JsonResult(new
@@ -148,6 +148,7 @@ namespace Horde.Storage.Controllers
             UpdatedRecords = updatedRecords;
         }
 
+        [System.Diagnostics.CodeAnalysis.SuppressMessage("Design", "CA1034:Nested types should not be visible", Justification = "Only used by serialization")]
         public class UpdatedRecord
         {
             public UpdatedRecord(RefRecord record, DateTime time)
diff --git a/Engine/Source/Programs/Horde/HordeStorage/Horde.Storage/Controllers/AuthController.cs b/Engine/Source/Programs/Horde/HordeStorage/Horde.Storage/Controllers/AuthController.cs
--- a/Engine/Source/Programs/Horde/HordeStorage/Horde.Storage/Controllers/AuthController.cs
+++ b/Engine/Source/Programs/Horde/HordeStorage/Horde.Storage/Controllers/AuthController.cs
@@ -1,15 +1,9 @@
 ﻿// Copyright Epic Games, Inc. All Rights Reserved.
 
-using Microsoft.AspNetCore.Http;
 using Microsoft.AspNetCore.Mvc;
-using System;
-using System.Collections.Generic;
 using System.ComponentModel.DataAnnotations;
-using System.Linq;
 using System.Threading.Tasks;
-using Datadog.Trace;
 using Jupiter;
-using Jupiter.Implementation;
 using Microsoft.AspNetCore.Authorization;
 using EpicGames.Horde.Storage;
 
diff --git a/Engine/Source/Programs/Horde/HordeStorage/Horde.Storage/Controllers/CompressedBlobController.cs b/Engine/Source/Programs/Horde/HordeStorage/Horde.Storage/Controllers/CompressedBlobController.cs
--- a/Engine/Source/Programs/Horde/HordeStorage/Horde.Storage/Controllers/CompressedBlobController.cs
+++ b/Engine/Source/Programs/Horde/HordeStorage/Horde.Storage/Controllers/CompressedBlobController.cs
@@ -34,11 +34,8 @@ namespace Horde.Storage.Controllers
         private readonly IAuthorizationService _authorizationService;
         private readonly CompressedBufferUtils _compressedBufferUtils;
         private readonly BufferedPayloadFactory _bufferedPayloadFactory;
-        private readonly FormatResolver _formatResolver;
 
-        private readonly ILogger _logger = Log.ForContext<CompressedBlobController>();
-
-        public CompressedBlobController(IBlobService storage, IContentIdStore contentIdStore, IDiagnosticContext diagnosticContext, IAuthorizationService authorizationService, CompressedBufferUtils compressedBufferUtils, BufferedPayloadFactory bufferedPayloadFactory, FormatResolver formatResolver)
+        public CompressedBlobController(IBlobService storage, IContentIdStore contentIdStore, IDiagnosticContext diagnosticContext, IAuthorizationService authorizationService, CompressedBufferUtils compressedBufferUtils, BufferedPayloadFactory bufferedPayloadFactory)
         {
             _storage = storage;
             _contentIdStore = contentIdStore;
@@ -46,10 +43,8 @@ namespace Horde.Storage.Controllers
             _authorizationService = authorizationService;
             _compressedBufferUtils = compressedBufferUtils;
             _bufferedPayloadFactory = bufferedPayloadFactory;
-            _formatResolver = formatResolver;
         }
 
-
         [HttpGet("{ns}/{id}")]
         [Authorize("Storage.read")]
         [ProducesResponseType(type: typeof(byte[]), 200)]
@@ -73,7 +68,9 @@ namespace Horde.Storage.Controllers
 
                 StringValues acceptHeader = Request.Headers["Accept"];
                 if (!acceptHeader.Contains("*/*") && acceptHeader.Count != 0 && !acceptHeader.Contains(mediaType))
+                {
                     return new UnsupportedMediaTypeResult();
+                }
 
                 return File(blobContents.Stream, mediaType, enableRangeProcessing: true);
             }
@@ -300,7 +297,7 @@ namespace Horde.Storage.Controllers
             // TODO: we should add a overload for decompress content that can work on streams, otherwise we are still limited to 2GB compressed blobs
             byte[] decompressedContent = _compressedBufferUtils.DecompressContent(await decompressStream.ToByteArray());
 
-            MemoryStream decompressedStream = new MemoryStream(decompressedContent);
+            await using MemoryStream decompressedStream = new MemoryStream(decompressedContent);
             ContentId identifierDecompressedPayload = ContentId.FromContentHash(await _storage.VerifyContentMatchesHash(decompressedStream, id));
 
             BlobIdentifier identifierCompressedPayload;
diff --git a/Engine/Source/Programs/Horde/HordeStorage/Horde.Storage/Controllers/ContentIdController.cs b/Engine/Source/Programs/Horde/HordeStorage/Horde.Storage/Controllers/ContentIdController.cs
--- a/Engine/Source/Programs/Horde/HordeStorage/Horde.Storage/Controllers/ContentIdController.cs
+++ b/Engine/Source/Programs/Horde/HordeStorage/Horde.Storage/Controllers/ContentIdController.cs
@@ -1,7 +1,5 @@
 // Copyright Epic Games, Inc. All Rights Reserved.
 
-using System;
-using System.Linq;
 using System.Net.Mime;
 using System.Threading.Tasks;
 using Datadog.Trace;
@@ -21,7 +19,7 @@ namespace Horde.Storage.Controllers
     public class ContentIdController : ControllerBase
     {
         private readonly IAuthorizationService _authorizationService;
-        private IContentIdStore _contentIdStore;
+        private readonly IContentIdStore _contentIdStore;
 
         public ContentIdController(IAuthorizationService authorizationService, IContentIdStore contentIdStore)
         {
@@ -34,12 +32,11 @@ namespace Horde.Storage.Controllers
         /// </summary>
         /// <param name="ns">Namespace. Each namespace is completely separated from each other. Use for different types of data that is never expected to be similar (between two different games for instance). Example: `uc4.ddc`</param>
         /// <param name="contentId">The content id to resolve </param>
-        /// <param name="format">Optional specifier to set which output format is used json/raw/cb</param>
         [HttpGet("{ns}/{contentId}.{format?}", Order = 500)]
         [ProducesDefaultResponseType]
         [ProducesResponseType(type: typeof(ProblemDetails), 400)]
         [Authorize("Cache.read")]
-        public async Task<IActionResult> Resolve(NamespaceId ns, ContentId contentId, string? format = null)
+        public async Task<IActionResult> Resolve(NamespaceId ns, ContentId contentId)
         {
             using (IScope _ = Tracer.Instance.StartActive("authorize"))
             {
diff --git a/Engine/Source/Programs/Horde/HordeStorage/Horde.Storage/Controllers/DDCRefController.cs b/Engine/Source/Programs/Horde/HordeStorage/Horde.Storage/Controllers/DDCRefController.cs
--- a/Engine/Source/Programs/Horde/HordeStorage/Horde.Storage/Controllers/DDCRefController.cs
+++ b/Engine/Source/Programs/Horde/HordeStorage/Horde.Storage/Controllers/DDCRefController.cs
@@ -18,9 +18,6 @@ using Microsoft.AspNetCore.Authorization;
 using Microsoft.AspNetCore.Http;
 using Microsoft.AspNetCore.Http.Extensions;
 using Microsoft.AspNetCore.Mvc;
-using Microsoft.AspNetCore.Mvc.Infrastructure;
-using Microsoft.Extensions.Configuration;
-using Microsoft.Extensions.Options;
 using Newtonsoft.Json;
 using Serilog;
 
@@ -53,18 +50,16 @@ namespace Horde.Storage.Controllers
         private readonly IRefsStore _refsStore;
         private readonly IDiagnosticContext _diagnosticContext;
         private readonly IAuthorizationService _authorizationService;
-        private readonly IOptionsMonitor<JupiterSettings> _jupiterSettings;
 
         private readonly ILogger _logger = Log.ForContext<DDCRefController>();
         private readonly IDDCRefService _ddcRefService;
 
-        public DDCRefController(IDDCRefService ddcRefService, IRefsStore refsStore, IDiagnosticContext diagnosticContext, IAuthorizationService authorizationService, IOptionsMonitor<JupiterSettings> jupiterSettings)
+        public DDCRefController(IDDCRefService ddcRefService, IRefsStore refsStore, IDiagnosticContext diagnosticContext, IAuthorizationService authorizationService)
         {
             _refsStore = refsStore;
             _ddcRefService = ddcRefService;
             _diagnosticContext = diagnosticContext;
             _authorizationService = authorizationService;
-            _jupiterSettings = jupiterSettings;
         }
 
         /// <summary>
@@ -123,7 +118,7 @@ namespace Horde.Storage.Controllers
 
             try
             {
-                bool isRaw = format?.ToLowerInvariant() == "raw";
+                bool isRaw = format?.ToUpperInvariant() == "RAW";
                 // if the raw format is used only the blob itself will be output so no need to fetch anything else
                 if (isRaw)
                 {
@@ -142,7 +137,7 @@ namespace Horde.Storage.Controllers
                         using (IScope _ = Tracer.Instance.StartActive("body.write"))
                         {
                             const int BufferSize = 64 * 1024;
-                            var outputStream = Response.Body;
+                            Stream outputStream = Response.Body;
                             Response.ContentLength = blob.Length;
                             Response.ContentType = MediaTypeNames.Application.Octet;
                             Response.StatusCode = StatusCodes.Status200OK;
@@ -170,7 +165,6 @@ namespace Horde.Storage.Controllers
                     record.Blob = await blobStream.ToByteArray();
                 }
 
-
                 return Ok(record);
 
             }
@@ -283,7 +277,6 @@ namespace Horde.Storage.Controllers
             }
         }
 
-
         // that structured data will match the route above first
         [RequiredContentType(MediaTypeNames.Application.Octet)]
         [HttpPut("ddc/{ns}/{bucket}/{key}", Order = 500)]
@@ -380,7 +373,6 @@ namespace Horde.Storage.Controllers
                 return BadRequest(new ProblemDetails {Title = $"Namespace {e.Namespace} did not exist"});
             }
 
-
             return NoContent();
         }
 
@@ -412,7 +404,6 @@ namespace Horde.Storage.Controllers
                 return BadRequest(new ProblemDetails {Title = $"Namespace {e.Namespace} did not exist"});
             }
 
-
             return NoContent();
         }
 
@@ -443,7 +434,9 @@ namespace Horde.Storage.Controllers
                 long deleteCount = await _ddcRefService.Delete(ns, bucket, key);
 
                 if (deleteCount == 0)
+                {
                     return BadRequest("Deleted 0 records, most likely the object did not exist");
+                }
             }
             catch (NamespaceNotFoundException e)
             {
@@ -452,9 +445,9 @@ namespace Horde.Storage.Controllers
             return NoContent();
         }
 
-
         // ReSharper disable UnusedAutoPropertyAccessor.Global
         // ReSharper disable once ClassNeverInstantiated.Global
+        [System.Diagnostics.CodeAnalysis.SuppressMessage("Design", "CA1034:Nested types should not be visible", Justification = "Required by serialization")]
         public class DdcBatchOp
         {
             // ReSharper disable InconsistentNaming
@@ -482,6 +475,7 @@ namespace Horde.Storage.Controllers
             public byte[]? Content { get; set; }
         }
 
+        [System.Diagnostics.CodeAnalysis.SuppressMessage("Design", "CA1034:Nested types should not be visible", Justification = "Required by serialization")]
         public class DdcBatchCall
         {
             public DdcBatchOp[]? Operations { get; set; }
@@ -508,7 +502,7 @@ namespace Horde.Storage.Controllers
 
             if (batch.Operations == null)
             {
-                throw new ArgumentNullException();
+                throw new InvalidOperationException("No operations specified, this is a required field");
             }
 
             Task<object?>[] tasks = new Task<object?>[batch.Operations.Length];
@@ -527,7 +521,7 @@ namespace Horde.Storage.Controllers
                 }
                 else
                 {
-                    tasks[index] = ProcessOp(op).ContinueWith(task => task.IsFaulted ? (object?)new ProblemDetails { Title = "Exception thrown", Detail = task!.Exception!.ToString()} : task.Result);
+                    tasks[index] = ProcessOp(op).ContinueWith((task, _) => task.IsFaulted ? (object?)new ProblemDetails { Title = "Exception thrown", Detail = task!.Exception!.ToString()} : task.Result, null, TaskScheduler.Current);
                 }
             }
 
@@ -542,26 +536,26 @@ namespace Horde.Storage.Controllers
         {
             if (op.Namespace == null)
             {
-                throw new ArgumentNullException("namespace");
+                throw new InvalidOperationException("namespace was null, is required");
             }
 
             if (op.Bucket == null)
             {
-                throw new ArgumentNullException("bucket");
+                throw new InvalidOperationException("bucket was null, is required");
             }
 
             if (op.Id == null)
             {
-                throw new ArgumentNullException("id");
+                throw new InvalidOperationException("id was null, is required");
             }
 
             switch (op.Op)
             {
                 case DdcBatchOp.DdcOperation.INVALID:
-                    throw new ArgumentOutOfRangeException();
+                    throw new NotImplementedException($"Unsupported batch op {DdcBatchOp.DdcOperation.INVALID} used in op {op.Id}");
                 case DdcBatchOp.DdcOperation.GET:
                     // TODO: support field filtering
-                    return _ddcRefService.Get(op.Namespace.Value, op.Bucket.Value, op.Id.Value, new string[0]).ContinueWith(t => 
+                    return _ddcRefService.Get(op.Namespace.Value, op.Bucket.Value, op.Id.Value, Array.Empty<string>()).ContinueWith((t,_) => 
                     {
                         // we are serializing this to json, so we just stream the blob into memory to return it.
                         (RefResponse response, BlobContents? blob) = t.Result;
@@ -571,19 +565,19 @@ namespace Horde.Storage.Controllers
                             response.Blob = blobContents.Stream.ToByteArray().Result;
                         }
                         return (object?)response;
-                    });
+                    }, null, TaskScheduler.Current);
                 case DdcBatchOp.DdcOperation.HEAD:
-                    return _ddcRefService.Exists(op.Namespace.Value, op.Bucket.Value, op.Id.Value).ContinueWith(t => t.Result == null ? (object?)null : op.Id);
+                    return _ddcRefService.Exists(op.Namespace.Value, op.Bucket.Value, op.Id.Value).ContinueWith((t,_) => t.Result == null ? (object?)null : op.Id, null, TaskScheduler.Current);
                 case DdcBatchOp.DdcOperation.PUT:
                 {
                     if (op.ContentHash == null)
                     {
-                        throw new ArgumentNullException("ContentHash");
+                        throw new InvalidOperationException("ContentHash was null, is required");
                     }
 
                     if (op.Content == null)
                     {
-                        throw new ArgumentNullException("Content");
+                        throw new InvalidOperationException("Content was null, is required");
                     }
 
                     ContentHash blobHash = ContentHash.FromBlob(op.Content);
@@ -591,15 +585,16 @@ namespace Horde.Storage.Controllers
                     {
                         throw new HashMismatchException(op.ContentHash, blobHash);
                     }
-                    return _ddcRefService.Put(op.Namespace.Value, op.Bucket.Value, op.Id.Value, blobHash, op.Content).ContinueWith(t => (object?)t.Result);
+                    return _ddcRefService.Put(op.Namespace.Value, op.Bucket.Value, op.Id.Value, blobHash, op.Content).ContinueWith((t,_) => (object?)t.Result, null, TaskScheduler.Current);
                 }
                 case DdcBatchOp.DdcOperation.DELETE:
-                    return _ddcRefService.Delete(op.Namespace.Value, op.Bucket.Value, op.Id.Value).ContinueWith(t => (object?)null);
+                    return _ddcRefService.Delete(op.Namespace.Value, op.Bucket.Value, op.Id.Value).ContinueWith((t,_) => (object?)null, null, TaskScheduler.Current);
                 default:
-                    throw new ArgumentOutOfRangeException();
+                    throw new NotImplementedException($"Unknown op used {op.Op}");
             }
         }
 
+        [System.Diagnostics.CodeAnalysis.SuppressMessage("Design", "CA1034:Nested types should not be visible", Justification = "Used by serialization")]
         public class BatchGetOp
         {
             public class GetOp
@@ -626,7 +621,7 @@ namespace Horde.Storage.Controllers
         {
             if (batch.Namespace == null)
             {
-                throw new ArgumentNullException();
+                throw new InvalidOperationException("A valid namespace must be specified");
             }
 
             {
diff --git a/Engine/Source/Programs/Horde/HordeStorage/Horde.Storage/Controllers/FormatResolver.cs b/Engine/Source/Programs/Horde/HordeStorage/Horde.Storage/Controllers/FormatResolver.cs
--- a/Engine/Source/Programs/Horde/HordeStorage/Horde.Storage/Controllers/FormatResolver.cs
+++ b/Engine/Source/Programs/Horde/HordeStorage/Horde.Storage/Controllers/FormatResolver.cs
@@ -31,7 +31,10 @@ namespace Horde.Storage.Controllers
             {
                 string? typeMapping = _mvcOptions.CurrentValue.FormatterMappings.GetMediaTypeMappingForFormat(format);
                 if (typeMapping == null)
+                {
                     throw new Exception($"No mapping defined from format {format} to mime type");
+                }
+
                 return typeMapping;
             }
             
@@ -51,10 +54,9 @@ namespace Horde.Storage.Controllers
 
             foreach (string header in acceptHeader)
             {
-                string s = header.ToLowerInvariant();
-                if (_validContentTypes.Contains(s))
+                if (_validContentTypes.Contains(header, StringComparer.OrdinalIgnoreCase))
                 {
-                    return s;
+                    return header;
                 }
             }
 
diff --git a/Engine/Source/Programs/Horde/HordeStorage/Horde.Storage/Controllers/ObjectController.cs b/Engine/Source/Programs/Horde/HordeStorage/Horde.Storage/Controllers/ObjectController.cs
--- a/Engine/Source/Programs/Horde/HordeStorage/Horde.Storage/Controllers/ObjectController.cs
+++ b/Engine/Source/Programs/Horde/HordeStorage/Horde.Storage/Controllers/ObjectController.cs
@@ -4,12 +4,10 @@ using System;
 using System.Collections.Concurrent;
 using System.Collections.Generic;
 using System.ComponentModel.DataAnnotations;
-using System.IO;
 using System.Linq;
 using System.Net.Mime;
 using System.Threading.Tasks;
 using async_enumerable_dotnet;
-using Datadog.Trace;
 using EpicGames.Horde.Storage;
 using EpicGames.Serialization;
 using Horde.Storage.Implementation;
@@ -18,7 +16,6 @@ using Jupiter.Common.Implementation;
 using Jupiter.Implementation;
 using Jupiter.Utils;
 using Microsoft.AspNetCore.Authorization;
-using Microsoft.AspNetCore.Http;
 using Microsoft.AspNetCore.Mvc;
 using Serilog;
 
@@ -48,7 +45,6 @@ namespace Horde.Storage.Controllers
             _bufferedPayloadFactory = bufferedPayloadFactory;
         }
 
-
         [HttpGet("{ns}/{id}")]
         [Authorize("Storage.read")]
         [ProducesDefaultResponseType]
@@ -118,7 +114,9 @@ namespace Horde.Storage.Controllers
             IEnumerable<Task> tasks = id.Select(async blob =>
             {
                 if (!await _storage.Exists(ns, blob))
+                {
                     missingBlobs.Add(blob);
+                }
             });
             await Task.WhenAll(tasks);
 
@@ -144,7 +142,9 @@ namespace Horde.Storage.Controllers
             IEnumerable<Task> tasks = bodyIds.Select(async blob =>
             {
                 if (!await _storage.Exists(ns, blob))
+                {
                     missingBlobs.Add(blob);
+                }
             });
             await Task.WhenAll(tasks);
 
@@ -176,8 +176,7 @@ namespace Horde.Storage.Controllers
         [Authorize("Storage.read")]
         public async Task<IActionResult> ResolveReferences(
             [Required] NamespaceId ns,
-            [Required] BlobIdentifier id,
-            [FromQuery] int depth = 1)
+            [Required] BlobIdentifier id)
         {
             AuthorizationResult authorizationResult = await _authorizationService.AuthorizeAsync(User, ns, NamespaceAccessRequirement.Name);
 
@@ -227,7 +226,6 @@ namespace Horde.Storage.Controllers
             }
         }
 
-
         [HttpDelete("{ns}/{id}")]
         [Authorize("Storage.delete")]
         public async Task<IActionResult> Delete(
@@ -249,7 +247,6 @@ namespace Horde.Storage.Controllers
             });
         }
 
-
         [HttpDelete("{ns}")]
         [Authorize("Admin")]
         public async Task<IActionResult> DeleteNamespace(
diff --git a/Engine/Source/Programs/Horde/HordeStorage/Horde.Storage/Controllers/ProblemTypes.cs b/Engine/Source/Programs/Horde/HordeStorage/Horde.Storage/Controllers/ProblemTypes.cs
--- a/Engine/Source/Programs/Horde/HordeStorage/Horde.Storage/Controllers/ProblemTypes.cs
+++ b/Engine/Source/Programs/Horde/HordeStorage/Horde.Storage/Controllers/ProblemTypes.cs
@@ -2,7 +2,7 @@
 
 namespace Horde.Storage.Controllers
 {
-    public class ProblemTypes
+    public static class ProblemTypes
     {
         public static string UseSnapshot { get; set; } = "http://jupiter.epicgames.com/replication/useSnapshot";
     }
diff --git a/Engine/Source/Programs/Horde/HordeStorage/Horde.Storage/Controllers/ReferencesController.cs b/Engine/Source/Programs/Horde/HordeStorage/Horde.Storage/Controllers/ReferencesController.cs
--- a/Engine/Source/Programs/Horde/HordeStorage/Horde.Storage/Controllers/ReferencesController.cs
+++ b/Engine/Source/Programs/Horde/HordeStorage/Horde.Storage/Controllers/ReferencesController.cs
@@ -18,6 +18,7 @@ using EpicGames.Core;
 using EpicGames.Horde.Storage;
 using EpicGames.Serialization;
 using Horde.Storage.Implementation;
+using JetBrains.Annotations;
 using Jupiter;
 using Jupiter.Common.Implementation;
 using Jupiter.Implementation;
@@ -90,7 +91,6 @@ namespace Horde.Storage.Controllers
         /// <param name="ns">Namespace. Each namespace is completely separated from each other. Use for different types of data that is never expected to be similar (between two different games for instance). Example: `uc4.ddc`</param>
         /// <param name="bucket">The category/type of record you are caching. Is a clustered key together with the actual key, but all records in the same bucket can be dropped easily. Example: `terrainTexture` </param>
         /// <param name="key">The unique name of this particular key. `iAmAVeryValidKey`</param>
-        /// <param name="fields">The fields to include in the response, omit this to include everything.</param>
         /// <param name="format">Optional specifier to set which output format is used json/raw/cb</param>
         [HttpGet("{ns}/{bucket}/{key}.{format?}", Order = 500)]
         [Produces(MediaTypeNames.Application.Json, MediaTypeNames.Application.Octet, CustomMediaTypeNames.UnrealCompactBinary, CustomMediaTypeNames.JupiterInlinedPayload)]
@@ -99,7 +99,6 @@ namespace Horde.Storage.Controllers
             [FromRoute] [Required] NamespaceId ns,
             [FromRoute] [Required] BucketId bucket,
             [FromRoute] [Required] IoHashKey key,
-            [FromQuery] string[] fields,
             [FromRoute] string? format = null)
         {
             {
@@ -117,7 +116,9 @@ namespace Horde.Storage.Controllers
                 (ObjectRecord objectRecord, BlobContents? blob) = await _objectService.Get(ns, bucket, key, Array.Empty<string>());
 
                 if (blob == null)
+                {
                     throw new InvalidOperationException($"Blob was null when attempting to fetch {ns} {bucket} {key}");
+                }
 
                 if (!objectRecord.IsFinalized)
                 {
@@ -134,7 +135,7 @@ namespace Horde.Storage.Controllers
                     long contentLength = blobContents.Length;
                     scope.Span.SetTag("content-length", contentLength.ToString());
                     const int BufferSize = 64 * 1024;
-                    var outputStream = Response.Body;
+                    Stream outputStream = Response.Body;
                     Response.ContentLength = contentLength;
                     Response.ContentType = contentType;
                     Response.StatusCode = StatusCodes.Status200OK;
@@ -165,7 +166,9 @@ namespace Horde.Storage.Controllers
                             {
                                 ++count;
                                 if (field.IsBinaryAttachment())
+                                {
                                     foundField = field;
+                                }
                             });
 
                             return (count, foundField);
@@ -187,7 +190,8 @@ namespace Horde.Storage.Controllers
                         }
 
                         // this doesn't look like the generated compact binary so we just return the payload
-                        await WriteBody(new BlobContents(blobMemory), MediaTypeNames.Application.Octet);
+                        await using BlobContents contents = new(blobMemory);
+                        await WriteBody(contents, MediaTypeNames.Application.Octet);
                         break;
                     }
                     case MediaTypeNames.Application.Json:
@@ -199,7 +203,8 @@ namespace Horde.Storage.Controllers
                         }
                         CbObject cb = new CbObject(blobMemory);
                         string s = cb.ToJson();
-                        await WriteBody(new BlobContents(Encoding.UTF8.GetBytes(s)), MediaTypeNames.Application.Json);
+                        await using BlobContents contents = new BlobContents(Encoding.UTF8.GetBytes(s));
+                        await WriteBody(contents, MediaTypeNames.Application.Json);
                         break;
 
                     }
@@ -285,7 +290,8 @@ namespace Horde.Storage.Controllers
                         else if (countOfBinaryAttachmentFields == 0 && countOfAttachmentFields == 0)
                         {
                             // no attachments so we just return the compact object instead
-                            await WriteBody(new BlobContents(blobMemory), CustomMediaTypeNames.JupiterInlinedPayload);
+                            await using BlobContents contents = new BlobContents(blobMemory);
+                            await WriteBody(contents, CustomMediaTypeNames.JupiterInlinedPayload);
                             return new EmptyResult();
                         }
 
@@ -362,7 +368,6 @@ namespace Horde.Storage.Controllers
             }
         }
 
-
         /// <summary>
         /// Checks if a object exists
         /// </summary>
@@ -462,7 +467,10 @@ namespace Horde.Storage.Controllers
             {
                 int separatorIndex = name.IndexOf(".", StringComparison.Ordinal);
                 if (separatorIndex == -1)
+                {
                     return BadRequest(new ProblemDetails() { Title = $"Key {name} did not contain a '.' separator" });
+                }
+
                 BucketId bucket = new BucketId(name.Substring(0, separatorIndex));
                 IoHashKey key = new IoHashKey(name.Substring(separatorIndex + 1));
                 requestedNames.Add((bucket, key));
@@ -504,7 +512,6 @@ namespace Horde.Storage.Controllers
             return Ok(new ExistCheckMultipleRefsResponse(missingObject.ToList()));
         }
 
-
         [HttpPut("{ns}/{bucket}/{key}.{format?}", Order = 500)]
         [DisableRequestSizeLimit]
         [Authorize("Object.write")]
@@ -565,7 +572,7 @@ namespace Horde.Storage.Controllers
                     }
                     case CustomMediaTypeNames.UnrealCompactBinary:
                     {
-                        MemoryStream ms = new MemoryStream();
+                        await using MemoryStream ms = new MemoryStream();
                         await using Stream payloadStream = payload.GetStream();
                         await payloadStream.CopyToAsync(ms);
                         payloadObject = new CbObject(ms.ToArray());
@@ -598,7 +605,6 @@ namespace Horde.Storage.Controllers
                 });
             }
 
-
             (ContentId[] missingReferences, BlobIdentifier[] missingBlobs) = await _objectService.Put(ns, bucket, key, blobHeader, payloadObject);
 
             List<ContentHash> missingHashes = new List<ContentHash>(missingReferences);
@@ -631,7 +637,6 @@ namespace Horde.Storage.Controllers
             return Ok(new PutObjectResponse(missingHashes.ToArray()));
         }
 
-
         [HttpPost("{ns}")]
         [Consumes(CustomMediaTypeNames.UnrealCompactBinary)]
         [Produces(CustomMediaTypeNames.UnrealCompactBinary)]
@@ -664,7 +669,9 @@ namespace Horde.Storage.Controllers
                     }
 
                     if (blob == null)
+                    {
                         throw new Exception();
+                    }
 
                     CbObject cb = new CbObject(await blob.Stream.ToByteArray());
 
@@ -708,9 +715,10 @@ namespace Horde.Storage.Controllers
                         List<BlobIdentifier>? _ = await references.ToListAsync();
                     }
 
-
                     if (blob == null)
+                    {
                         throw new Exception();
+                    }
 
                     return (CbObject.Build(writer => writer.WriteBool("exists", true)), HttpStatusCode.OK);
                 }
@@ -771,7 +779,7 @@ namespace Horde.Storage.Controllers
                         break;
                     case BatchOps.BatchOp.Operation.INVALID:
                     default:
-                        throw new ArgumentOutOfRangeException();
+                        throw new NotImplementedException($"Unknown op type {op.Op}");
                 }
 
                 await Task.CompletedTask;
@@ -791,7 +799,7 @@ namespace Horde.Storage.Controllers
             });
         }
 
-        private (CbObject, HttpStatusCode) ToErrorResult(Exception exception, HttpStatusCode statusCode = HttpStatusCode.InternalServerError)
+        private static (CbObject, HttpStatusCode) ToErrorResult(Exception exception, HttpStatusCode statusCode = HttpStatusCode.InternalServerError)
         {
             Exception e = exception;
             CbWriter writer = new CbWriter();
@@ -829,7 +837,6 @@ namespace Horde.Storage.Controllers
                 return NotFound(new ProblemDetails {Title = $"Namespace {e.Namespace} did not exist"});
             }
 
-
             return NoContent();
         }
 
@@ -852,7 +859,7 @@ namespace Horde.Storage.Controllers
                 return Forbid();
             }
 
-            long countOfDeletedRecords = 0;
+            long countOfDeletedRecords;
             try
             {
                 countOfDeletedRecords = await _objectService.DeleteBucket(ns, bucket);
@@ -862,7 +869,6 @@ namespace Horde.Storage.Controllers
                 return NotFound(new ProblemDetails {Title = $"Namespace {e.Namespace} did not exist"});
             }
 
-
             return Ok(new BucketDeletedResponse(countOfDeletedRecords));
         }
 
@@ -947,6 +953,7 @@ namespace Horde.Storage.Controllers
             Ops = Array.Empty<BatchOp>();
         }
 
+        [System.Diagnostics.CodeAnalysis.SuppressMessage("Design", "CA1034:Nested types should not be visible", Justification = "For serialization only")]
         public class BatchOp
         {
             public BatchOp()
@@ -967,13 +974,13 @@ namespace Horde.Storage.Controllers
             [CbField("opId")]
             public uint OpId { get; set; }
 
-
             [CbField("op")]
             [JsonIgnore]
+            [UsedImplicitly]
             public string OpString
             {
-                get { return Op.ToString(); }
-                set { Op = Enum.Parse<Operation>(value); }
+                get => Op.ToString();
+                set => Op = Enum.Parse<Operation>(value);
             }
 
             [Required]
@@ -1008,6 +1015,7 @@ namespace Horde.Storage.Controllers
 
         }
 
+        [System.Diagnostics.CodeAnalysis.SuppressMessage("Design", "CA1034:Nested types should not be visible", Justification = "For serialization only")]
         public class OpResponses
         {
             public OpResponses()
@@ -1025,6 +1033,7 @@ namespace Horde.Storage.Controllers
         }
 
         [CbField("results")]
+        [System.Diagnostics.CodeAnalysis.SuppressMessage("Usage", "CA2227:Collection properties should be read only", Justification = "Used by serialization")]
         public List<OpResponses> Results { get; set; } = new List<OpResponses>();
     }
 
@@ -1119,8 +1128,10 @@ namespace Horde.Storage.Controllers
         }
 
         [CbField("missing")]
+        [System.Diagnostics.CodeAnalysis.SuppressMessage("Usage", "CA2227:Collection properties should be read only", Justification = "Used by serialization")]
         public List<MissingReference> Missing { get; set; }
 
+        [System.Diagnostics.CodeAnalysis.SuppressMessage("Design", "CA1034:Nested types should not be visible", Justification = "For serialization only")]
         public class MissingReference
         {
             [CbField("bucket")]
diff --git a/Engine/Source/Programs/Horde/HordeStorage/Horde.Storage/Controllers/ReplicationController.cs b/Engine/Source/Programs/Horde/HordeStorage/Horde.Storage/Controllers/ReplicationController.cs
--- a/Engine/Source/Programs/Horde/HordeStorage/Horde.Storage/Controllers/ReplicationController.cs
+++ b/Engine/Source/Programs/Horde/HordeStorage/Horde.Storage/Controllers/ReplicationController.cs
@@ -30,7 +30,6 @@ namespace Horde.Storage.Controllers
             _authorizationService = authorizationService;
         }
 
-
         [HttpGet("{ns}")]
         [ProducesDefaultResponseType]
         public async Task<IActionResult> Get(
@@ -73,7 +72,6 @@ namespace Horde.Storage.Controllers
 
         }
 
-
         [HttpPost("{ns}/{replicatorName}/{offset}")]
         [ProducesDefaultResponseType]
         public async Task<IActionResult> SetReplicationOffset(
@@ -99,7 +97,7 @@ namespace Horde.Storage.Controllers
                 });
             }
 
-            IReplicator? replicator = replicators.FirstOrDefault(replicator => string.Equals(replicator.Info.ReplicatorName, replicatorName, StringComparison.InvariantCultureIgnoreCase));
+            IReplicator? replicator = replicators.FirstOrDefault(replicator => string.Equals(replicator.Info.ReplicatorName, replicatorName, StringComparison.OrdinalIgnoreCase));
             if (replicator == null)
             {
                 return BadRequest(new ValidationProblemDetails
@@ -138,7 +136,7 @@ namespace Horde.Storage.Controllers
                 });
             }
 
-            IReplicator? replicator = replicators.FirstOrDefault(replicator => string.Equals(replicator.Info.ReplicatorName, replicatorName, StringComparison.InvariantCultureIgnoreCase));
+            IReplicator? replicator = replicators.FirstOrDefault(replicator => string.Equals(replicator.Info.ReplicatorName, replicatorName, StringComparison.OrdinalIgnoreCase));
             if (replicator == null)
             {
                 return BadRequest(new ValidationProblemDetails
@@ -159,7 +157,6 @@ namespace Horde.Storage.Controllers
             return Ok();
         }
 
-
         [HttpDelete("{ns}")]
         [ProducesDefaultResponseType]
         public async Task<IActionResult> Delete(
@@ -188,7 +185,6 @@ namespace Horde.Storage.Controllers
                 await replicator.DeleteState();
             }
 
-
             return Ok();
 
         }
diff --git a/Engine/Source/Programs/Horde/HordeStorage/Horde.Storage/Controllers/ReplicationLogController.cs b/Engine/Source/Programs/Horde/HordeStorage/Horde.Storage/Controllers/ReplicationLogController.cs
--- a/Engine/Source/Programs/Horde/HordeStorage/Horde.Storage/Controllers/ReplicationLogController.cs
+++ b/Engine/Source/Programs/Horde/HordeStorage/Horde.Storage/Controllers/ReplicationLogController.cs
@@ -38,7 +38,6 @@ namespace Horde.Storage.Controllers
             _snapshotSettings = snapshotSettings;
         }
 
-
         [HttpGet("snapshots/{ns}")]
         [ProducesDefaultResponseType]
         [ProducesResponseType(type: typeof(ProblemDetails), 400)]
@@ -140,7 +139,6 @@ namespace Horde.Storage.Controllers
                     Title = $"Namespace {ns} was not found",
                 });
             }
-
         }
     }
 
@@ -173,6 +171,7 @@ namespace Horde.Storage.Controllers
             Snapshots = snapshots;
         }
 
+        [System.Diagnostics.CodeAnalysis.SuppressMessage("Usage", "CA2227:Collection properties should be read only", Justification = "Used by serialization")]
         public List<SnapshotInfo> Snapshots { get; set; }
     }
 
@@ -189,6 +188,7 @@ namespace Horde.Storage.Controllers
             Events = events;
         }
 
+        [System.Diagnostics.CodeAnalysis.SuppressMessage("Usage", "CA2227:Collection properties should be read only", Justification = "Used by serialization")]
         public List<ReplicationLogEvent> Events { get; set; }
     }
 }
diff --git a/Engine/Source/Programs/Horde/HordeStorage/Horde.Storage/Controllers/StatusController.cs b/Engine/Source/Programs/Horde/HordeStorage/Horde.Storage/Controllers/StatusController.cs
--- a/Engine/Source/Programs/Horde/HordeStorage/Horde.Storage/Controllers/StatusController.cs
+++ b/Engine/Source/Programs/Horde/HordeStorage/Horde.Storage/Controllers/StatusController.cs
@@ -1,5 +1,6 @@
 ﻿// Copyright Epic Games, Inc. All Rights Reserved.
 
+using System;
 using System.Collections.Generic;
 using System.Linq;
 using System.Reflection;
@@ -46,15 +47,20 @@ namespace Horde.Storage.Controllers
             string srcControlIdentifier = "Unknown";
             AssemblyMetadataAttribute? gitHashAttribute = attrs.FirstOrDefault(attr => attr.Key == "GitHash");
             if (gitHashAttribute?.Value != null && !string.IsNullOrEmpty(gitHashAttribute.Value))
+            {
                 srcControlIdentifier = gitHashAttribute.Value;
+            }
 
             AssemblyMetadataAttribute? p4ChangeAttribute = attrs.FirstOrDefault(attr => attr.Key == "PerforceChangelist");
             if (p4ChangeAttribute?.Value != null && !string.IsNullOrEmpty(p4ChangeAttribute.Value))
+            {
                 srcControlIdentifier = p4ChangeAttribute.Value;
+            }
+
             return Ok(new StatusResponse(_versionFile.VersionString ?? "Unknown", srcControlIdentifier, GetCapabilities(), _jupiterSettings.CurrentValue.CurrentSite));
         }
 
-        private string[] GetCapabilities()
+        private static string[] GetCapabilities()
         {
             return new string[]
             {
@@ -96,13 +102,14 @@ namespace Horde.Storage.Controllers
 
         public string CurrentSite { get; set; }
 
+        [System.Diagnostics.CodeAnalysis.SuppressMessage("Usage", "CA2227:Collection properties should be read only", Justification = "Used by serialization")]
         public List<KnownPeer> Peers { get; set; } = new List<KnownPeer>();
     }
 
     public class KnownPeer
     {
         [JsonConstructor]
-        public KnownPeer(string site, string fullName, List<string> endpoints, int latency)
+        public KnownPeer(string site, string fullName, List<Uri> endpoints, int latency)
         {
             Site = site;
             FullName = fullName;
@@ -116,10 +123,13 @@ namespace Horde.Storage.Controllers
             FullName = peerSettings.FullName;
             IEnumerable<PeerEndpoints> endpoints = peerSettings.Endpoints;
             if (!includeInternalEndpoints)
+            {
                 endpoints = endpoints.Where(s => !s.IsInternal);
+            }
+
             Endpoints = endpoints.Select(e => e.Url).ToList();
 
-            IPeerStatusService.PeerStatus? peerStatus = statusService.GetPeerStatus(peerSettings.Name);
+            PeerStatus? peerStatus = statusService.GetPeerStatus(peerSettings.Name);
             if (peerStatus != null)
             {
                 Latency = peerStatus.Latency;
@@ -131,7 +141,8 @@ namespace Horde.Storage.Controllers
 
         public int Latency { get; set; }
 
-        public List<string> Endpoints { get; set; }
+        [System.Diagnostics.CodeAnalysis.SuppressMessage("Usage", "CA2227:Collection properties should be read only", Justification = "Used by serialization")]
+        public List<Uri> Endpoints { get; set; }
     }
 
     public class StatusResponse
diff --git a/Engine/Source/Programs/Horde/HordeStorage/Horde.Storage/Controllers/StorageController.cs b/Engine/Source/Programs/Horde/HordeStorage/Horde.Storage/Controllers/StorageController.cs
--- a/Engine/Source/Programs/Horde/HordeStorage/Horde.Storage/Controllers/StorageController.cs
+++ b/Engine/Source/Programs/Horde/HordeStorage/Horde.Storage/Controllers/StorageController.cs
@@ -8,7 +8,6 @@ using System.IO;
 using System.Linq;
 using System.Net.Mime;
 using System.Threading.Tasks;
-using Datadog.Trace;
 using EpicGames.Horde.Storage;
 using EpicGames.Serialization;
 using Horde.Storage.Implementation;
@@ -18,7 +17,6 @@ using Jupiter.Implementation;
 using Microsoft.AspNetCore.Authorization;
 using Microsoft.AspNetCore.Http;
 using Microsoft.AspNetCore.Mvc;
-using Microsoft.Extensions.Options;
 using Serilog;
 
 namespace Horde.Storage.Controllers
@@ -34,20 +32,15 @@ namespace Horde.Storage.Controllers
         private readonly IDiagnosticContext _diagnosticContext;
         private readonly IAuthorizationService _authorizationService;
         private readonly BufferedPayloadFactory _bufferedPayloadFactory;
-        private readonly HordeStorageSettings _settings;
 
-        private readonly ILogger _logger = Log.ForContext<StorageController>();
-
-        public StorageController(IBlobService storage, IOptions<HordeStorageSettings> settings, IDiagnosticContext diagnosticContext, IAuthorizationService authorizationService, BufferedPayloadFactory bufferedPayloadFactory)
+        public StorageController(IBlobService storage, IDiagnosticContext diagnosticContext, IAuthorizationService authorizationService, BufferedPayloadFactory bufferedPayloadFactory)
         {
             _storage = storage;
             _diagnosticContext = diagnosticContext;
             _authorizationService = authorizationService;
             _bufferedPayloadFactory = bufferedPayloadFactory;
-            _settings = settings.Value;
         }
 
-
         [HttpGet("{ns}/{id}")]
         [Authorize("Storage.read")]
         [ProducesDefaultResponseType]
@@ -116,7 +109,9 @@ namespace Horde.Storage.Controllers
             IEnumerable<Task> tasks = id.Select(async blob =>
             {
                 if (!await _storage.Exists(ns, blob))
+                {
                     missingBlobs.Add(blob);
+                }
             });
             await Task.WhenAll(tasks);
 
@@ -142,7 +137,9 @@ namespace Horde.Storage.Controllers
             IEnumerable<Task> tasks = bodyIds.Select(async blob =>
             {
                 if (!await _storage.Exists(ns, blob))
+                {
                     missingBlobs.Add(blob);
+                }
             });
             await Task.WhenAll(tasks);
 
@@ -158,7 +155,10 @@ namespace Horde.Storage.Controllers
             catch (BlobNotFoundException)
             {
                 if (!_storage.ShouldFetchBlobOnDemand(ns))
+                {
                     throw;
+                }
+
                 return await _storage.ReplicateObject(ns, blob);
             }
         }
@@ -223,7 +223,6 @@ namespace Horde.Storage.Controllers
             });
         }
 
-
         [HttpDelete("{ns}/{id}")]
         [Authorize("Storage.delete")]
         public async Task<IActionResult> Delete(
@@ -260,7 +259,6 @@ namespace Horde.Storage.Controllers
             return NoContent();
         }
 
-
         private async Task DeleteImpl(NamespaceId ns, BlobIdentifier id)
         {
             await _storage.DeleteObject(ns, id);
@@ -268,6 +266,7 @@ namespace Horde.Storage.Controllers
 
         // ReSharper disable UnusedAutoPropertyAccessor.Global
         // ReSharper disable once ClassNeverInstantiated.Global
+        [System.Diagnostics.CodeAnalysis.SuppressMessage("Design", "CA1034:Nested types should not be visible", Justification = "Used by serialization")]
         public class BatchOp
         {
             // ReSharper disable once InconsistentNaming
@@ -289,6 +288,7 @@ namespace Horde.Storage.Controllers
             public byte[]? Content { get; set; }
         }
 
+        [System.Diagnostics.CodeAnalysis.SuppressMessage("Design", "CA1034:Nested types should not be visible", Justification = "Used by serialization")]
         public class BatchCall
         {
             public BatchOp[]? Operations { get; set; }
@@ -316,7 +316,7 @@ namespace Horde.Storage.Controllers
 
             if (batch?.Operations == null)
             {
-                throw new ArgumentNullException();
+                throw new ();
             }
 
             Task<object?>[] tasks = new Task<object?>[batch.Operations.Length];
@@ -325,7 +325,7 @@ namespace Horde.Storage.Controllers
                 BatchOp op = batch.Operations[index];
                 if (op.Namespace == null)
                 {
-                    throw new ArgumentNullException("namespace");
+                    throw new Exception(nameof(op.Namespace));
                 }
 
                 AuthorizationResult authorizationResultNamespace = await _authorizationService.AuthorizeAsync(User, op.Namespace, NamespaceAccessRequirement.Name);
@@ -339,14 +339,14 @@ namespace Horde.Storage.Controllers
                 switch (op.Op)
                 {
                     case BatchOp.Operation.INVALID:
-                        throw new ArgumentOutOfRangeException();
+                        throw new InvalidOperationException($"Op type {BatchOp.Operation.INVALID} is not a valid op type for operation id {op.Id}");
                     case BatchOp.Operation.GET:
                         if (op.Id == null)
                         {
-                            throw new ArgumentNullException("id");
+                            return BadRequest();
                         }
 
-                        tasks[index] = GetImpl(op.Namespace.Value, op.Id).ContinueWith(t =>
+                        tasks[index] = GetImpl(op.Namespace.Value, op.Id).ContinueWith((t, _) =>
                         {
                             // TODO: This is very allocation heavy but given that the end result is a json object we can not really stream this anyway
                             using BlobContents blobContents = t.Result;
@@ -356,16 +356,16 @@ namespace Horde.Storage.Controllers
                             ms.Seek(0, SeekOrigin.Begin);
                             string str = Convert.ToBase64String(ms.ToArray());
                             return (object?) str;
-                        });
+                        }, null, TaskScheduler.Current);
                         break;
                     case BatchOp.Operation.HEAD:
                         if (op.Id == null)
                         {
-                            throw new ArgumentNullException("id");
+                            return BadRequest();
                         }
 
                         tasks[index] = _storage.Exists(op.Namespace.Value, op.Id)
-                            .ContinueWith(t => t.Result ? (object?) null : op.Id);
+                            .ContinueWith((t,_) => t.Result ? (object?) null : op.Id, null, TaskScheduler.Current);
                         break;
                     case BatchOp.Operation.PUT:
                     {
@@ -380,19 +380,19 @@ namespace Horde.Storage.Controllers
                         }
 
                         using MemoryBufferedPayload payload = new MemoryBufferedPayload(op.Content);
-                        tasks[index] = _storage.PutObject(op.Namespace.Value, payload, op.Id).ContinueWith(t => (object?) t.Result);
+                        tasks[index] = _storage.PutObject(op.Namespace.Value, payload, op.Id).ContinueWith((t, _) => (object?) t.Result, null, TaskScheduler.Current);
                         break;
                     }
                     case BatchOp.Operation.DELETE:
                         if (op.Id == null)
                         {
-                            throw new ArgumentNullException("id");
+                            return BadRequest();
                         }
 
-                        tasks[index] = DeleteImpl(op.Namespace.Value, op.Id).ContinueWith(t => (object?) null);
+                        tasks[index] = DeleteImpl(op.Namespace.Value, op.Id).ContinueWith((t, _) => (object?) null, null, TaskScheduler.Current);
                         break;
                     default:
-                        throw new ArgumentOutOfRangeException();
+                        throw new NotImplementedException($"{op.Op} is not a support op type");
                 }
             }
 
@@ -402,7 +402,6 @@ namespace Horde.Storage.Controllers
 
             return Ok(results);
         }
-
     }
 
     public class HeadMultipleResponse
diff --git a/Engine/Source/Programs/Horde/HordeStorage/Horde.Storage/HealthChecks.cs b/Engine/Source/Programs/Horde/HordeStorage/Horde.Storage/HealthChecks.cs
--- a/Engine/Source/Programs/Horde/HordeStorage/Horde.Storage/HealthChecks.cs
+++ b/Engine/Source/Programs/Horde/HordeStorage/Horde.Storage/HealthChecks.cs
@@ -1,12 +1,10 @@
 // Copyright Epic Games, Inc. All Rights Reserved.
 
-using System;
 using System.Threading;
 using System.Threading.Tasks;
 using Horde.Storage.Implementation;
 using Horde.Storage.Implementation.LeaderElection;
 using Microsoft.Extensions.Diagnostics.HealthChecks;
-using Microsoft.Extensions.Hosting;
 using Microsoft.Extensions.Options;
 
 namespace Horde.Storage
@@ -106,7 +104,9 @@ namespace Horde.Storage
         public Task<HealthCheckResult> CheckHealthAsync(HealthCheckContext context, CancellationToken cancellationToken = new CancellationToken())
         {
             if (!_snapshotSettings.CurrentValue.Enabled)
+            {
                 return Task.FromResult(HealthCheckResult.Healthy());
+            }
 
             if (_replicationSnapshotService.Running)
             {
@@ -132,7 +132,9 @@ namespace Horde.Storage
         {
             // if replication is disabled we consider it healthy
             if (!_replicationSetting.CurrentValue.Enabled)
+            {
                 return Task.FromResult(HealthCheckResult.Healthy());
+            }
 
             if (_replicationService.Running)
             {
diff --git a/Engine/Source/Programs/Horde/HordeStorage/Horde.Storage/Horde.Storage.Startup.cs b/Engine/Source/Programs/Horde/HordeStorage/Horde.Storage/Horde.Storage.Startup.cs
--- a/Engine/Source/Programs/Horde/HordeStorage/Horde.Storage/Horde.Storage.Startup.cs
+++ b/Engine/Source/Programs/Horde/HordeStorage/Horde.Storage/Horde.Storage.Startup.cs
@@ -22,8 +22,6 @@ using Jupiter;
 using Jupiter.Common.Implementation;
 using Jupiter.Utils;
 using Microsoft.AspNetCore.Authorization;
-using Microsoft.AspNetCore.Builder;
-using Microsoft.AspNetCore.Hosting;
 using Microsoft.AspNetCore.Routing;
 using Microsoft.Extensions.Configuration;
 using Microsoft.Extensions.DependencyInjection;
@@ -35,13 +33,12 @@ namespace Horde.Storage
     // ReSharper disable once ClassNeverInstantiated.Global
     public class HordeStorageStartup : BaseStartup
     {
-        public HordeStorageStartup(IConfiguration configuration, IWebHostEnvironment environment) : base(configuration,
-            environment)
+        public HordeStorageStartup(IConfiguration configuration) : base(configuration)
         {
             string? ddAgentHost = System.Environment.GetEnvironmentVariable("DD_AGENT_HOST");
             if (!string.IsNullOrEmpty(ddAgentHost))
             {
-                _logger.Information("Initializing Dogstatsd to connect to: {DatadogAgentHost}", ddAgentHost);
+                Logger.Information("Initializing Dogstatsd to connect to: {DatadogAgentHost}", ddAgentHost);
                 StatsdConfig dogstatsdConfig = new StatsdConfig
                 {
                     StatsdServerName = ddAgentHost,
@@ -153,7 +150,6 @@ namespace Horde.Storage
                     .AddAny();
             });
 
-
             authorizationOptions.AddPolicy("Admin", policy =>
             {
                 policy.AuthenticationSchemes = defaultSchemes;
@@ -209,7 +205,6 @@ namespace Horde.Storage
 
             services.AddOptions<KubernetesLeaderElectionSettings>().Configure(o => Configuration.GetSection("Kubernetes").Bind(o)).ValidateDataAnnotations();
 
-
             services.AddSingleton(typeof(OodleCompressor), CreateOodleCompressor);
             services.AddSingleton(typeof(CompressedBufferUtils), CreateCompressedBufferUtils);
 
@@ -329,7 +324,7 @@ namespace Horde.Storage
                 case HordeStorageSettings.BlobIndexImplementations.Cache:
                     return ActivatorUtilities.CreateInstance<CachedBlobIndex>(provider);
                 default:
-                    throw new ArgumentOutOfRangeException();
+                    throw new NotImplementedException($"Unknown blob index implementation: {settings.BlobIndexImplementation}");
             }
         }
 
@@ -369,7 +364,7 @@ namespace Horde.Storage
                 case HordeStorageSettings.ContentIdStoreImplementations.Cache:
                     return ActivatorUtilities.CreateInstance<CacheContentIdStore>(provider);
                 default:
-                    throw new ArgumentOutOfRangeException();
+                    throw new NotImplementedException($"Unknown content id store implementation: {settings.ContentIdStoreImplementation}");
             }
         }
 
@@ -378,7 +373,7 @@ namespace Horde.Storage
             ScyllaSettings settings = provider.GetService<IOptionsMonitor<ScyllaSettings>>()!.CurrentValue!;
             ISecretResolver secretResolver = provider.GetService<ISecretResolver>()!;
 
-            Serilog.Extensions.Logging.SerilogLoggerProvider serilogLoggerProvider = new();
+            using Serilog.Extensions.Logging.SerilogLoggerProvider serilogLoggerProvider = new();
             Diagnostics.AddLoggerProvider(serilogLoggerProvider);
             const string DefaultKeyspaceName = "jupiter";
 
@@ -397,7 +392,6 @@ namespace Horde.Storage
                 .WithExecutionProfiles(options =>
                     options.WithProfile("default", builder => builder.WithConsistencyLevel(ConsistencyLevel.LocalOne)));
 
-
             if (settings.UseAzureCosmosDB)
             {
                 CassandraConnectionStringBuilder connectionStringBuilder = new CassandraConnectionStringBuilder(connectionString);
@@ -405,19 +399,21 @@ namespace Horde.Storage
                 string[] contactPoints = connectionStringBuilder.ContactPoints;
 
                 // Connect to cassandra cluster using TLSv1.2.
-                var sslOptions = new SSLOptions(SslProtocols.Tls12, false,
+                SSLOptions sslOptions = new SSLOptions(SslProtocols.None, false,
                     (sender, certificate, chain, sslPolicyErrors) =>
                     {
                         if (sslPolicyErrors == SslPolicyErrors.None)
+                        {
                             return true;
+                        }
 
-                        _logger.Error("Certificate error: {0}", sslPolicyErrors);
+                        Logger.Error("Certificate error: {0}", sslPolicyErrors);
                         // Do not allow this client to communicate with unauthenticated servers.
                         return false;
                     });
 
                 // Prepare a map to resolve the host name from the IP address.
-                var hostNameByIp = new Dictionary<IPAddress, string>();
+                Dictionary<IPAddress, string> hostNameByIp = new Dictionary<IPAddress, string>();
                 foreach (string contactPoint in contactPoints)
                 {
                     IPAddress[] resolvedIps = Dns.GetHostAddresses(contactPoint);
@@ -433,7 +429,7 @@ namespace Horde.Storage
                     {
                         return resolvedName;
                     }
-                    var hostEntry = Dns.GetHostEntry(ipAddress.ToString());
+                    IPHostEntry hostEntry = Dns.GetHostEntry(ipAddress.ToString());
                     return hostEntry.HostName;
                 });
 
@@ -445,7 +441,7 @@ namespace Horde.Storage
             if (settings.KeyspaceReplicationStrategy != null)
             {
                 replicationStrategy = settings.KeyspaceReplicationStrategy;
-                _logger.Information("Scylla Replication strategy for replicated keyspace is set to {@ReplicationStrategy}", replicationStrategy);
+                Logger.Information("Scylla Replication strategy for replicated keyspace is set to {@ReplicationStrategy}", replicationStrategy);
             }
             ISession replicatedSession = cluster.ConnectAndCreateDefaultKeyspaceIfNotExists(replicationStrategy);
 
@@ -461,7 +457,7 @@ namespace Horde.Storage
             if (settings.LocalKeyspaceReplicationStrategy != null)
             {
                 replicationStrategyLocal = settings.LocalKeyspaceReplicationStrategy;
-                _logger.Information("Scylla Replication strategy for local keyspace is set to {@ReplicationStrategy}", replicationStrategyLocal);
+                Logger.Information("Scylla Replication strategy for local keyspace is set to {@ReplicationStrategy}", replicationStrategyLocal);
             }
 
             // the local keyspace is never replicated so we do not support controlling how the replication strategy is setup
@@ -488,7 +484,7 @@ namespace Horde.Storage
                 case HordeStorageSettings.ReferencesDbImplementations.Cache:
                     return ActivatorUtilities.CreateInstance<CacheReferencesStore>(provider);
                 default:
-                    throw new ArgumentOutOfRangeException();
+                    throw new NotImplementedException($"Unknown references db implementation: {settings.ReferencesDbImplementation}");
             }
         }
         
@@ -516,7 +512,7 @@ namespace Horde.Storage
             }
             
             bool isS3InUse = settings.StorageImplementations.Any(x =>
-                String.Equals(x, HordeStorageSettings.StorageBackendImplementations.S3.ToString(), StringComparison.CurrentCultureIgnoreCase));
+                string.Equals(x, HordeStorageSettings.StorageBackendImplementations.S3.ToString(), StringComparison.OrdinalIgnoreCase));
 
             if (isS3InUse)
             {
@@ -538,7 +534,6 @@ namespace Horde.Storage
                     {
                         c.ForcePathStyle = true;
                     }
-
                 }
                 else
                 {
@@ -569,7 +564,9 @@ namespace Horde.Storage
 
             AWSOptions awsOptions = Configuration.GetAWSOptions();
             if (dynamoDbSettings.ConnectionString.ToUpper() != "AWS")
+            {
                 awsOptions.DefaultClientConfig.ServiceURL = dynamoDbSettings.ConnectionString;
+            }
 
             awsOptions.Credentials = awsCredentials;
 
@@ -674,7 +671,7 @@ namespace Horde.Storage
                     }
                     break;*/
                 default:
-                    throw new ArgumentOutOfRangeException();
+                    throw new NotImplementedException($"Unknown ref db implementation: {settings.RefDbImplementation} when adding health checks");
             }
 
             foreach (HordeStorageSettings.StorageBackendImplementations impl in settings.GetStorageImplementations())
diff --git a/Engine/Source/Programs/Horde/HordeStorage/Horde.Storage/Implementation/AnyRequirement.cs b/Engine/Source/Programs/Horde/HordeStorage/Horde.Storage/Implementation/AnyRequirement.cs
--- a/Engine/Source/Programs/Horde/HordeStorage/Horde.Storage/Implementation/AnyRequirement.cs
+++ b/Engine/Source/Programs/Horde/HordeStorage/Horde.Storage/Implementation/AnyRequirement.cs
@@ -31,7 +31,7 @@ namespace Horde.Storage
                 IAuthorizationHandler? handler = innerRequirement as IAuthorizationHandler;
                 if (handler != null)
                 {
-                    var innerContext = new AuthorizationHandlerContext(new[] { innerRequirement }, context.User, context.Resource);
+                    AuthorizationHandlerContext innerContext = new AuthorizationHandlerContext(new[] { innerRequirement }, context.User, context.Resource);
                     await handler.HandleAsync(innerContext);
                     if (innerContext.HasSucceeded)
                     {
@@ -64,7 +64,7 @@ namespace Horde.Storage
 
         public AuthorizationPolicyBuilder AddAny()
         {
-            var anyRequirement = new AnyRequirement(_innerRequirements.ToArray());
+            AnyRequirement anyRequirement = new AnyRequirement(_innerRequirements.ToArray());
             return _policyBuilder.AddRequirements(anyRequirement);
         }
     }
diff --git a/Engine/Source/Programs/Horde/HordeStorage/Horde.Storage/Implementation/Blob/AmazonS3Store.cs b/Engine/Source/Programs/Horde/HordeStorage/Horde.Storage/Implementation/Blob/AmazonS3Store.cs
--- a/Engine/Source/Programs/Horde/HordeStorage/Horde.Storage/Implementation/Blob/AmazonS3Store.cs
+++ b/Engine/Source/Programs/Horde/HordeStorage/Horde.Storage/Implementation/Blob/AmazonS3Store.cs
@@ -14,14 +14,12 @@ using Jupiter.Common;
 using Jupiter.Common.Implementation;
 using Jupiter.Implementation;
 using Microsoft.Extensions.Options;
-using Serilog;
 using KeyNotFoundException = System.Collections.Generic.KeyNotFoundException;
 
 namespace Horde.Storage.Implementation
 {
     public class AmazonS3Store : IBlobStore
     {
-        private readonly ILogger _logger = Log.ForContext<AmazonS3Store>();
         private readonly IAmazonS3 _amazonS3;
         private readonly IBlobIndex _blobIndex;
         private readonly INamespacePolicyResolver _namespacePolicyResolver;
@@ -54,7 +52,7 @@ namespace Horde.Storage.Implementation
                     bool bucketExist = await _amazonS3.DoesS3BucketExistAsync(bucketName);
                     if (!bucketExist)
                     {
-                        var putBucketRequest = new PutBucketRequest
+                        PutBucketRequest putBucketRequest = new PutBucketRequest
                         {
                             BucketName = bucketName,
                             UseClientRegion = true
@@ -114,7 +112,10 @@ namespace Horde.Storage.Implementation
                 }
 
                 if (e.StatusCode == HttpStatusCode.TooManyRequests)
+                {
                     throw new ResourceHasToManyRequestsException(e);
+                }
+
                 throw;
             }
 
@@ -141,7 +142,6 @@ namespace Horde.Storage.Implementation
             }
         }
 
-
         public async Task<BlobContents> GetObject(NamespaceId ns, BlobIdentifier blob)
         {
             string bucketName = GetBucketName(ns);
@@ -168,7 +168,7 @@ namespace Horde.Storage.Implementation
 
         public async Task<bool> Exists(NamespaceId ns, BlobIdentifier blobIdentifier)
         {
-            NamespaceSettings.PerNamespaceSettings policies = _namespacePolicyResolver.GetPoliciesForNs(ns);
+            NamespacePolicy policies = _namespacePolicyResolver.GetPoliciesForNs(ns);
             if (_settings.UseBlobIndexForExistsCheck && policies.UseBlobIndexForSlowExists)
             {
                 return await _blobIndex.BlobExistsInRegion(ns, blobIdentifier);
diff --git a/Engine/Source/Programs/Horde/HordeStorage/Horde.Storage/Implementation/Blob/AzureBlobStore.cs b/Engine/Source/Programs/Horde/HordeStorage/Horde.Storage/Implementation/Blob/AzureBlobStore.cs
--- a/Engine/Source/Programs/Horde/HordeStorage/Horde.Storage/Implementation/Blob/AzureBlobStore.cs
+++ b/Engine/Source/Programs/Horde/HordeStorage/Horde.Storage/Implementation/Blob/AzureBlobStore.cs
@@ -96,7 +96,7 @@ namespace Horde.Storage.Implementation
             return await PutObject(ns, stream, blobIdentifier);
         }
 
-        private async Task TouchBlob(BlobClient blob)
+        private static async Task TouchBlob(BlobClient blob)
         {
             Dictionary<string, string> metadata = new Dictionary<string, string>
             {
@@ -126,7 +126,9 @@ namespace Horde.Storage.Implementation
             catch (RequestFailedException e)
             {
                 if (e.Status == 404)
+                {
                     throw new BlobNotFoundException(ns, blobIdentifier);
+                }
 
                 throw;
             }
@@ -173,9 +175,11 @@ namespace Horde.Storage.Implementation
             BlobContainerClient container = new BlobContainerClient(_connectionString, fixedNamespace);
             bool exists = await container.ExistsAsync();
             if (!exists)
+            {
                 yield break;
+            }
 
-            await foreach (var item in container.GetBlobsAsync(BlobTraits.Metadata))
+            await foreach (BlobItem? item in container.GetBlobsAsync(BlobTraits.Metadata))
             {
                 yield return (new BlobIdentifier(item.Name), item.Properties?.LastModified?.DateTime ?? DateTime.Now);
             }
@@ -183,7 +187,7 @@ namespace Horde.Storage.Implementation
 
         private static string SanitizeNamespace(NamespaceId ns)
         {
-            return ns.ToString().Replace(".", "-").ToLower();
+            return ns.ToString().Replace(".", "-", StringComparison.OrdinalIgnoreCase).ToLower();
         }
     }
 }
