Horde Storage - Updated editorconfig to more closely match what we run for Horde and updated code to fix all static analysis errors

diff --git a/Engine/Source/Programs/Horde/HordeStorage/Horde.Storage/Implementation/Blob/BlobIndexConsistencyCheckService.cs b/Engine/Source/Programs/Horde/HordeStorage/Horde.Storage/Implementation/Blob/BlobIndexConsistencyCheckService.cs
--- a/Engine/Source/Programs/Horde/HordeStorage/Horde.Storage/Implementation/Blob/BlobIndexConsistencyCheckService.cs
+++ b/Engine/Source/Programs/Horde/HordeStorage/Horde.Storage/Implementation/Blob/BlobIndexConsistencyCheckService.cs
@@ -15,8 +15,12 @@ using Serilog;
 
 namespace Horde.Storage.Implementation
 {
+    public class BlobIndexConsistencyState
+    {
+    }
+
     // ReSharper disable once ClassNeverInstantiated.Global
-    public class BlobIndexConsistencyCheckService : PollingService<BlobIndexConsistencyCheckService.ConsistencyState>
+    public class BlobIndexConsistencyCheckService : PollingService<BlobIndexConsistencyState>
     {
         private readonly IOptionsMonitor<ConsistencyCheckSettings> _settings;
         private readonly IOptionsMonitor<JupiterSettings> _jupiterSettings;
@@ -26,17 +30,12 @@ namespace Horde.Storage.Implementation
         private readonly INamespacePolicyResolver _namespacePolicyResolver;
         private readonly ILogger _logger = Log.ForContext<BlobIndexConsistencyCheckService>();
 
-        public class ConsistencyState
-        {
-        }
-
         protected override bool ShouldStartPolling()
         {
             return _settings.CurrentValue.EnableBlobIndexChecks;
         }
 
-        public BlobIndexConsistencyCheckService(IOptionsMonitor<ConsistencyCheckSettings> settings, IOptionsMonitor<JupiterSettings> jupiterSettings, ILeaderElection leaderElection, IBlobIndex blobIndex, IBlobService blobService, INamespacePolicyResolver namespacePolicyResolver) :
-            base(serviceName: nameof(BlobStoreConsistencyCheckService), TimeSpan.FromSeconds(settings.CurrentValue.ConsistencyCheckPollFrequencySeconds), new ConsistencyState())
+        public BlobIndexConsistencyCheckService(IOptionsMonitor<ConsistencyCheckSettings> settings, IOptionsMonitor<JupiterSettings> jupiterSettings, ILeaderElection leaderElection, IBlobIndex blobIndex, IBlobService blobService, INamespacePolicyResolver namespacePolicyResolver) : base(serviceName: nameof(BlobStoreConsistencyCheckService), TimeSpan.FromSeconds(settings.CurrentValue.ConsistencyCheckPollFrequencySeconds), new BlobIndexConsistencyState())
         {
             _settings = settings;
             _jupiterSettings = jupiterSettings;
@@ -46,7 +45,7 @@ namespace Horde.Storage.Implementation
             _namespacePolicyResolver = namespacePolicyResolver;
         }
 
-        public override async Task<bool> OnPoll(ConsistencyState state, CancellationToken cancellationToken)
+        public override async Task<bool> OnPoll(BlobIndexConsistencyState state, CancellationToken cancellationToken)
         {
             if (!_settings.CurrentValue.EnableBlobIndexChecks)
             {
@@ -69,7 +68,7 @@ namespace Horde.Storage.Implementation
         {
             try
             {
-                NamespaceSettings.PerNamespaceSettings policy = _namespacePolicyResolver.GetPoliciesForNs(ns);
+                NamespacePolicy policy = _namespacePolicyResolver.GetPoliciesForNs(ns);
 
                 return policy.IsLegacyNamespace.HasValue && !policy.IsLegacyNamespace.Value;
             }
@@ -78,7 +77,6 @@ namespace Horde.Storage.Implementation
                 _logger.Warning("Namespace {Namespace} does not configure any policy, not running blob index consistency checks on it.", ns);
                 return false;
             }
-
         }
 
         private async Task RunConsistencyCheck()
@@ -95,12 +93,16 @@ namespace Horde.Storage.Implementation
                     Interlocked.Increment(ref countOfBlobsChecked);
 
                     if (countOfBlobsChecked % 100 == 0)
+                    {
                         _logger.Information("Consistency check running on blob index, count of blobs processed so far: {CountOfBlobs}", countOfBlobsChecked);
+                    }
 
                     // skip namespace that we shouldn't consistency check, e.g. the legacy namespaces
                     if (!NamespaceShouldBeCheckedForConsistency(blobInfo.Namespace))
+                    {
                         return;
-                    
+                    }
+
                     using IScope scope = Tracer.Instance.StartActive("consistency_check.blob_index");
                     scope.Span.ResourceName = $"{blobInfo.Namespace}.{blobInfo.BlobIdentifier}";
 
@@ -205,7 +207,7 @@ namespace Horde.Storage.Implementation
             _logger.Information("Blob Index Consistency check finished, found {CountOfIncorrectBlobs} incorrect blobs. Processed {CountOfBlobs} blobs.", countOfIncorrectBlobsFound, countOfBlobsChecked);
         }
 
-        protected override Task OnStopping(ConsistencyState state)
+        protected override Task OnStopping(BlobIndexConsistencyState state)
         {
             return Task.CompletedTask;
         }
diff --git a/Engine/Source/Programs/Horde/HordeStorage/Horde.Storage/Implementation/Blob/BlobStoreConsistencyCheckService.cs b/Engine/Source/Programs/Horde/HordeStorage/Horde.Storage/Implementation/Blob/BlobStoreConsistencyCheckService.cs
--- a/Engine/Source/Programs/Horde/HordeStorage/Horde.Storage/Implementation/Blob/BlobStoreConsistencyCheckService.cs
+++ b/Engine/Source/Programs/Horde/HordeStorage/Horde.Storage/Implementation/Blob/BlobStoreConsistencyCheckService.cs
@@ -12,14 +12,17 @@ using EpicGames.Horde.Storage;
 using Horde.Storage.Implementation.Blob;
 using Jupiter;
 using Jupiter.Implementation;
-using Microsoft.Extensions.DependencyInjection;
 using Microsoft.Extensions.Options;
 using Serilog;
 
 namespace Horde.Storage.Implementation
 {
+    public class ConsistencyState
+    {
+    }
+
     // ReSharper disable once ClassNeverInstantiated.Global
-    public class BlobStoreConsistencyCheckService : PollingService<BlobStoreConsistencyCheckService.ConsistencyState>
+    public class BlobStoreConsistencyCheckService : PollingService<ConsistencyState>
     {
         private readonly IOptionsMonitor<ConsistencyCheckSettings> _settings;
         private readonly IOptionsMonitor<HordeStorageSettings> _hordeStorageSettings;
@@ -30,17 +33,12 @@ namespace Horde.Storage.Implementation
         private readonly IBlobIndex _blobIndex;
         private readonly ILogger _logger = Log.ForContext<BlobStoreConsistencyCheckService>();
 
-        public class ConsistencyState
-        {
-        }
-
         protected override bool ShouldStartPolling()
         {
             return _settings.CurrentValue.EnableBlobStoreChecks;
         }
 
-        public BlobStoreConsistencyCheckService(IOptionsMonitor<ConsistencyCheckSettings> settings, IOptionsMonitor<HordeStorageSettings> hordeStorageSettings, IServiceProvider provider, ILeaderElection leaderElection, IRefsStore refsStore, IReferencesStore referencesStore, IBlobIndex blobIndex) :
-            base(serviceName: nameof(BlobStoreConsistencyCheckService), TimeSpan.FromSeconds(settings.CurrentValue.ConsistencyCheckPollFrequencySeconds), new ConsistencyState())
+        public BlobStoreConsistencyCheckService(IOptionsMonitor<ConsistencyCheckSettings> settings, IOptionsMonitor<HordeStorageSettings> hordeStorageSettings, IServiceProvider provider, ILeaderElection leaderElection, IRefsStore refsStore, IReferencesStore referencesStore, IBlobIndex blobIndex) : base(serviceName: nameof(BlobStoreConsistencyCheckService), TimeSpan.FromSeconds(settings.CurrentValue.ConsistencyCheckPollFrequencySeconds), new ConsistencyState())
         {
             _settings = settings;
             _hordeStorageSettings = hordeStorageSettings;
@@ -74,7 +72,9 @@ namespace Horde.Storage.Implementation
                 bool requiresLeader = blobStore is not FileSystemStore;
 
                 if (!_settings.CurrentValue.RunBlobStoreConsistencyCheckOnRootStore && isRootStore)
+                {
                     continue;
+                }
 
                 if (requiresLeader && !_leaderElection.IsThisInstanceLeader())
                 {
@@ -98,7 +98,10 @@ namespace Horde.Storage.Implementation
                         scope.Span.SetTag("BlobStore", blobStoreName);
 
                         if (countOfBlobsChecked % 100 == 0)
+                        {
                             _logger.Information("Consistency check running on Blob Store {BlobStore}, count of blobs processed so far: {CountOfBlobs}", blobStoreName, countOfBlobsChecked);
+                        }
+
                         Interlocked.Increment(ref countOfBlobsChecked);
                         
                         BlobContents contents = await blobStore.GetObject(ns, blob);
diff --git a/Engine/Source/Programs/Horde/HordeStorage/Horde.Storage/Implementation/Blob/CachedBlobIndex.cs b/Engine/Source/Programs/Horde/HordeStorage/Horde.Storage/Implementation/Blob/CachedBlobIndex.cs
--- a/Engine/Source/Programs/Horde/HordeStorage/Horde.Storage/Implementation/Blob/CachedBlobIndex.cs
+++ b/Engine/Source/Programs/Horde/HordeStorage/Horde.Storage/Implementation/Blob/CachedBlobIndex.cs
@@ -21,11 +21,11 @@ public class CachedBlobIndex : IBlobIndex
         _fileSystemStore = fileSystemStore;
     }
 
-    public async Task<IBlobIndex.BlobInfo?> GetBlobInfo(NamespaceId ns, BlobIdentifier id)
+    public async Task<BlobInfo?> GetBlobInfo(NamespaceId ns, BlobIdentifier id)
     {
         if (await BlobExistsInRegion(ns, id))
         {
-            return new IBlobIndex.BlobInfo
+            return new BlobInfo
             {
                 Namespace = ns,
                 BlobIdentifier = id,
@@ -67,7 +67,7 @@ public class CachedBlobIndex : IBlobIndex
         await Task.CompletedTask;
     }
 
-    public IAsyncEnumerable<IBlobIndex.BlobInfo> GetAllBlobs()
+    public IAsyncEnumerable<BlobInfo> GetAllBlobs()
     {
         throw new NotImplementedException();
     }
diff --git a/Engine/Source/Programs/Horde/HordeStorage/Horde.Storage/Implementation/Blob/FileSystemStore.cs b/Engine/Source/Programs/Horde/HordeStorage/Horde.Storage/Implementation/Blob/FileSystemStore.cs
--- a/Engine/Source/Programs/Horde/HordeStorage/Horde.Storage/Implementation/Blob/FileSystemStore.cs
+++ b/Engine/Source/Programs/Horde/HordeStorage/Horde.Storage/Implementation/Blob/FileSystemStore.cs
@@ -20,9 +20,7 @@ namespace Horde.Storage.Implementation
         private readonly ILogger _logger = Log.ForContext<FileSystemStore>();
         private readonly IOptionsMonitor<FilesystemSettings> _settings;
 
-        internal const int DefaultBufferSize = 4096;
-
-        private const int SharingViolationErrorCode = 32;
+        private const int DefaultBufferSize = 4096;
 
         public FileSystemStore(IOptionsMonitor<FilesystemSettings> settings)
         {
@@ -55,7 +53,7 @@ namespace Horde.Storage.Implementation
             return new DirectoryInfo(Path.Combine(GetRootDir(), ns.ToString()));
         }
 
-        static string s_processSuffix = Guid.NewGuid().ToString();
+        static readonly string s_processSuffix = Guid.NewGuid().ToString();
         static int s_uniqueId = 0;
 
         public async Task<BlobIdentifier> PutObject(NamespaceId ns, ReadOnlyMemory<byte> content, BlobIdentifier blobIdentifier)
@@ -91,7 +89,9 @@ namespace Horde.Storage.Implementation
                 filePath.Refresh();
 
                 if (filePath.Length == 0)
+                {
                     _logger.Warning("0 byte file written as {Blob} {Namespace} {Method}", blobIdentifier, ns, "Stream");
+                }
             }
 
             UpdateLastWriteTime(filePath.FullName, DateTime.UnixEpoch);
@@ -109,7 +109,9 @@ namespace Horde.Storage.Implementation
             FileInfo filePath = GetFilesystemPath(ns, blob);
 
             if (!filePath.Exists)
+            {
                 throw new BlobNotFoundException(ns, blob);
+            }
 
             UpdateLastWriteTime(filePath.FullName, DateTime.UtcNow);
             FileStream fs = new FileStream(filePath.FullName, FileMode.Open, FileAccess.Read, FileShare.Read, DefaultBufferSize, FileOptions.Asynchronous | FileOptions.SequentialScan);
@@ -130,7 +132,7 @@ namespace Horde.Storage.Implementation
         /// </summary>
         /// <param name="filePath"></param>
         /// <param name="lastAccessed">Time the file was last accessed</param>
-        private void UpdateLastWriteTime(string filePath, DateTime lastAccessed)
+        private static void UpdateLastWriteTime(string filePath, DateTime lastAccessed)
         {
             try
             {
@@ -160,7 +162,9 @@ namespace Horde.Storage.Implementation
         {
             DirectoryInfo namespaceDirectory = GetFilesystemPath(ns);
             if (namespaceDirectory.Exists)
+            {
                 namespaceDirectory.Delete(true);
+            }
 
             return Task.CompletedTask;
         }
diff --git a/Engine/Source/Programs/Horde/HordeStorage/Horde.Storage/Implementation/Blob/IBlobIndex.cs b/Engine/Source/Programs/Horde/HordeStorage/Horde.Storage/Implementation/Blob/IBlobIndex.cs
--- a/Engine/Source/Programs/Horde/HordeStorage/Horde.Storage/Implementation/Blob/IBlobIndex.cs
+++ b/Engine/Source/Programs/Horde/HordeStorage/Horde.Storage/Implementation/Blob/IBlobIndex.cs
@@ -7,16 +7,16 @@ using Jupiter.Implementation;
 
 namespace Horde.Storage.Implementation.Blob
 {
-    public interface IBlobIndex
+    public class BlobInfo
     {
-        public class BlobInfo
-        {
-            public HashSet<string> Regions { get; init; } = new HashSet<string>();
-            public NamespaceId Namespace { get; init; }
-            public BlobIdentifier BlobIdentifier { get; init; } = null!;
-            public List<(BucketId, IoHashKey)> References { get; init; } = new List<(BucketId, IoHashKey)>();
-        }
+        public HashSet<string> Regions { get; init; } = new HashSet<string>();
+        public NamespaceId Namespace { get; init; }
+        public BlobIdentifier BlobIdentifier { get; init; } = null!;
+        public List<(BucketId, IoHashKey)> References { get; init; } = new List<(BucketId, IoHashKey)>();
+    }
 
+    public interface IBlobIndex
+    {
         Task AddBlobToIndex(NamespaceId ns, BlobIdentifier id, string? region = null);
 
         Task<BlobInfo?> GetBlobInfo(NamespaceId ns, BlobIdentifier id);
diff --git a/Engine/Source/Programs/Horde/HordeStorage/Horde.Storage/Implementation/Blob/IBlobService.cs b/Engine/Source/Programs/Horde/HordeStorage/Horde.Storage/Implementation/Blob/IBlobService.cs
--- a/Engine/Source/Programs/Horde/HordeStorage/Horde.Storage/Implementation/Blob/IBlobService.cs
+++ b/Engine/Source/Programs/Horde/HordeStorage/Horde.Storage/Implementation/Blob/IBlobService.cs
@@ -46,7 +46,6 @@ public interface IBlobService
     // Delete a object
     Task DeleteObject(NamespaceId ns, BlobIdentifier blob);
 
-
     // delete the whole namespace
     Task DeleteNamespace(NamespaceId ns);
 
@@ -71,8 +70,8 @@ public class BlobService : IBlobService
 
     internal IEnumerable<IBlobStore> BlobStore
     {
-        get { return _blobStores; }
-        set { _blobStores = value.ToList(); } 
+        get => _blobStores;
+        set => _blobStores = value.ToList();
     }
 
     public BlobService(IServiceProvider provider, IOptionsMonitor<HordeStorageSettings> settings, IBlobIndex blobIndex, IPeerStatusService peerStatusService, IHttpClientFactory httpClientFactory, IServiceCredentials serviceCredentials, INamespacePolicyResolver namespacePolicyResolver)
@@ -101,10 +100,12 @@ public class BlobService : IBlobService
             HordeStorageSettings.StorageBackendImplementations.Memory => provider.GetService<MemoryCacheBlobStore>(),
             HordeStorageSettings.StorageBackendImplementations.MemoryBlobStore => provider.GetService<MemoryBlobStore>(),
             HordeStorageSettings.StorageBackendImplementations.Relay => provider.GetService<RelayBlobStore>(),
-            _ => throw new ArgumentOutOfRangeException()
+            _ => throw new NotImplementedException("Unknown blob store {store")
         };
         if (store == null)
+        {
             throw new ArgumentException($"Failed to find a provider service for type: {impl}");
+        }
 
         return store;
     }
@@ -180,7 +181,6 @@ public class BlobService : IBlobService
             await store.PutObject(ns, s, identifier);
         }
 
-
         return identifier;
     }
 
@@ -249,14 +249,16 @@ public class BlobService : IBlobService
             // Don't populate the last store, as that is where we got the hit
             for (int i = 0; i < numStoreMisses; i++)
             {
-                var blobStore = _blobStores[i];
+                IBlobStore blobStore = _blobStores[i];
                 using IScope scope = Tracer.Instance.StartActive("HierarchicalStore.PopulateStore");
                 scope.Span.SetTag("BlobStore", blobStore.GetType().Name);
                 // Populate each store traversed that did not have the content found lower in the hierarchy
                 await blobStore.PutObject(ns, data, blob);
             }
 
-            blobContents = new BlobContents(new MemoryStream(data), data.Length);
+#pragma warning disable CA2000 // Dispose objects before losing scope , ownership is transfered to caller
+            blobContents = new BlobContents(data);
+#pragma warning restore CA2000 // Dispose objects before losing scope
         }
         
         return blobContents;
@@ -265,16 +267,22 @@ public class BlobService : IBlobService
     public async Task<BlobContents> ReplicateObject(NamespaceId ns, BlobIdentifier blob, bool force = false)
     {
         if (!force && !ShouldFetchBlobOnDemand(ns))
+        {
             throw new NotSupportedException($"Replication is not allowed in namespace {ns}");
+        }
 
         using IScope scope = Tracer.Instance.StartActive("HierarchicalStore.Replicate");
 
-        IBlobIndex.BlobInfo? blobInfo = await _blobIndex.GetBlobInfo(ns, blob);
+        BlobInfo? blobInfo = await _blobIndex.GetBlobInfo(ns, blob);
         if (blobInfo == null)
+        {
             throw new BlobNotFoundException(ns, blob);
+        }
 
         if (!blobInfo.Regions.Any())
+        {
             throw new BlobReplicationException(ns, blob, "Blob not found in any region");
+        }
 
         _logger.Information("On-demand replicating blob {Blob} in Namespace {Namespace}", blob, ns);
         List<(int, string)> possiblePeers = new List<(int, string)>(_peerStatusService.GetPeersByLatency(blobInfo.Regions.ToList()));
@@ -282,15 +290,17 @@ public class BlobService : IBlobService
         bool replicated = false;
         foreach ((int latency, string? region) in possiblePeers)
         {
-            IPeerStatusService.PeerStatus? peerStatus = _peerStatusService.GetPeerStatus(region);
+            PeerStatus? peerStatus = _peerStatusService.GetPeerStatus(region);
             if (peerStatus == null)
+            {
                 throw new Exception($"Failed to find peer {region}");
+            }
 
             _logger.Information("Attempting to replicate blob {Blob} in Namespace {Namespace} from {Region}", blob, ns, region);
 
             PeerEndpoints peerEndpoint = peerStatus.Endpoints.First();
-            HttpClient httpClient = _httpClientFactory.CreateClient();
-            HttpRequestMessage blobRequest = BuildHttpRequest(HttpMethod.Get, $"{peerEndpoint.Url}/api/v1/blobs/{ns}/{blob}");
+            using HttpClient httpClient = _httpClientFactory.CreateClient();
+            using HttpRequestMessage blobRequest = BuildHttpRequest(HttpMethod.Get, new Uri($"{peerEndpoint.Url}/api/v1/blobs/{ns}/{blob}"));
             HttpResponseMessage blobResponse = await httpClient.SendAsync(blobRequest, HttpCompletionOption.ResponseHeadersRead);
 
             if (blobResponse.StatusCode == HttpStatusCode.NotFound)
@@ -312,7 +322,9 @@ public class BlobService : IBlobService
         }
 
         if (!replicated)
+        {
             throw new BlobReplicationException(ns, blob, $"Failed to replicate {blob} in {ns} due to it not existing in any region");
+        }
 
         return await GetObject(ns, blob);
     }
@@ -322,12 +334,15 @@ public class BlobService : IBlobService
         return _settings.CurrentValue.EnableOnDemandReplication && _namespacePolicyResolver.GetPoliciesForNs(ns).OnDemandReplication;
     }
 
-    private HttpRequestMessage BuildHttpRequest(HttpMethod httpMethod, string uri)
+    private HttpRequestMessage BuildHttpRequest(HttpMethod httpMethod, Uri uri)
     {
         string? token = _serviceCredentials.GetToken();
         HttpRequestMessage request = new HttpRequestMessage(httpMethod, uri);
         if (!string.IsNullOrEmpty(token))
+        {
             request.Headers.Add("Authorization", "Bearer " + token);
+        }
+
         return request;
     }
 
@@ -412,10 +427,14 @@ public class BlobService : IBlobService
         await _blobIndex.RemoveBlobFromRegion(ns, blob);
 
         if (deletedAtLeastOnce)
+        {
             return;
+        }
 
         if (blobNotFound)
+        {
             throw new BlobNotFoundException(ns, blob);
+        }
 
         throw new NamespaceNotFoundException(ns);
     }
@@ -439,7 +458,9 @@ public class BlobService : IBlobService
         }
 
         if (deletedAtLeastOnce)
+        {
             return;
+        }
 
         throw new NamespaceNotFoundException(ns);
     }
@@ -454,7 +475,7 @@ public class BlobService : IBlobService
     {
         var tasks = blobs.Select(async blobIdentifier => new { BlobIdentifier = blobIdentifier, Exists = await Exists(ns, blobIdentifier) });
         var blobResults = await Task.WhenAll(tasks);
-        var filteredBlobs = blobResults.Where(ac => !ac.Exists).Select(ac => ac.BlobIdentifier);
+        IEnumerable<BlobIdentifier> filteredBlobs = blobResults.Where(ac => !ac.Exists).Select(ac => ac.BlobIdentifier);
         return filteredBlobs.ToArray();
     }
 
@@ -477,7 +498,9 @@ public class BlobService : IBlobService
         catch (ParallelForEachException e)
         {
             if (e.InnerException is PartialReferenceResolveException)
+            {
                 throw e.InnerException;
+            }
 
             throw;
         }
@@ -506,6 +529,4 @@ public class BlobService : IBlobService
 
         return new BlobContents(ms, ms.Length);
     }
-
-
 }
diff --git a/Engine/Source/Programs/Horde/HordeStorage/Horde.Storage/Implementation/Blob/IBlobStore.cs b/Engine/Source/Programs/Horde/HordeStorage/Horde.Storage/Implementation/Blob/IBlobStore.cs
--- a/Engine/Source/Programs/Horde/HordeStorage/Horde.Storage/Implementation/Blob/IBlobStore.cs
+++ b/Engine/Source/Programs/Horde/HordeStorage/Horde.Storage/Implementation/Blob/IBlobStore.cs
@@ -1,14 +1,9 @@
 ﻿// Copyright Epic Games, Inc. All Rights Reserved.
 
 using System;
-using System.Collections.Concurrent;
 using System.Collections.Generic;
 using System.IO;
-using System.Linq;
 using System.Threading.Tasks;
-using async_enumerable_dotnet;
-using Dasync.Collections;
-using Datadog.Trace;
 using EpicGames.Horde.Storage;
 using Jupiter.Implementation;
 
@@ -26,7 +21,6 @@ namespace Horde.Storage.Implementation
         // Delete a object
         Task DeleteObject(NamespaceId ns, BlobIdentifier blob);
 
-
         // delete the whole namespace
         Task DeleteNamespace(NamespaceId ns);
 
@@ -66,7 +60,6 @@ namespace Horde.Storage.Implementation
         {
             Namespace = @namespace;
         }
-
     }
 
     public class BlobToLargeException : Exception
@@ -77,7 +70,6 @@ namespace Horde.Storage.Implementation
         {
             Blob = blob;
         }
-
     }
 
     public class ResourceHasToManyRequestsException : Exception
@@ -86,6 +78,4 @@ namespace Horde.Storage.Implementation
         {
         }
     }
-
-
 }
diff --git a/Engine/Source/Programs/Horde/HordeStorage/Horde.Storage/Implementation/Blob/MemoryBlobIndex.cs b/Engine/Source/Programs/Horde/HordeStorage/Horde.Storage/Implementation/Blob/MemoryBlobIndex.cs
--- a/Engine/Source/Programs/Horde/HordeStorage/Horde.Storage/Implementation/Blob/MemoryBlobIndex.cs
+++ b/Engine/Source/Programs/Horde/HordeStorage/Horde.Storage/Implementation/Blob/MemoryBlobIndex.cs
@@ -1,6 +1,5 @@
 // Copyright Epic Games, Inc. All Rights Reserved.
 
-using System;
 using System.Collections.Concurrent;
 using System.Collections.Generic;
 using System.Threading.Tasks;
@@ -13,7 +12,7 @@ namespace Horde.Storage.Implementation.Blob;
 
 public class MemoryBlobIndex : IBlobIndex
 {
-    private class MemoryBlobInfo : IBlobIndex.BlobInfo
+    private class MemoryBlobInfo : BlobInfo
     {
 
     }
@@ -39,14 +38,16 @@ public class MemoryBlobIndex : IBlobIndex
         return Task.CompletedTask;
     }
 
-    public Task<IBlobIndex.BlobInfo?> GetBlobInfo(NamespaceId ns, BlobIdentifier id)
+    public Task<BlobInfo?> GetBlobInfo(NamespaceId ns, BlobIdentifier id)
     {
         ConcurrentDictionary<BlobIdentifier, MemoryBlobInfo> index = GetNamespaceContainer(ns);
 
         if (!index.TryGetValue(id, out MemoryBlobInfo? blobInfo))
-            return Task.FromResult<IBlobIndex.BlobInfo?>(null);
+        {
+            return Task.FromResult<BlobInfo?>(null);
+        }
 
-        return Task.FromResult<IBlobIndex.BlobInfo?>(blobInfo);
+        return Task.FromResult<BlobInfo?>(blobInfo);
     }
 
     public Task<bool> RemoveBlobFromIndex(NamespaceId ns, BlobIdentifier id)
@@ -75,7 +76,7 @@ public class MemoryBlobIndex : IBlobIndex
 
     public async Task<bool> BlobExistsInRegion(NamespaceId ns, BlobIdentifier blobIdentifier)
     {
-        IBlobIndex.BlobInfo? blobInfo = await GetBlobInfo(ns, blobIdentifier);
+        BlobInfo? blobInfo = await GetBlobInfo(ns, blobIdentifier);
         return blobInfo?.Regions.Contains(_jupiterSettings.CurrentValue.CurrentSite) ?? false;
     }
 
@@ -95,25 +96,25 @@ public class MemoryBlobIndex : IBlobIndex
                 info.References.Add((bucket, key));
                 return info;
             });
-        };
+        }
 
         return Task.CompletedTask;
     }
 
-    public async IAsyncEnumerable<IBlobIndex.BlobInfo> GetAllBlobs()
+    public async IAsyncEnumerable<BlobInfo> GetAllBlobs()
     {
         await Task.CompletedTask;
 
         foreach (KeyValuePair<NamespaceId, ConcurrentDictionary<BlobIdentifier, MemoryBlobInfo>> pair in _index)
         {
-            foreach ((BlobIdentifier? blobIdentifier, MemoryBlobInfo? blobInfo) in pair.Value)
+            foreach ((BlobIdentifier? _, MemoryBlobInfo? blobInfo) in pair.Value)
             {
                 yield return blobInfo;
             }
         }
     }
 
-    private MemoryBlobInfo NewBlobInfo(NamespaceId ns, BlobIdentifier blob, string region)
+    private static MemoryBlobInfo NewBlobInfo(NamespaceId ns, BlobIdentifier blob, string region)
     {
         MemoryBlobInfo info = new MemoryBlobInfo
         {
diff --git a/Engine/Source/Programs/Horde/HordeStorage/Horde.Storage/Implementation/Blob/MemoryBlobStore.cs b/Engine/Source/Programs/Horde/HordeStorage/Horde.Storage/Implementation/Blob/MemoryBlobStore.cs
--- a/Engine/Source/Programs/Horde/HordeStorage/Horde.Storage/Implementation/Blob/MemoryBlobStore.cs
+++ b/Engine/Source/Programs/Horde/HordeStorage/Horde.Storage/Implementation/Blob/MemoryBlobStore.cs
@@ -7,7 +7,6 @@ using System.IO;
 using System.Linq;
 using System.Threading.Tasks;
 using EpicGames.Horde.Storage;
-using Jupiter;
 using Jupiter.Implementation;
 using Jupiter.Utils;
 
@@ -53,7 +52,7 @@ namespace Horde.Storage.Implementation
                 identifier = BlobIdentifier.FromBlob(blob);
             }
 
-            var namespaceContainer = _blobs.GetOrAdd(ns, new ConcurrentDictionary<BlobIdentifier, BlobContainer>());
+            ConcurrentDictionary<BlobIdentifier, BlobContainer> namespaceContainer = _blobs.GetOrAdd(ns, new ConcurrentDictionary<BlobIdentifier, BlobContainer>());
             if (throwOnOverwrite && namespaceContainer.ContainsKey(identifier))
             {
                 throw new Exception($"Blob {identifier} already exists in {ns}");
@@ -75,7 +74,7 @@ namespace Horde.Storage.Implementation
 
         public Task<BlobContents> GetObject(NamespaceId ns, BlobIdentifier blob)
         {
-            if (!_blobs.TryGetValue(ns, value: out var namespaceContainer))
+            if (!_blobs.TryGetValue(ns, value: out ConcurrentDictionary<BlobIdentifier, BlobContainer>? namespaceContainer))
             {
                 throw new NamespaceNotFoundException(ns);
             }
@@ -91,7 +90,7 @@ namespace Horde.Storage.Implementation
 
         public Task DeleteObject(NamespaceId ns, BlobIdentifier blob)
         {
-            if (!_blobs.TryGetValue(ns, value: out var namespaceContainer))
+            if (!_blobs.TryGetValue(ns, value: out ConcurrentDictionary<BlobIdentifier, BlobContainer>? namespaceContainer))
             {
                 throw new NamespaceNotFoundException(ns);
             }
@@ -106,7 +105,7 @@ namespace Horde.Storage.Implementation
 
         public Task<bool> Exists(NamespaceId ns, BlobIdentifier blob)
         {
-            if (!_blobs.TryGetValue(ns, value: out var namespaceContainer))
+            if (!_blobs.TryGetValue(ns, value: out ConcurrentDictionary<BlobIdentifier, BlobContainer>? namespaceContainer))
             {
                 throw new NamespaceNotFoundException(ns);
             }
@@ -116,7 +115,7 @@ namespace Horde.Storage.Implementation
 
         public Task DeleteNamespace(NamespaceId ns)
         {
-            _blobs.Remove(ns, out var nsBlobs);
+            _blobs.Remove(ns, out ConcurrentDictionary<BlobIdentifier, BlobContainer>? nsBlobs);
             if (nsBlobs == null)
             {
                 throw new NamespaceNotFoundException(ns);
@@ -127,7 +126,7 @@ namespace Horde.Storage.Implementation
 
         public async IAsyncEnumerable<(BlobIdentifier,DateTime)> ListObjects(NamespaceId ns)
         {
-            if (!_blobs.TryGetValue(ns, value: out var namespaceContainer))
+            if (!_blobs.TryGetValue(ns, value: out ConcurrentDictionary<BlobIdentifier, BlobContainer>? namespaceContainer))
             {
                 throw new NamespaceNotFoundException(ns);
             }
@@ -141,7 +140,7 @@ namespace Horde.Storage.Implementation
 
         internal IEnumerable<BlobIdentifier> GetIdentifiers(NamespaceId ns)
         {
-            if (!_blobs.TryGetValue(ns, value: out var namespaceContainer))
+            if (!_blobs.TryGetValue(ns, value: out ConcurrentDictionary<BlobIdentifier, BlobContainer>? namespaceContainer))
             {
                 throw new NamespaceNotFoundException(ns);
             }
@@ -152,7 +151,7 @@ namespace Horde.Storage.Implementation
         // only for unit tests to update the last modified time
         internal void SetLastModifiedTime(NamespaceId ns, BlobIdentifier blob, DateTime modifiedTime)
         {
-            if (!_blobs.TryGetValue(ns, value: out var namespaceContainer))
+            if (!_blobs.TryGetValue(ns, value: out ConcurrentDictionary<BlobIdentifier, BlobContainer>? namespaceContainer))
             {
                 throw new Exception($"Namespace {ns} not found");
             }
diff --git a/Engine/Source/Programs/Horde/HordeStorage/Horde.Storage/Implementation/Blob/MemoryCacheBlobStore.cs b/Engine/Source/Programs/Horde/HordeStorage/Horde.Storage/Implementation/Blob/MemoryCacheBlobStore.cs
--- a/Engine/Source/Programs/Horde/HordeStorage/Horde.Storage/Implementation/Blob/MemoryCacheBlobStore.cs
+++ b/Engine/Source/Programs/Horde/HordeStorage/Horde.Storage/Implementation/Blob/MemoryCacheBlobStore.cs
@@ -23,7 +23,7 @@ namespace Horde.Storage.Implementation
     /// Primarily intended for use together with HierarchicalBlobStore to speed up reads from other stores, such as Amazon S3 and disk.
     /// </summary>
     // ReSharper disable once ClassNeverInstantiated.Global
-    public class MemoryCacheBlobStore : IBlobStore
+    public class MemoryCacheBlobStore : IBlobStore, IDisposable
     {
         private readonly IOptionsMonitor<MemoryCacheBlobSettings> _memoryOptions;
 
@@ -31,7 +31,7 @@ namespace Horde.Storage.Implementation
 
         private readonly ILogger _logger = Log.ForContext<MemoryCacheBlobStore>();
         private readonly FieldInfo _sizeField;
-        private Timer _cacheInfoTimer;
+        private readonly Timer _cacheInfoTimer;
 
         public MemoryCacheBlobStore(IOptionsMonitor<MemoryCacheBlobSettings> memoryOptions)
         {
@@ -69,7 +69,9 @@ namespace Horde.Storage.Implementation
         public async Task<BlobIdentifier> PutObject(NamespaceId ns, Stream blob, BlobIdentifier identifier)
         {
             if (blob.Length > int.MaxValue)
+            {
                 throw new BlobToLargeException(identifier);
+            }
 
             return await PutObject(ns, blob: await blob.ToByteArray(), identifier);
         }
@@ -109,7 +111,7 @@ namespace Horde.Storage.Implementation
             return Array.Empty<(BlobIdentifier, DateTime)>().ToAsyncEnumerable();
         }
         
-        private string BuildKey(NamespaceId ns, BlobIdentifier blob)
+        private static string BuildKey(NamespaceId ns, BlobIdentifier blob)
         {
             return $"{ns}.{blob}";
         }
@@ -118,5 +120,20 @@ namespace Horde.Storage.Implementation
         {
             return (long?)_sizeField.GetValue(_memoryCache);
         }
+
+        protected virtual void Dispose(bool disposing)
+        {
+            if (disposing)
+            {
+                _memoryCache.Dispose();
+                _cacheInfoTimer.Dispose();
+            }
+        }
+
+        public void Dispose()
+        {
+            Dispose(true);
+            GC.SuppressFinalize(this);
+        }
     }
 }
diff --git a/Engine/Source/Programs/Horde/HordeStorage/Horde.Storage/Implementation/Blob/MongoBlobIndex.cs b/Engine/Source/Programs/Horde/HordeStorage/Horde.Storage/Implementation/Blob/MongoBlobIndex.cs
--- a/Engine/Source/Programs/Horde/HordeStorage/Horde.Storage/Implementation/Blob/MongoBlobIndex.cs
+++ b/Engine/Source/Programs/Horde/HordeStorage/Horde.Storage/Implementation/Blob/MongoBlobIndex.cs
@@ -54,13 +54,15 @@ public class MongoBlobIndex : MongoStore, IBlobIndex
         });
     }
 
-    public async Task<IBlobIndex.BlobInfo?> GetBlobInfo(NamespaceId ns, BlobIdentifier id)
+    public async Task<BlobInfo?> GetBlobInfo(NamespaceId ns, BlobIdentifier id)
     {
         IMongoCollection<MongoBlobIndexModelV0> collection = GetCollection<MongoBlobIndexModelV0>();
         IAsyncCursor<MongoBlobIndexModelV0>? cursor = await collection.FindAsync(m => m.Ns == ns.ToString() && m.BlobId == id.ToString());
         MongoBlobIndexModelV0? model = await cursor.FirstOrDefaultAsync();
         if (model == null)
+        {
             return null;
+        }
 
         return model.ToBlobInfo();
     }
@@ -77,7 +79,7 @@ public class MongoBlobIndex : MongoStore, IBlobIndex
     {
         region ??= _jupiterSettings.CurrentValue.CurrentSite;
 
-        IBlobIndex.BlobInfo? blobInfo = await GetBlobInfo(ns, id);
+        BlobInfo? blobInfo = await GetBlobInfo(ns, id);
         IMongoCollection<MongoBlobIndexModelV0> collection = GetCollection<MongoBlobIndexModelV0>();
         MongoBlobIndexModelV0 model = new MongoBlobIndexModelV0(blobInfo!);
         model.Regions.Remove(region);
@@ -91,7 +93,7 @@ public class MongoBlobIndex : MongoStore, IBlobIndex
 
     public async Task<bool> BlobExistsInRegion(NamespaceId ns, BlobIdentifier blobIdentifier)
     {
-        IBlobIndex.BlobInfo? blobInfo = await GetBlobInfo(ns, blobIdentifier);
+        BlobInfo? blobInfo = await GetBlobInfo(ns, blobIdentifier);
         return blobInfo?.Regions.Contains(_jupiterSettings.CurrentValue.CurrentSite) ?? false;
     }
 
@@ -116,7 +118,7 @@ public class MongoBlobIndex : MongoStore, IBlobIndex
         await Task.WhenAll(refUpdateTasks);
     }
 
-    public async IAsyncEnumerable<IBlobIndex.BlobInfo> GetAllBlobs()
+    public async IAsyncEnumerable<BlobInfo> GetAllBlobs()
     {
         IMongoCollection<MongoBlobIndexModelV0> collection = GetCollection<MongoBlobIndexModelV0>();
         IAsyncCursor<MongoBlobIndexModelV0>? cursor = await collection.FindAsync(FilterDefinition<MongoBlobIndexModelV0>.Empty);
@@ -145,7 +147,7 @@ class MongoBlobIndexModelV0
         References = references;
     }
 
-    public MongoBlobIndexModelV0(IBlobIndex.BlobInfo blobInfo)
+    public MongoBlobIndexModelV0(BlobInfo blobInfo)
     {
         Ns = blobInfo.Namespace.ToString();
         BlobId = blobInfo.BlobIdentifier.ToString();
@@ -170,9 +172,9 @@ class MongoBlobIndexModelV0
     [BsonDictionaryOptions(DictionaryRepresentation.Document)]
     public List<Dictionary<string, string>> References { get; set; } = new List<Dictionary<string, string>>();
 
-    public IBlobIndex.BlobInfo ToBlobInfo()
+    public BlobInfo ToBlobInfo()
     {
-        return new IBlobIndex.BlobInfo()
+        return new BlobInfo()
         {
             Namespace = new NamespaceId(Ns),
             BlobIdentifier = new BlobIdentifier(BlobId),
diff --git a/Engine/Source/Programs/Horde/HordeStorage/Horde.Storage/Implementation/Blob/PathUtil.cs b/Engine/Source/Programs/Horde/HordeStorage/Horde.Storage/Implementation/Blob/PathUtil.cs
--- a/Engine/Source/Programs/Horde/HordeStorage/Horde.Storage/Implementation/Blob/PathUtil.cs
+++ b/Engine/Source/Programs/Horde/HordeStorage/Horde.Storage/Implementation/Blob/PathUtil.cs
@@ -11,6 +11,6 @@ public static class PathUtil
     public static string ResolvePath(string path)
     {
         return Environment.ExpandEnvironmentVariables(path)
-            .Replace("$(ExecutableLocation)", Path.GetDirectoryName(Assembly.GetEntryAssembly()!.Location));
+            .Replace("$(ExecutableLocation)", Path.GetDirectoryName(Assembly.GetEntryAssembly()!.Location), StringComparison.OrdinalIgnoreCase);
     }
 }
diff --git a/Engine/Source/Programs/Horde/HordeStorage/Horde.Storage/Implementation/Blob/RelayBlobStore.cs b/Engine/Source/Programs/Horde/HordeStorage/Horde.Storage/Implementation/Blob/RelayBlobStore.cs
--- a/Engine/Source/Programs/Horde/HordeStorage/Horde.Storage/Implementation/Blob/RelayBlobStore.cs
+++ b/Engine/Source/Programs/Horde/HordeStorage/Horde.Storage/Implementation/Blob/RelayBlobStore.cs
@@ -22,7 +22,7 @@ namespace Horde.Storage.Implementation
 
         public async Task<BlobIdentifier> PutObject(NamespaceId ns, byte[] blob, BlobIdentifier identifier)
         {
-            HttpRequestMessage putObjectRequest = BuildHttpRequest(HttpMethod.Put, $"api/v1/blobs/{ns}/{identifier}");
+            using HttpRequestMessage putObjectRequest = BuildHttpRequest(HttpMethod.Put, new Uri($"api/v1/blobs/{ns}/{identifier}", UriKind.Relative));
 
             putObjectRequest.Content = new ByteArrayContent(blob);
             putObjectRequest.Content.Headers.ContentType = new MediaTypeHeaderValue(MediaTypeNames.Application.Octet);
@@ -42,7 +42,7 @@ namespace Horde.Storage.Implementation
 
         public async Task<BlobIdentifier> PutObject(NamespaceId ns, Stream content, BlobIdentifier identifier)
         {
-            HttpRequestMessage putObjectRequest = BuildHttpRequest(HttpMethod.Put, $"api/v1/blobs/{ns}/{identifier}");
+            using HttpRequestMessage putObjectRequest = BuildHttpRequest(HttpMethod.Put, new Uri($"api/v1/blobs/{ns}/{identifier}", UriKind.Relative));
             putObjectRequest.Content = new StreamContent(content);
             putObjectRequest.Content.Headers.ContentType = new MediaTypeHeaderValue(MediaTypeNames.Application.Octet);
             putObjectRequest.Content.Headers.Add(CommonHeaders.HashHeaderName, identifier.ToString());
@@ -55,25 +55,33 @@ namespace Horde.Storage.Implementation
 
         public async Task<BlobContents> GetObject(NamespaceId ns, BlobIdentifier blob)
         {
-            HttpRequestMessage getObjectRequest = BuildHttpRequest(HttpMethod.Get, $"api/v1/blobs/{ns}/{blob}");
+            using HttpRequestMessage getObjectRequest = BuildHttpRequest(HttpMethod.Get, new Uri($"api/v1/blobs/{ns}/{blob}", UriKind.Relative));
             getObjectRequest.Headers.Add("Accept", MediaTypeNames.Application.Octet);
             HttpResponseMessage response = await HttpClient.SendAsync(getObjectRequest);
             if (response.StatusCode == HttpStatusCode.NotFound)
+            {
                 throw new BlobNotFoundException(ns, blob);
+            }
+
             response.EnsureSuccessStatusCode();
 
             long? contentLength = response.Content.Headers.ContentLength;
             if (contentLength == null)
+            {
                 throw new Exception($"Content length missing in response from upstream blob store. This is not supported");
+            }
+
             return new BlobContents(await response.Content.ReadAsStreamAsync(), contentLength.Value);
         }
 
         public async Task<bool> Exists(NamespaceId ns, BlobIdentifier blob)
         {
-            HttpRequestMessage headObjectRequest = BuildHttpRequest(HttpMethod.Head, $"api/v1/blobs/{ns}/{blob}");
+            using HttpRequestMessage headObjectRequest = BuildHttpRequest(HttpMethod.Head, new Uri($"api/v1/blobs/{ns}/{blob}", UriKind.Relative));
             HttpResponseMessage response = await HttpClient.SendAsync(headObjectRequest);
             if (response.StatusCode == HttpStatusCode.NotFound)
+            {
                 return false;
+            }
 
             response.EnsureSuccessStatusCode();
 
diff --git a/Engine/Source/Programs/Horde/HordeStorage/Horde.Storage/Implementation/Blob/ScyllaBlobIndex.cs b/Engine/Source/Programs/Horde/HordeStorage/Horde.Storage/Implementation/Blob/ScyllaBlobIndex.cs
--- a/Engine/Source/Programs/Horde/HordeStorage/Horde.Storage/Implementation/Blob/ScyllaBlobIndex.cs
+++ b/Engine/Source/Programs/Horde/HordeStorage/Horde.Storage/Implementation/Blob/ScyllaBlobIndex.cs
@@ -45,7 +45,7 @@ public class ScyllaBlobIndex : IBlobIndex
             new string[] { region }, ns.ToString(), new ScyllaBlobIdentifier(id));
     }
 
-    public async Task<IBlobIndex.BlobInfo?> GetBlobInfo(NamespaceId ns, BlobIdentifier id)
+    public async Task<BlobInfo?> GetBlobInfo(NamespaceId ns, BlobIdentifier id)
     {
         using IScope scope = Tracer.Instance.StartActive("scylla.fetch_blob_index");
         scope.Span.ResourceName = $"{ns}.{id}";
@@ -55,7 +55,9 @@ public class ScyllaBlobIndex : IBlobIndex
                 ns.ToString(), new ScyllaBlobIdentifier(id));
 
         if (blobIndex == null)
+        {
             return null;
+        }
 
         return blobIndex.ToBlobInfo();
     }
@@ -82,7 +84,7 @@ public class ScyllaBlobIndex : IBlobIndex
 
     public async Task<bool> BlobExistsInRegion(NamespaceId ns, BlobIdentifier blobIdentifier)
     {
-        IBlobIndex.BlobInfo? blobInfo = await GetBlobInfo(ns, blobIdentifier);
+        BlobInfo? blobInfo = await GetBlobInfo(ns, blobIdentifier);
         return blobInfo?.Regions.Contains(_jupiterSettings.CurrentValue.CurrentSite) ?? false;
     }
 
@@ -91,14 +93,14 @@ public class ScyllaBlobIndex : IBlobIndex
         using IScope _ = Tracer.Instance.StartActive("scylla.add_ref_blobs");
 
         string nsAsString = ns.ToString();
-        ScyllaObjectReference r = new ScyllaObjectReference(bucket, key);
+        ScyllaObjectReference reference = new ScyllaObjectReference(bucket, key);
         Task[] refUpdateTasks = new Task[blobs.Length];
         for (int i = 0; i < blobs.Length; i++)
         {
             BlobIdentifier id = blobs[i];
             refUpdateTasks[i] = _mapper.UpdateAsync<ScyllaBlobIndexTable>(
                 "SET references = references + ? WHERE namespace = ? AND blob_id = ?",
-                new[] { r },
+                new[] { reference },
                 nsAsString,
                 new ScyllaBlobIdentifier(id));
         }
@@ -106,7 +108,7 @@ public class ScyllaBlobIndex : IBlobIndex
         await Task.WhenAll(refUpdateTasks);
     }
 
-    public async IAsyncEnumerable<IBlobIndex.BlobInfo> GetAllBlobs()
+    public async IAsyncEnumerable<BlobInfo> GetAllBlobs()
     {
         using IScope _ = Tracer.Instance.StartActive("scylla.get_all_blobs");
 
@@ -117,7 +119,6 @@ public class ScyllaBlobIndex : IBlobIndex
     }
 }
 
-
 [Cassandra.Mapping.Attributes.Table("blob_index")]
 class ScyllaBlobIndexTable
 {
@@ -150,9 +151,9 @@ class ScyllaBlobIndexTable
     [Cassandra.Mapping.Attributes.Column("references")]
     public List<ScyllaObjectReference>? References { get; set; }
 
-    public IBlobIndex.BlobInfo ToBlobInfo()
+    public BlobInfo ToBlobInfo()
     {
-        return new IBlobIndex.BlobInfo
+        return new BlobInfo
         {
             Regions = Regions ?? new HashSet<string>(),
             BlobIdentifier = BlobId.AsBlobIdentifier(),
diff --git a/Engine/Source/Programs/Horde/HordeStorage/Horde.Storage/Implementation/ContentId/CacheContentIdStore.cs b/Engine/Source/Programs/Horde/HordeStorage/Horde.Storage/Implementation/ContentId/CacheContentIdStore.cs
--- a/Engine/Source/Programs/Horde/HordeStorage/Horde.Storage/Implementation/ContentId/CacheContentIdStore.cs
+++ b/Engine/Source/Programs/Horde/HordeStorage/Horde.Storage/Implementation/ContentId/CacheContentIdStore.cs
@@ -1,5 +1,6 @@
 // Copyright Epic Games, Inc. All Rights Reserved.
 
+using System;
 using System.Net;
 using System.Net.Http;
 using System.Threading.Tasks;
@@ -19,11 +20,13 @@ namespace Horde.Storage.Implementation
 
         public async Task<BlobIdentifier[]?> Resolve(NamespaceId ns, ContentId contentId, bool mustBeContentId)
         {
-            HttpRequestMessage getContentIdRequest = BuildHttpRequest(HttpMethod.Get, $"api/v1/content-id/{ns}/{contentId}");
+            using HttpRequestMessage getContentIdRequest = BuildHttpRequest(HttpMethod.Get, new Uri($"api/v1/content-id/{ns}/{contentId}", UriKind.Relative));
             HttpResponseMessage response = await HttpClient.SendAsync(getContentIdRequest);
 
             if (response.StatusCode == HttpStatusCode.NotFound)
+            {
                 return null;
+            }
 
             response.EnsureSuccessStatusCode();
             ResolvedContentIdResponse resolvedContentId = await response.Content.ReadAsAsync<ResolvedContentIdResponse>();
@@ -33,7 +36,7 @@ namespace Horde.Storage.Implementation
 
         public async Task Put(NamespaceId ns, ContentId contentId, BlobIdentifier blobIdentifier, int contentWeight)
         {
-            HttpRequestMessage putContentIdRequest = BuildHttpRequest(HttpMethod.Put, $"api/v1/content-id/{ns}/{contentId}/update/{blobIdentifier}/{contentWeight}");
+            using HttpRequestMessage putContentIdRequest = BuildHttpRequest(HttpMethod.Put, new Uri($"api/v1/content-id/{ns}/{contentId}/update/{blobIdentifier}/{contentWeight}", UriKind.Relative));
             HttpResponseMessage response = await HttpClient.SendAsync(putContentIdRequest);
 
             response.EnsureSuccessStatusCode();
diff --git a/Engine/Source/Programs/Horde/HordeStorage/Horde.Storage/Implementation/ContentId/CompressedBufferUtils.cs b/Engine/Source/Programs/Horde/HordeStorage/Horde.Storage/Implementation/ContentId/CompressedBufferUtils.cs
--- a/Engine/Source/Programs/Horde/HordeStorage/Horde.Storage/Implementation/ContentId/CompressedBufferUtils.cs
+++ b/Engine/Source/Programs/Horde/HordeStorage/Horde.Storage/Implementation/ContentId/CompressedBufferUtils.cs
@@ -3,10 +3,7 @@
 using System;
 using System.Buffers.Binary;
 using System.IO;
-using Blake3;
-using Datadog.Trace;
 using Force.Crc32;
-using Jupiter;
 using Jupiter.Implementation;
 using Jupiter.Utils;
 using K4os.Compression.LZ4;
@@ -22,8 +19,9 @@ namespace Horde.Storage.Implementation
             _oodleCompressor = oodleCompressor;
         }
 
-        public class Header
+        private class Header
         {
+            [System.Diagnostics.CodeAnalysis.SuppressMessage("Design", "CA1028:Enum Storage should be Int32", Justification = "Interop requires byte")]
             public enum CompressionMethod : byte
             {
                 // Header is followed by one uncompressed block. 
@@ -33,32 +31,35 @@ namespace Horde.Storage.Implementation
                 LZ4 = 4,
             }
 
-            public static uint ExpectedMagic = 0xb7756362; // <dot>ucb
-            public static uint HeaderLength = 64;
+            public const uint ExpectedMagic = 0xb7756362; // <dot>ucb
+            public const uint HeaderLength = 64;
 
             // A magic number to identify a compressed buffer. Always 0xb7756362.
-            public uint Magic;
+
+            public uint Magic { get; set; }
             // A CRC-32 used to check integrity of the buffer. Uses the polynomial 0x04c11db7.
-            public uint Crc32;
+
+            public uint Crc32 { get; set; }
+
             // The method used to compress the buffer. Affects layout of data following the header. 
-            public CompressionMethod Method;
-            public byte CompressionLevel;
-            public byte CompressionMethodUsed;
+            public CompressionMethod Method { get; set; }
+            public byte CompressionLevel { get; set; }
+            public byte CompressionMethodUsed { get; set; }
 
             // The power of two size of every uncompressed block except the last. Size is 1 << BlockSizeExponent. 
-            public byte BlockSizeExponent;
+            public byte BlockSizeExponent { get; set; }
 
             // The number of blocks that follow the header. 
-            public uint BlockCount;
+            public uint BlockCount { get; set; }
 
             // The total size of the uncompressed data. 
-            public ulong TotalRawSize;
+            public ulong TotalRawSize { get; set; }
 
             // The total size of the compressed data including the header. 
-            public ulong TotalCompressedSize;
+            public ulong TotalCompressedSize { get; set; }
 
             /** The hash of the uncompressed data. */
-            public byte[] RawHash = Array.Empty<byte>();
+            public byte[] RawHash { get; set; } = Array.Empty<byte>();
 
             public void ByteSwap()
             {
@@ -70,12 +71,14 @@ namespace Horde.Storage.Implementation
             }
         }
 
-        public Header ExtractHeader(byte[] content)
+        private static Header ExtractHeader(byte[] content)
         {
             // the header is always stored big endian
             bool needsByteSwap = BitConverter.IsLittleEndian;
             if (content.Length < Header.HeaderLength)
+            {
                 throw new ArgumentOutOfRangeException(nameof(content), $"Content was less then {Header.HeaderLength} bytes and thus is not a compressed buffer");
+            }
 
             using MemoryStream ms = new MemoryStream(content);
             using BinaryReader reader = new BinaryReader(ms);
@@ -96,34 +99,43 @@ namespace Horde.Storage.Implementation
             header.RawHash = hash;
 
             if (needsByteSwap)
+            {
                 header.ByteSwap();
+            }
 
             if (ms.Position != Header.HeaderLength)
+            {
                 throw new Exception($"Read {ms.Position} bytes but expected to read {Header.HeaderLength}");
+            }
 
             if (header.Magic != Header.ExpectedMagic)
+            {
                 throw new InvalidMagicException(header.Magic, Header.ExpectedMagic);
+            }
 
             // calculate the crc from the start of the method field (skipping magic which is a constant and the crc field itself)
-            const int methodOffset = sizeof(uint) + sizeof(uint);
+            const int MethodOffset = sizeof(uint) + sizeof(uint);
 
             // none compressed objects have no extra blocks
             uint blocksByteUsed = header.Method != Header.CompressionMethod.None ? header.BlockCount * (uint)sizeof(uint) : 0;
-            uint calculatedCrc = Crc32Algorithm.Compute(content, methodOffset, (int)(Header.HeaderLength - methodOffset + blocksByteUsed));
+            uint calculatedCrc = Crc32Algorithm.Compute(content, MethodOffset, (int)(Header.HeaderLength - MethodOffset + blocksByteUsed));
             
             if (header.Crc32 != calculatedCrc)
+            {
                 throw new InvalidHashException(header.Crc32, calculatedCrc);
+            }
 
             return header;
         }
 
-
         public byte[] DecompressContent(byte[] content)
         {
             Header header = ExtractHeader(content);
 
             if (content.LongLength < (long)header.TotalCompressedSize)
+            {
                 throw new Exception($"Expected buffer to be {header.TotalCompressedSize} but it was {content.LongLength}");
+            }
 
             ulong decompressedPayloadOffset = 0;
             byte[] decompressedPayload = new byte[header.TotalRawSize];
@@ -176,7 +188,9 @@ namespace Horde.Storage.Implementation
             }
            
             if (header.TotalRawSize != decompressedPayloadOffset)
+            {
                 throw new Exception("Did not decompress the full payload");
+            }
 
             {
                 // only read the first 20 bytes of the hash field as IoHashes are 20 bytes and not 32 bytes
@@ -187,10 +201,11 @@ namespace Horde.Storage.Implementation
                 BlobIdentifier contentHash = BlobIdentifier.FromBlob(decompressedPayload);
                
                 if (!headerIdentifier.Equals(contentHash))
+                {
                     throw new Exception($"Payload was expected to be {headerIdentifier} but was {contentHash}");
+                }
             }
 
-
             return decompressedPayload;
         }
 
@@ -203,10 +218,7 @@ namespace Horde.Storage.Implementation
                     return compressedPayload.Length;
                 case Header.CompressionMethod.Oodle:
                 {
-                    // we have separate enums from the official oodle api to make sure they are stable enums so we need to translate to the ones expected by the api
-                    OodleLZ_Compressor compressor = OodleUtils.ToOodleApiCompressor((OoodleCompressorMethod)header.CompressionMethodUsed);
-                    OodleLZ_CompressionLevel compressionLevel = OodleUtils.ToOodleApiCompressionLevel((OoodleCompressionLevel)header.CompressionLevel);
-                    long writtenBytes = _oodleCompressor.Decompress(compressor,compressedPayload.ToArray(), compressionLevel, (long)rawBlockSize, out byte[] result);
+                    long writtenBytes = _oodleCompressor.Decompress(compressedPayload.ToArray(), (long)rawBlockSize, out byte[] result);
                     if (writtenBytes == 0)
                     {
                         throw new Exception("Failed to run oodle decompress");
@@ -220,12 +232,12 @@ namespace Horde.Storage.Implementation
                     return writtenBytes;
                 }
                 default:
-                    throw new ArgumentOutOfRangeException(nameof(header.Method), header.Method, null);
+                    throw new NotImplementedException($"Method {header.Method} is not a support value");
             }
         }
     }
 
-    internal class InvalidHashException : Exception
+	public class InvalidHashException : Exception
     {
         public InvalidHashException(uint headerCrc32, uint calculatedCrc) : base($"Header specified crc \"{headerCrc32}\" but calculated hash was \"{calculatedCrc}\"")
         {
@@ -233,16 +245,18 @@ namespace Horde.Storage.Implementation
         }
     }
 
-    internal class InvalidMagicException : Exception
+	public class InvalidMagicException : Exception
     {
         public InvalidMagicException(uint headerMagic, uint expectedMagic) : base($"Header magic \"{headerMagic}\" was incorrect, expected to be {expectedMagic}")
         {
         }
     }
 
-    // from OodleDataCompression.h , we define our own enums for oodle compressions used and convert to the ones expected in the oodle api
-    public enum OoodleCompressorMethod: byte
-    {
+	// from OodleDataCompression.h , we define our own enums for oodle compressions used and convert to the ones expected in the oodle api
+#pragma warning disable CA1028 // Enum Storage should be Int32
+	public enum OoodleCompressorMethod: byte
+
+	{
         NotSet = 0,
         Selkie = 1,
         Mermaid = 2,
@@ -266,7 +280,7 @@ namespace Horde.Storage.Implementation
         Optimal3 = 7,
         Optimal4 = 8,
     }
-
+#pragma warning restore CA1028 // Enum Storage should be Int32
     public static class OodleUtils
     {
         public static OodleLZ_CompressionLevel ToOodleApiCompressionLevel(OoodleCompressionLevel compressionLevel)
@@ -322,6 +336,5 @@ namespace Horde.Storage.Implementation
                     throw new ArgumentOutOfRangeException(nameof(compressor), compressor, null);
             }
         }
-
     }
 }
diff --git a/Engine/Source/Programs/Horde/HordeStorage/Horde.Storage/Implementation/ContentId/MemoryContentIdStore.cs b/Engine/Source/Programs/Horde/HordeStorage/Horde.Storage/Implementation/ContentId/MemoryContentIdStore.cs
--- a/Engine/Source/Programs/Horde/HordeStorage/Horde.Storage/Implementation/ContentId/MemoryContentIdStore.cs
+++ b/Engine/Source/Programs/Horde/HordeStorage/Horde.Storage/Implementation/ContentId/MemoryContentIdStore.cs
@@ -24,11 +24,13 @@ namespace Horde.Storage.Implementation
             {
                 if (contentIdsForNamespace.TryGetValue(contentId, out SortedList<int, BlobIdentifier[]>? contentIdMappings))
                 {
-                    foreach ((int weight, BlobIdentifier[] blobs) in contentIdMappings)
+                    foreach ((int _, BlobIdentifier[] blobs) in contentIdMappings)
                     {
                         BlobIdentifier[] missingBlobs = await _blobStore.FilterOutKnownBlobs(ns, blobs);
                         if (missingBlobs.Length == 0)
+                        {
                             return blobs;
+                        }
                         // blobs are missing continue testing with the next content id in the weighted list as that might exist
                     }
                 }
@@ -37,7 +39,9 @@ namespace Horde.Storage.Implementation
             BlobIdentifier uncompressedBlobIdentifier = contentId.AsBlobIdentifier();
             // if no content id is found, but we have a blob that matches the content id (so a unchunked and uncompressed version of the data) we use that instead
             if (!mustBeContentId && await _blobStore.Exists(ns, uncompressedBlobIdentifier))
+            {
                 return new[] { uncompressedBlobIdentifier };
+            }
 
             return null;
         }
diff --git a/Engine/Source/Programs/Horde/HordeStorage/Horde.Storage/Implementation/ContentId/MongoContentIdStore.cs b/Engine/Source/Programs/Horde/HordeStorage/Horde.Storage/Implementation/ContentId/MongoContentIdStore.cs
--- a/Engine/Source/Programs/Horde/HordeStorage/Horde.Storage/Implementation/ContentId/MongoContentIdStore.cs
+++ b/Engine/Source/Programs/Horde/HordeStorage/Horde.Storage/Implementation/ContentId/MongoContentIdStore.cs
@@ -53,7 +53,9 @@ namespace Horde.Storage.Implementation
                     {
                         BlobIdentifier[] missingBlobs = await _blobStore.FilterOutKnownBlobs(ns, blobs);
                         if (missingBlobs.Length == 0)
+                        {
                             return blobs;
+                        }
                     }
                     // blobs are missing continue testing with the next content id in the weighted list as that might exist
 
@@ -115,7 +117,7 @@ namespace Horde.Storage.Implementation
         public string ContentId { get; set; }
 
         [BsonDictionaryOptions(DictionaryRepresentation.ArrayOfDocuments)]
+        [System.Diagnostics.CodeAnalysis.SuppressMessage("Usage", "CA2227:Collection properties should be read only", Justification = "Used by serialization")]
         public Dictionary<int, string[]> ContentWeightToBlobsMap { get; set; }
-
     }
 }
diff --git a/Engine/Source/Programs/Horde/HordeStorage/Horde.Storage/Implementation/ContentId/ScyllaContentIdStore.cs b/Engine/Source/Programs/Horde/HordeStorage/Horde.Storage/Implementation/ContentId/ScyllaContentIdStore.cs
--- a/Engine/Source/Programs/Horde/HordeStorage/Horde.Storage/Implementation/ContentId/ScyllaContentIdStore.cs
+++ b/Engine/Source/Programs/Horde/HordeStorage/Horde.Storage/Implementation/ContentId/ScyllaContentIdStore.cs
@@ -1,7 +1,6 @@
 // Copyright Epic Games, Inc. All Rights Reserved.
 
 using System;
-using System.Collections.Generic;
 using System.Linq;
 using System.Threading.Tasks;
 using Cassandra;
@@ -16,7 +15,7 @@ namespace Horde.Storage.Implementation
     {
         private readonly ISession _session;
         private readonly IBlobService _blobStore;
-        private Mapper _mapper;
+        private readonly Mapper _mapper;
 
         public ScyllaContentIdStore(IScyllaSessionManager scyllaSessionManager, IBlobService blobStore)
         {
@@ -34,7 +33,6 @@ namespace Horde.Storage.Implementation
             ));
         }
 
-
         public async Task<BlobIdentifier[]?> Resolve(NamespaceId ns, ContentId contentId, bool mustBeContentId)
         {
             using IScope scope = Tracer.Instance.StartActive("ScyllaContentIdStore.ResolveContentId");
@@ -51,8 +49,10 @@ namespace Horde.Storage.Implementation
             foreach (ScyllaContentId? resolvedContentId in await _mapper.FetchAsync<ScyllaContentId>("WHERE content_id = ? ORDER BY content_weight DESC", new ScyllaBlobIdentifier(contentId)))
             {
                 if (resolvedContentId == null)
+                {
                     throw new InvalidContentIdException(contentId);
-                
+                }
+
                 BlobIdentifier[] blobs = resolvedContentId.Chunks.Select(b => b.AsBlobIdentifier()).ToArray();
 
                 {
@@ -60,7 +60,9 @@ namespace Horde.Storage.Implementation
 
                     BlobIdentifier[] missingBlobs = await _blobStore.FilterOutKnownBlobs(ns, blobs);
                     if (missingBlobs.Length == 0)
+                    {
                         return blobs;
+                    }
                 }
                 // blobs are missing continue testing with the next content id in the weighted list as that might exist
             }
@@ -71,7 +73,9 @@ namespace Horde.Storage.Implementation
                 bool contentIdBlobExists = await blobStoreExistsTask!;
 
                 if (contentIdBlobExists)
+                {
                     return new[] { contentIdBlob };
+                }
             }
             
             // unable to resolve the content id
diff --git a/Engine/Source/Programs/Horde/HordeStorage/Horde.Storage/Implementation/DDCRefService.cs b/Engine/Source/Programs/Horde/HordeStorage/Horde.Storage/Implementation/DDCRefService.cs
--- a/Engine/Source/Programs/Horde/HordeStorage/Horde.Storage/Implementation/DDCRefService.cs
+++ b/Engine/Source/Programs/Horde/HordeStorage/Horde.Storage/Implementation/DDCRefService.cs
@@ -41,6 +41,7 @@ namespace Horde.Storage.Implementation
         [CbField("lastAccessTime")]
         public DateTime? LastAccessTime { get; set; }
 
+        [System.Diagnostics.CodeAnalysis.SuppressMessage("Usage", "CA2227:Collection properties should be read only", Justification = "Required by serialization")]
         public Dictionary<string, object>? Metadata { get; set; }
 
         [CbField("contentHash")]
@@ -107,9 +108,14 @@ namespace Horde.Storage.Implementation
 
             IRefsStore.ExtraFieldsFlag flags = IRefsStore.ExtraFieldsFlag.None;
             if (needsLastAccess)
+            {
                 flags |= IRefsStore.ExtraFieldsFlag.LastAccess;
+            }
+
             if (needsMetadata)
+            {
                 flags |= IRefsStore.ExtraFieldsFlag.Metadata;
+            }
 
             string resource = $"{ns}.{bucket}.{key}";
             RefRecord? record;
@@ -125,13 +131,13 @@ namespace Horde.Storage.Implementation
 
             // its not critical that this finishes, so we just log errors in case it fails but never await
             {
-                Task _ = _lastAccessTracker.TrackUsed(record).ContinueWith(task =>
+                Task _ = _lastAccessTracker.TrackUsed(record).ContinueWith((task, _) =>
                 {
                     if (task.Exception != null)
                     {
                         _logger.Error(task.Exception, "Exception when tracking last access record");
                     }
-                });
+                }, null, TaskScheduler.Current);
             }
 
             BlobContents? maybeBlob = null;
@@ -189,7 +195,6 @@ namespace Horde.Storage.Implementation
             return record;
         }
 
-
         public async Task<PutRequestResponse> Put(
             NamespaceId ns,
             BucketId bucket,
@@ -232,7 +237,9 @@ namespace Horde.Storage.Implementation
             }
             int countOfBlobs = blobReferences.Length;
             if (countOfBlobs == 0)
+            {
                 throw new ArgumentException("No blobs found when determining partitioning");
+            }
 
             // as the body is just a binary blob we are unable to receive metadata for this route
             // TODO: We could accept metadata from query parameters
@@ -251,7 +258,9 @@ namespace Horde.Storage.Implementation
             else
             {
                 if (slices == null)
+                {
                     throw new Exception("Slices was never set when uploading partitioned blob.");
+                }
 
                 Task[] insertTasks = new Task[countOfBlobs];
                 for (int i = 0; i < countOfBlobs; i++)
@@ -318,7 +327,10 @@ namespace Horde.Storage.Implementation
         {
             long deleteCount = await _refsStore.Delete(ns, bucket, key);
             if (deleteCount != 0)
+            {
                 await _transactionLog.Delete(ns, bucket, key);
+            }
+
             return deleteCount;
         }
     }
diff --git a/Engine/Source/Programs/Horde/HordeStorage/Horde.Storage/Implementation/GC/BlobCleanupService.cs b/Engine/Source/Programs/Horde/HordeStorage/Horde.Storage/Implementation/GC/BlobCleanupService.cs
--- a/Engine/Source/Programs/Horde/HordeStorage/Horde.Storage/Implementation/GC/BlobCleanupService.cs
+++ b/Engine/Source/Programs/Horde/HordeStorage/Horde.Storage/Implementation/GC/BlobCleanupService.cs
@@ -6,14 +6,17 @@ using System.Threading;
 using System.Threading.Tasks;
 using Jupiter;
 using Microsoft.Extensions.Options;
-using Datadog.Trace;
-using Jupiter.Implementation;
 using Microsoft.Extensions.DependencyInjection;
 using Serilog;
 
 namespace Horde.Storage.Implementation
 {
-    public class BlobCleanupService : PollingService<BlobCleanupService.BlobCleanupState>
+    public class BlobCleanupState
+    {
+        public List<IBlobCleanup> BlobCleanups { get; } = new List<IBlobCleanup>();
+    }
+
+    public class BlobCleanupService : PollingService<BlobCleanupState>
     {
         private readonly IOptionsMonitor<GCSettings> _settings;
         private volatile bool _alreadyPolling;
@@ -24,14 +27,7 @@ namespace Horde.Storage.Implementation
             return _settings.CurrentValue.BlobCleanupServiceEnabled;
         }
 
-        public class BlobCleanupState
-        {
-            public List<IBlobCleanup> BlobCleanups { get; } = new List<IBlobCleanup>();
-        }
-
-        public BlobCleanupService(IServiceProvider provider, IOptionsMonitor<GCSettings> settings) :
-            base(serviceName: nameof(BlobCleanupService), settings.CurrentValue.BlobCleanupPollFrequency,
-                new BlobCleanupState())
+        public BlobCleanupService(IServiceProvider provider, IOptionsMonitor<GCSettings> settings) : base(serviceName: nameof(BlobCleanupService), settings.CurrentValue.BlobCleanupPollFrequency, new BlobCleanupState())
         {
             _settings = settings;
 
@@ -63,7 +59,9 @@ namespace Horde.Storage.Implementation
         public override async Task<bool> OnPoll(BlobCleanupState state, CancellationToken cancellationToken)
         {
             if (_alreadyPolling)
+            {
                 return false;
+            }
 
             _alreadyPolling = true;
             try
@@ -82,7 +80,9 @@ namespace Horde.Storage.Implementation
             foreach (IBlobCleanup blobCleanup in state.BlobCleanups)
             {
                 if (!blobCleanup.ShouldRun())
+                {
                     continue;
+                }
 
                 string type = blobCleanup.GetType().ToString();
                 _logger.Information("Blob cleanup running for {BlobCleanup}", type);
diff --git a/Engine/Source/Programs/Horde/HordeStorage/Horde.Storage/Implementation/GC/CacheCleanup.cs b/Engine/Source/Programs/Horde/HordeStorage/Horde.Storage/Implementation/GC/CacheCleanup.cs
--- a/Engine/Source/Programs/Horde/HordeStorage/Horde.Storage/Implementation/GC/CacheCleanup.cs
+++ b/Engine/Source/Programs/Horde/HordeStorage/Horde.Storage/Implementation/GC/CacheCleanup.cs
@@ -1,7 +1,6 @@
 ﻿// Copyright Epic Games, Inc. All Rights Reserved.
 
 using System;
-using System.Collections.Generic;
 using System.Threading;
 using System.Threading.Tasks;
 using Datadog.Trace;
@@ -41,7 +40,7 @@ namespace Horde.Storage.Implementation
 
         public Task<int> Cleanup(NamespaceId ns, CancellationToken cancellationToken)
         {
-            NamespaceSettings.PerNamespaceSettings policies;
+            NamespacePolicy policies;
             try
             {
                 policies = _namespacePolicyResolver.GetPoliciesForNs(ns);
@@ -87,7 +86,9 @@ namespace Horde.Storage.Implementation
                 Interlocked.Increment(ref consideredCount);
 
                 if (lastAccessTime > cutoffTime)
+                {
                     continue;
+                }
 
                 _logger.Information("Attempting to delete object {Namespace} {Bucket} {Name} as it was last updated {LastAccessTime} which is older then {CutoffTime}", ns, bucket, name, lastAccessTime, cutoffTime);
                 using IScope scope = Tracer.Instance.StartActive("gc.ref");
@@ -139,7 +140,7 @@ namespace Horde.Storage.Implementation
 
                 try
                 {
-                    Task.WaitAll(storeDelete, transactionLogDelete);
+                    await Task.WhenAll(storeDelete, transactionLogDelete);
                 }
                 catch (Exception e)
                 {
diff --git a/Engine/Source/Programs/Horde/HordeStorage/Horde.Storage/Implementation/GC/OrphanBlobCleanup.cs b/Engine/Source/Programs/Horde/HordeStorage/Horde.Storage/Implementation/GC/OrphanBlobCleanup.cs
--- a/Engine/Source/Programs/Horde/HordeStorage/Horde.Storage/Implementation/GC/OrphanBlobCleanup.cs
+++ b/Engine/Source/Programs/Horde/HordeStorage/Horde.Storage/Implementation/GC/OrphanBlobCleanup.cs
@@ -30,17 +30,15 @@ namespace Horde.Storage.Implementation
     {
         private readonly IBlobService _blobService;
         private readonly ILeaderElection _leaderElection;
-        private readonly IOptionsMonitor<GCSettings> _gcSettings;
         private readonly INamespacePolicyResolver _namespacePolicyResolver;
         private readonly IRestClient _client;
         private readonly ILogger _logger = Log.ForContext<OrphanBlobCleanup>();
 
         // ReSharper disable once UnusedMember.Global
-        public OrphanBlobCleanup(IBlobService blobService, ILeaderElection leaderElection, IOptionsMonitor<CallistoTransactionLogSettings> callistoSettings, IOptionsMonitor<GCSettings> gcSettings, IServiceCredentials serviceCredentials, INamespacePolicyResolver namespacePolicyResolver)
+        public OrphanBlobCleanup(IBlobService blobService, ILeaderElection leaderElection, IOptionsMonitor<CallistoTransactionLogSettings> callistoSettings, IServiceCredentials serviceCredentials, INamespacePolicyResolver namespacePolicyResolver)
         {
             _blobService = blobService;
             _leaderElection = leaderElection;
-            _gcSettings = gcSettings;
             _namespacePolicyResolver = namespacePolicyResolver;
 
             _client = new RestClient(callistoSettings.CurrentValue.ConnectionString)
@@ -49,12 +47,11 @@ namespace Horde.Storage.Implementation
             }.UseSerializer(() => new JsonNetSerializer());
         }
 
-        internal OrphanBlobCleanup(IBlobService blobService, ILeaderElection leaderElection, IRestClient callistoClient, IOptionsMonitor<GCSettings> gcSettings, INamespacePolicyResolver namespacePolicyResolver)
+        internal OrphanBlobCleanup(IBlobService blobService, ILeaderElection leaderElection, IRestClient callistoClient, INamespacePolicyResolver namespacePolicyResolver)
         {
             _blobService = blobService;
             _leaderElection = leaderElection;
             _client = callistoClient;
-            _gcSettings = gcSettings;
             _namespacePolicyResolver = namespacePolicyResolver;
         }
 
@@ -70,7 +67,6 @@ namespace Horde.Storage.Implementation
                 TransactionLogGeneration = transactionLogGeneration;
                 GcRoots = gcRoots;
             }
-
         }
 
         public bool ShouldRun()
@@ -91,7 +87,6 @@ namespace Horde.Storage.Implementation
                 return 0;
             }
 
-
             List<NamespaceId> namespaces = await ListNamespaces().Where(NamespaceShouldBeCleaned).ToListAsync();
             ConcurrentDictionary<NamespaceId, GCRootState> perNamespaceRoots = new();
             await namespaces.ParallelForEachAsync(async ns =>
@@ -111,7 +106,9 @@ namespace Horde.Storage.Implementation
                 await foreach ((BlobIdentifier blob, DateTime lastModified) in _blobService.ListObjects(@namespace).WithCancellation(cancellationToken))
                 {
                     if (lastModified > DateTime.Now.AddMinutes(-60))
+                    {
                         continue;
+                    }
 
                     // TODO: This could be turned into a parallel for as each iteration does not mutate the state
                     // there is no parallel for each for async enumerable yet
@@ -147,7 +144,10 @@ namespace Horde.Storage.Implementation
             foreach (BlobIdentifier blob in blobsToRemove)
             {
                 if (cancellationToken.IsCancellationRequested)
+                {
                     break;
+                }
+
                 _logger.Information("Attempting to GC Orphan blob {Blob}", blob);
 
                 bool deleted = true;
@@ -177,7 +177,7 @@ namespace Horde.Storage.Implementation
 
         private bool NamespaceShouldBeCleaned(NamespaceId ns)
         {
-            NamespaceSettings.PerNamespaceSettings policy = _namespacePolicyResolver.GetPoliciesForNs(ns);
+            NamespacePolicy policy = _namespacePolicyResolver.GetPoliciesForNs(ns);
 
             return policy.IsLegacyNamespace.HasValue && policy.IsLegacyNamespace.Value;
         }
@@ -198,18 +198,18 @@ namespace Horde.Storage.Implementation
             }
         }
 
+        [System.Diagnostics.CodeAnalysis.SuppressMessage("Design", "CA1034:Nested types should not be visible", Justification = "Used by serialization")]
         public class ListNamespaceResponse
         {
             public NamespaceId[] Logs { get; set; } = null!;
         }
 
-
         private async Task<bool> HasBlobBeenAddedToGcSet(NamespaceId ns, BlobIdentifier blob, long transactionLogPointer, Guid transactionLogGeneration, CancellationToken cancellationToken)
         {
             // request all events since the transaction log pointer and make sure this blob has not been added since
 
             CallistoReader reader = new CallistoReader(_client, ns);
-            await foreach (bool found in reader.GetOps(transactionLogPointer, transactionLogGeneration, maxOffsetsAttempted: 1000000).Any(transactionEvent =>
+            await foreach (bool found in reader.GetOps(transactionLogPointer, transactionLogGeneration, maxOffsetsAttempted: 1000000, cancellationToken: cancellationToken).Any(transactionEvent =>
             {
                 if (transactionEvent is AddTransactionEvent addTransactionEvent)
                 {
@@ -239,14 +239,18 @@ namespace Horde.Storage.Implementation
             await foreach (TransactionEvent transactionEvent in reader.GetOps(cancellationToken: cancellationToken, maxOffsetsAttempted: 1000000, enumerationState: enumerationState).Prefetch(prefetchCount).WithCancellation(cancellationToken))
             {
                 if (transactionEvent == null)
+                {
                     continue;
-                
+                }
+
                 string key = $"{transactionEvent.Bucket}.{transactionEvent.Name}";
                 switch (transactionEvent)
                 {
                     case AddTransactionEvent add:
                         if (add.Blobs == null)
+                        {
                             throw new Exception("Expected to find blobs in a add transaction event");
+                        }
 
                         if (!events.TryAdd(key, add.Blobs))
                         {
diff --git a/Engine/Source/Programs/Horde/HordeStorage/Horde.Storage/Implementation/GC/OrphanBlobCleanupRefs.cs b/Engine/Source/Programs/Horde/HordeStorage/Horde.Storage/Implementation/GC/OrphanBlobCleanupRefs.cs
--- a/Engine/Source/Programs/Horde/HordeStorage/Horde.Storage/Implementation/GC/OrphanBlobCleanupRefs.cs
+++ b/Engine/Source/Programs/Horde/HordeStorage/Horde.Storage/Implementation/GC/OrphanBlobCleanupRefs.cs
@@ -14,7 +14,6 @@ using Horde.Storage.Implementation.Blob;
 using Jupiter;
 using Jupiter.Common;
 using Jupiter.Implementation;
-using Microsoft.Extensions.Options;
 using Serilog;
 
 namespace Horde.Storage.Implementation
@@ -25,18 +24,16 @@ namespace Horde.Storage.Implementation
         private readonly IObjectService _objectService;
         private readonly IBlobIndex _blobIndex;
         private readonly ILeaderElection _leaderElection;
-        private readonly IOptionsMonitor<GCSettings> _gcSettings;
         private readonly INamespacePolicyResolver _namespacePolicyResolver;
         private readonly ILogger _logger = Log.ForContext<OrphanBlobCleanupRefs>();
 
         // ReSharper disable once UnusedMember.Global
-        public OrphanBlobCleanupRefs(IBlobService blobService, IObjectService objectService, IBlobIndex blobIndex, ILeaderElection leaderElection, IOptionsMonitor<GCSettings> gcSettings, INamespacePolicyResolver namespacePolicyResolver)
+        public OrphanBlobCleanupRefs(IBlobService blobService, IObjectService objectService, IBlobIndex blobIndex, ILeaderElection leaderElection, INamespacePolicyResolver namespacePolicyResolver)
         {
             _blobService = blobService;
             _objectService = objectService;
             _blobIndex = blobIndex;
             _leaderElection = leaderElection;
-            _gcSettings = gcSettings;
             _namespacePolicyResolver = namespacePolicyResolver;
         }
 
@@ -65,7 +62,9 @@ namespace Horde.Storage.Implementation
             foreach (NamespaceId @namespace in namespaces)
             {
                 if (cancellationToken.IsCancellationRequested)
+                {
                     break;
+                }
 
                 // only consider blobs that have been around for 60 minutes
                 // this due to cases were blobs are uploaded first
@@ -73,24 +72,30 @@ namespace Horde.Storage.Implementation
                 await foreach ((BlobIdentifier blob, DateTime lastModified) in _blobService.ListObjects(@namespace).WithCancellation(cancellationToken))
                 {
                     if (lastModified > cutoff)
+                    {
                         continue;
-                    
+                    }
+
                     if (cancellationToken.IsCancellationRequested)
+                    {
                         break;
+                    }
 
                     using IScope removeBlobScope = Tracer.Instance.StartActive("gc.blob");
                     removeBlobScope.Span.ResourceName = $"{@namespace}.{blob}";
 
                     bool found = false;
-                    NamespaceSettings.PerNamespaceSettings policy = _namespacePolicyResolver.GetPoliciesForNs(@namespace);
+                    NamespacePolicy policy = _namespacePolicyResolver.GetPoliciesForNs(@namespace);
                     
                     // check all other namespaces that share the same storage pool for presence of the blob
                     foreach (NamespaceId blobNamespace in namespaces.Where(ns => _namespacePolicyResolver.GetPoliciesForNs(ns).StoragePool == policy.StoragePool))
                     {
                         if (cancellationToken.IsCancellationRequested)
+                        {
                             break;
-                        
-                        IBlobIndex.BlobInfo? blobIndex = await _blobIndex.GetBlobInfo(blobNamespace, blob);
+                        }
+
+                        BlobInfo? blobIndex = await _blobIndex.GetBlobInfo(blobNamespace, blob);
 
                         if (blobIndex == null)
                         {
@@ -100,7 +105,9 @@ namespace Horde.Storage.Implementation
                         foreach ((BucketId bucket, IoHashKey key) in blobIndex.References)
                         {
                             if (cancellationToken.IsCancellationRequested)
+                            {
                                 break;
+                            }
 
                             try
                             {
@@ -120,7 +127,9 @@ namespace Horde.Storage.Implementation
                     }
 
                     if (cancellationToken.IsCancellationRequested)
+                    {
                         break;
+                    }
 
                     removeBlobScope.Span.SetTag("removed", (!found).ToString());
 
@@ -152,7 +161,7 @@ namespace Horde.Storage.Implementation
         {
             try
             {
-                NamespaceSettings.PerNamespaceSettings policy = _namespacePolicyResolver.GetPoliciesForNs(ns);
+                NamespacePolicy policy = _namespacePolicyResolver.GetPoliciesForNs(ns);
 
                 return policy.IsLegacyNamespace.HasValue && !policy.IsLegacyNamespace.Value;
             }
diff --git a/Engine/Source/Programs/Horde/HordeStorage/Horde.Storage/Implementation/GC/RefCleanupService.cs b/Engine/Source/Programs/Horde/HordeStorage/Horde.Storage/Implementation/GC/RefCleanupService.cs
--- a/Engine/Source/Programs/Horde/HordeStorage/Horde.Storage/Implementation/GC/RefCleanupService.cs
+++ b/Engine/Source/Programs/Horde/HordeStorage/Horde.Storage/Implementation/GC/RefCleanupService.cs
@@ -5,7 +5,6 @@ using System.Collections.Generic;
 using System.Threading;
 using System.Threading.Tasks;
 using Dasync.Collections;
-using Datadog.Trace;
 using EpicGames.Horde.Storage;
 using Jupiter;
 using Microsoft.Extensions.Options;
@@ -13,7 +12,19 @@ using Serilog;
 
 namespace Horde.Storage.Implementation
 {
-    public class RefCleanupService : PollingService<RefCleanupService.RefCleanupState>
+    public class RefCleanupState
+    {
+        public RefCleanupState(IRefsStore refs, IRefCleanup refCleanup)
+        {
+            Refs = refs;
+            RefCleanup = refCleanup;
+        }
+
+        public IRefsStore Refs { get; }
+        public IRefCleanup RefCleanup { get; }
+    }
+
+    public class RefCleanupService : PollingService<RefCleanupState>
     {
         private readonly IOptionsMonitor<GCSettings> _settings;
         private readonly ILeaderElection _leaderElection;
@@ -21,21 +32,7 @@ namespace Horde.Storage.Implementation
         private volatile bool _alreadyPolling;
         private readonly ILogger _logger = Log.ForContext<RefCleanupService>();
 
-
-        public class RefCleanupState
-        {
-            public RefCleanupState(IRefsStore refs, IRefCleanup refCleanup)
-            {
-                Refs = refs;
-                RefCleanup = refCleanup;
-            }
-
-            public IRefsStore Refs { get; }
-            public IRefCleanup RefCleanup { get; }
-        }
-
-        public RefCleanupService(IOptionsMonitor<GCSettings> settings, IRefsStore store, IRefCleanup refCleanup, ILeaderElection leaderElection, IReferencesStore referencesStore) : 
-            base(serviceName: nameof(RefCleanupService), settings.CurrentValue.RefCleanupPollFrequency, new RefCleanupState(store, refCleanup))
+        public RefCleanupService(IOptionsMonitor<GCSettings> settings, IRefsStore store, IRefCleanup refCleanup, ILeaderElection leaderElection, IReferencesStore referencesStore) : base(serviceName: nameof(RefCleanupService), settings.CurrentValue.RefCleanupPollFrequency, new RefCleanupState(store, refCleanup))
         {
             _settings = settings;
             _leaderElection = leaderElection;
@@ -50,7 +47,9 @@ namespace Horde.Storage.Implementation
         public override async Task<bool> OnPoll(RefCleanupState state, CancellationToken cancellationToken)
         {
             if (_alreadyPolling)
+            {
                 return false;
+            }
 
             _alreadyPolling = true;
             try
diff --git a/Engine/Source/Programs/Horde/HordeStorage/Horde.Storage/Implementation/IReferenceResolver.cs b/Engine/Source/Programs/Horde/HordeStorage/Horde.Storage/Implementation/IReferenceResolver.cs
--- a/Engine/Source/Programs/Horde/HordeStorage/Horde.Storage/Implementation/IReferenceResolver.cs
+++ b/Engine/Source/Programs/Horde/HordeStorage/Horde.Storage/Implementation/IReferenceResolver.cs
@@ -61,7 +61,9 @@ namespace Horde.Storage.Implementation
                             pendingCompactBinaryAttachments.Add(ParseCompactBinaryAttachment(ns, blobIdentifier));
                         }
                         else
+                        { 
                             throw new NotImplementedException($"Unknown attachment type for field {field}");
+                        }
                     });
 
                     foreach (BlobIdentifier blobIdentifier in blobIdentifiers)
diff --git a/Engine/Source/Programs/Horde/HordeStorage/Horde.Storage/Implementation/LastAccess/LastAccessService.cs b/Engine/Source/Programs/Horde/HordeStorage/Horde.Storage/Implementation/LastAccess/LastAccessService.cs
--- a/Engine/Source/Programs/Horde/HordeStorage/Horde.Storage/Implementation/LastAccess/LastAccessService.cs
+++ b/Engine/Source/Programs/Horde/HordeStorage/Horde.Storage/Implementation/LastAccess/LastAccessService.cs
@@ -58,9 +58,7 @@ namespace Horde.Storage.Implementation
         private void OnUpdate(object? state)
         {
             // call results to make sure we join the task
-            Task.WaitAll(
-                ProcessLastAccessRecords()
-            );
+            ProcessLastAccessRecords().Wait();
         }
 
         internal async Task<List<(RefRecord, DateTime)>> ProcessLastAccessRecords()
@@ -70,7 +68,9 @@ namespace Horde.Storage.Implementation
             foreach ((RefRecord record, DateTime lastAccessTime) in records)
             {
                 if (!ShouldTrackLastAccess(record.Namespace))
+                {
                     continue;
+                }
 
                 using IScope scope = Tracer.Instance.StartActive("lastAccess.update");
                 scope.Span.ResourceName = $"{record.Namespace}:{record.Bucket}.{record.RefName}";
@@ -88,7 +88,16 @@ namespace Horde.Storage.Implementation
 
         public void Dispose()
         {
-            _timer?.Dispose();
+            Dispose(true);
+            GC.SuppressFinalize(this);
+        }
+
+        protected virtual void Dispose(bool disposing)
+        {
+            if (disposing)
+            {
+                _timer?.Dispose();
+            }
         }
     }
 }
diff --git a/Engine/Source/Programs/Horde/HordeStorage/Horde.Storage/Implementation/LastAccess/LastAccessServiceReferences.cs b/Engine/Source/Programs/Horde/HordeStorage/Horde.Storage/Implementation/LastAccess/LastAccessServiceReferences.cs
--- a/Engine/Source/Programs/Horde/HordeStorage/Horde.Storage/Implementation/LastAccess/LastAccessServiceReferences.cs
+++ b/Engine/Source/Programs/Horde/HordeStorage/Horde.Storage/Implementation/LastAccess/LastAccessServiceReferences.cs
@@ -60,9 +60,7 @@ namespace Horde.Storage.Implementation
             try
             {
                 // call results to make sure we join the task
-                Task.WaitAll(
-                    ProcessLastAccessRecords()
-                );
+                ProcessLastAccessRecords().Wait();
             }
             catch (Exception e)
             {
@@ -77,7 +75,10 @@ namespace Horde.Storage.Implementation
             foreach ((LastAccessRecord record, DateTime lastAccessTime) in records)
             {
                 if (!ShouldTrackLastAccess(record.Namespace))
+                {
                     continue;
+                }
+
                 using IScope scope = Tracer.Instance.StartActive("lastAccess.update");
                 scope.Span.ResourceName = $"{record.Namespace}:{record.Bucket}.{record.Key}";
                 _logger.Debug("Updating last access time to {LastAccessTime} for {Record}", lastAccessTime, record);
@@ -96,7 +97,17 @@ namespace Horde.Storage.Implementation
 
         public void Dispose()
         {
-            _timer?.Dispose();
+            Dispose(true);
+            GC.SuppressFinalize(this);
+        }
+
+        protected virtual void Dispose(bool disposing)
+        {
+            if (disposing)
+            {
+                _timer?.Dispose();
+
+            }
         }
     }
 }
diff --git a/Engine/Source/Programs/Horde/HordeStorage/Horde.Storage/Implementation/LastAccess/LastAccessTracker.cs b/Engine/Source/Programs/Horde/HordeStorage/Horde.Storage/Implementation/LastAccess/LastAccessTracker.cs
--- a/Engine/Source/Programs/Horde/HordeStorage/Horde.Storage/Implementation/LastAccess/LastAccessTracker.cs
+++ b/Engine/Source/Programs/Horde/HordeStorage/Horde.Storage/Implementation/LastAccess/LastAccessTracker.cs
@@ -62,7 +62,7 @@ namespace Horde.Storage.Implementation
         // we will exchange the refs dictionary when fetching the records and use a rw lock to make sure no-one is trying to add things at the same time
         private readonly ReaderWriterLock _rwLock = new ReaderWriterLock();
 
-        public LastAccessTracker(IOptionsMonitor<HordeStorageSettings> settings)
+        protected LastAccessTracker(IOptionsMonitor<HordeStorageSettings> settings)
         {
             _settings = settings;
         }
@@ -72,7 +72,9 @@ namespace Horde.Storage.Implementation
         public Task TrackUsed(T record)
         {
             if (!_settings.CurrentValue.EnableLastAccessTracking)
+            {
                 return Task.CompletedTask;
+            }
 
             return Task.Run(() =>
             {
@@ -134,5 +136,4 @@ namespace Horde.Storage.Implementation
             }
         }
     }
-
 }
diff --git a/Engine/Source/Programs/Horde/HordeStorage/Horde.Storage/Implementation/LeaderElection/ILeaderElection.cs b/Engine/Source/Programs/Horde/HordeStorage/Horde.Storage/Implementation/LeaderElection/ILeaderElection.cs
--- a/Engine/Source/Programs/Horde/HordeStorage/Horde.Storage/Implementation/LeaderElection/ILeaderElection.cs
+++ b/Engine/Source/Programs/Horde/HordeStorage/Horde.Storage/Implementation/LeaderElection/ILeaderElection.cs
@@ -1,26 +1,23 @@
 ﻿// Copyright Epic Games, Inc. All Rights Reserved.
 
 using System;
-using System.Collections.Generic;
-using System.Linq;
-using System.Threading.Tasks;
 
 namespace Horde.Storage.Implementation
 {
-    public interface ILeaderElection
+    public class OnLeaderChangedEventArgs
     {
-        public class OnLeaderChangedEventArgs
+        public bool IsLeader { get; }
+        public string LeaderName { get; }
+
+        public OnLeaderChangedEventArgs(bool isLeader, string leaderName)
         {
-            public bool IsLeader { get; }
-            public string LeaderName { get; }
-
-            public OnLeaderChangedEventArgs(bool isLeader, string leaderName)
-            {
-                IsLeader = isLeader;
-                LeaderName = leaderName;
-            }
+            IsLeader = isLeader;
+            LeaderName = leaderName;
         }
+    }
 
+    public interface ILeaderElection
+    {
         bool IsThisInstanceLeader();
 
         event EventHandler<OnLeaderChangedEventArgs>? OnLeaderChanged;
@@ -34,7 +31,7 @@ namespace Horde.Storage.Implementation
         {
             _isLeader = isLeader;
 
-            OnLeaderChanged?.Invoke(this, new ILeaderElection.OnLeaderChangedEventArgs(isLeader, ""));
+            OnLeaderChanged?.Invoke(this, new OnLeaderChangedEventArgs(isLeader, ""));
         }
 
         public bool IsThisInstanceLeader()
@@ -42,6 +39,6 @@ namespace Horde.Storage.Implementation
             return _isLeader;
         }
 
-        public event EventHandler<ILeaderElection.OnLeaderChangedEventArgs>? OnLeaderChanged;
+        public event EventHandler<OnLeaderChangedEventArgs>? OnLeaderChanged;
     }
 }
diff --git a/Engine/Source/Programs/Horde/HordeStorage/Horde.Storage/Implementation/LeaderElection/KubernetesLeaderElection.cs b/Engine/Source/Programs/Horde/HordeStorage/Horde.Storage/Implementation/LeaderElection/KubernetesLeaderElection.cs
--- a/Engine/Source/Programs/Horde/HordeStorage/Horde.Storage/Implementation/LeaderElection/KubernetesLeaderElection.cs
+++ b/Engine/Source/Programs/Horde/HordeStorage/Horde.Storage/Implementation/LeaderElection/KubernetesLeaderElection.cs
@@ -39,8 +39,8 @@ namespace Horde.Storage.Implementation.LeaderElection
 
         private readonly Kubernetes _client;
         private readonly LeaderElector _leaderElector;
-        private ConfigMapLock _configMapLock;
-        private string _identity;
+        private readonly ConfigMapLock _configMapLock;
+        private readonly string _identity;
 
         public KubernetesLeaderElection(IOptionsMonitor<KubernetesLeaderElectionSettings> leaderSettings) : base("Kubernetes Leader Election", TimeSpan.FromSeconds(1), new KubernetesLeaderElectionState(), startAtRandomTime: true)
         {
@@ -70,8 +70,8 @@ namespace Horde.Storage.Implementation.LeaderElection
         {
             _logger.Warning("{Instance} is the new leader", leaderName);
 
-            bool isLeader = string.Equals(leaderName, _identity, StringComparison.InvariantCultureIgnoreCase);
-            OnLeaderChanged?.Invoke(this, new ILeaderElection.OnLeaderChangedEventArgs(isLeader, leaderName));
+            bool isLeader = string.Equals(leaderName, _identity, StringComparison.OrdinalIgnoreCase);
+            OnLeaderChanged?.Invoke(this, new OnLeaderChangedEventArgs(isLeader, leaderName));
         }
 
         public bool IsThisInstanceLeader()
@@ -79,7 +79,7 @@ namespace Horde.Storage.Implementation.LeaderElection
             return _leaderElector.IsLeader();
         }
 
-        public event EventHandler<ILeaderElection.OnLeaderChangedEventArgs>? OnLeaderChanged;
+        public event EventHandler<OnLeaderChangedEventArgs>? OnLeaderChanged;
 
         protected override Task OnStopping(KubernetesLeaderElectionState state)
         {
diff --git a/Engine/Source/Programs/Horde/HordeStorage/Horde.Storage/Implementation/MongoStore.cs b/Engine/Source/Programs/Horde/HordeStorage/Horde.Storage/Implementation/MongoStore.cs
--- a/Engine/Source/Programs/Horde/HordeStorage/Horde.Storage/Implementation/MongoStore.cs
+++ b/Engine/Source/Programs/Horde/HordeStorage/Horde.Storage/Implementation/MongoStore.cs
@@ -1,12 +1,10 @@
 // Copyright Epic Games, Inc. All Rights Reserved.
 
 using System;
-using System.Collections.Generic;
 using System.Security.Authentication;
 using System.Threading.Tasks;
 using Microsoft.Extensions.Options;
 using MongoDB.Bson;
-using MongoDB.Bson.Serialization;
 using MongoDB.Driver;
 
 namespace Horde.Storage.Implementation;
@@ -14,7 +12,7 @@ namespace Horde.Storage.Implementation;
 public class MongoStore
 {
     private readonly string? _overrideDatabaseName;
-    protected readonly MongoClient _client;
+    private readonly MongoClient _client;
 
     public MongoStore(IOptionsMonitor<MongoSettings> settings, string? overrideDatabaseName = null)
     {
@@ -24,7 +22,7 @@ public class MongoStore
         );
         if (settings.CurrentValue.RequireTls12)
         {
-            mongoClientSettings.SslSettings = new SslSettings {EnabledSslProtocols = SslProtocols.Tls12};
+            mongoClientSettings.SslSettings = new SslSettings {EnabledSslProtocols = SslProtocols.None};
         }
 
         _client = new MongoClient(mongoClientSettings);
@@ -38,7 +36,7 @@ public class MongoStore
 
         // Try to avoid exceptions breaking in the debugger unnecessarily by checking for the existence of the collection beforehand.
         FilterDefinition<BsonDocument> filter = new BsonDocument("name", collectionName);
-        if (await database.ListCollectionNames(new ListCollectionNamesOptions { Filter = filter }).AnyAsync())
+        if (await (await database.ListCollectionNamesAsync(new ListCollectionNamesOptions { Filter = filter })).AnyAsync())
         {
             return;
         }
@@ -50,7 +48,9 @@ public class MongoStore
         catch (MongoCommandException e)
         {
             if (e.CodeName != "NamespaceExists")
+            {
                 throw ;
+            }
         }
     }
 
@@ -83,7 +83,8 @@ public class MongoStore
     }
 }
 
-public class MongoCollectionNameAttribute : Attribute
+[AttributeUsage(AttributeTargets.Class)]
+public sealed class MongoCollectionNameAttribute : Attribute
 {
     public string CollectionName { get; }
 
