Horde Storage - Updated editorconfig to more closely match what we run for Horde and updated code to fix all static analysis errors

diff --git a/Engine/Source/Programs/Horde/HordeStorage/Horde.Storage/Implementation/ObjectService.cs b/Engine/Source/Programs/Horde/HordeStorage/Horde.Storage/Implementation/ObjectService.cs
--- a/Engine/Source/Programs/Horde/HordeStorage/Horde.Storage/Implementation/ObjectService.cs
+++ b/Engine/Source/Programs/Horde/HordeStorage/Horde.Storage/Implementation/ObjectService.cs
@@ -2,7 +2,6 @@
 
 using System;
 using System.Collections.Generic;
-using System.IO;
 using System.Linq;
 using System.Threading.Tasks;
 using async_enumerable_dotnet;
@@ -57,7 +56,9 @@ namespace Horde.Storage.Implementation
             {
                 // empty array means fetch all fields
                 if (fields.Length == 0)
+                {
                     flags = IReferencesStore.FieldFlags.All;
+                }
                 else
                 {
                     flags = fields.Contains("payload")
@@ -68,20 +69,22 @@ namespace Horde.Storage.Implementation
             ObjectRecord o = await _referencesStore.Get(ns, bucket, key, flags);
 
             // we do not wait for the last access tracking as it does not matter when it completes
-            Task lastAccessTask = _lastAccessTracker.TrackUsed(new LastAccessRecord(ns, bucket, key)).ContinueWith(task =>
+            Task lastAccessTask = _lastAccessTracker.TrackUsed(new LastAccessRecord(ns, bucket, key)).ContinueWith((task, _) =>
             {
                 if (task.Exception != null)
                 {
                     _logger.Error(task.Exception, "Exception when tracking last access record");
                 }
-            });;
+            }, null, TaskScheduler.Current);
 
             BlobContents? blobContents = null;
             if ((flags & IReferencesStore.FieldFlags.IncludePayload) != 0)
             {
                 if (o.InlinePayload != null && o.InlinePayload.Length != 0)
                 {
+#pragma warning disable CA2000 // Dispose objects before losing scope , ownership is transfered to caller
                     blobContents = new BlobContents(o.InlinePayload);
+#pragma warning restore CA2000 // Dispose objects before losing scope
                 }
                 else
                 {
@@ -116,7 +119,9 @@ namespace Horde.Storage.Implementation
                 {
                     bool hasAttachment = HasAttachments(field.AsObject());
                     if (hasAttachment)
+                    {
                         return true;
+                    }
                 }
 
                 if (field.IsArray())
@@ -125,7 +130,9 @@ namespace Horde.Storage.Implementation
                     {
                         bool hasAttachment = FieldHasAttachments(subField);
                         if (hasAttachment)
+                        {
                             return true;
+                        }
                     }
                 }
 
@@ -139,12 +146,17 @@ namespace Horde.Storage.Implementation
         {
             (ObjectRecord o, BlobContents? blob) = await Get(ns, bucket, key);
             if (blob == null)
+            {
                 throw new InvalidOperationException("No blob when attempting to finalize");
+            }
+
             byte[] blobContents = await blob.Stream.ToByteArray();
             CbObject payload = new CbObject(blobContents);
 
             if (!o.BlobIdentifier.Equals(blobHash))
+            {
                 throw new ObjectHashMismatchException(ns, bucket, key, blobHash, o.BlobIdentifier);
+            }
 
             return await DoFinalize(ns, bucket, key, blobHash, payload);
         }
diff --git a/Engine/Source/Programs/Horde/HordeStorage/Horde.Storage/Implementation/Objects/CacheReferencesStore.cs b/Engine/Source/Programs/Horde/HordeStorage/Horde.Storage/Implementation/Objects/CacheReferencesStore.cs
--- a/Engine/Source/Programs/Horde/HordeStorage/Horde.Storage/Implementation/Objects/CacheReferencesStore.cs
+++ b/Engine/Source/Programs/Horde/HordeStorage/Horde.Storage/Implementation/Objects/CacheReferencesStore.cs
@@ -107,11 +107,14 @@ namespace Horde.Storage.Implementation
 
         public async Task<ObjectRecord> Get(NamespaceId ns, BucketId bucket, IoHashKey key, IReferencesStore.FieldFlags flags)
         {
-            HttpRequestMessage getObjectRequest = BuildHttpRequest(HttpMethod.Get, $"api/v1/refs/{ns}/{bucket}/{key}/metadata");
+            using HttpRequestMessage getObjectRequest = BuildHttpRequest(HttpMethod.Get, new Uri($"api/v1/refs/{ns}/{bucket}/{key}/metadata", UriKind.Relative));
             getObjectRequest.Headers.Add("Accept", MediaTypeNames.Application.Json);
             HttpResponseMessage response = await HttpClient.SendAsync(getObjectRequest);
             if (response.StatusCode == HttpStatusCode.NotFound)
+            {
                 throw new ObjectNotFoundException(ns, bucket, key);
+            }
+
             response.EnsureSuccessStatusCode();
 
             RefMetadataResponse metadataResponse = await response.Content.ReadAsAsync<RefMetadataResponse>();
@@ -121,7 +124,7 @@ namespace Horde.Storage.Implementation
 
         public async Task Put(NamespaceId ns, BucketId bucket, IoHashKey key, BlobIdentifier blobHash, byte[] blob, bool isFinalized)
         {
-            HttpRequestMessage putObjectRequest = BuildHttpRequest(HttpMethod.Put, $"api/v1/refs/{ns}/{bucket}/{key}");
+            using HttpRequestMessage putObjectRequest = BuildHttpRequest(HttpMethod.Put, new Uri($"api/v1/refs/{ns}/{bucket}/{key}", UriKind.Relative));
             putObjectRequest.Headers.Add("Accept", MediaTypeNames.Application.Json);
             putObjectRequest.Headers.Add(CommonHeaders.HashHeaderName, blobHash.ToString());
             putObjectRequest.Content = new ByteArrayContent(blob);
@@ -141,7 +144,7 @@ namespace Horde.Storage.Implementation
 
         public async Task Finalize(NamespaceId ns, BucketId bucket, IoHashKey key, BlobIdentifier blobIdentifier)
         {
-            HttpRequestMessage putObjectRequest = BuildHttpRequest(HttpMethod.Post, $"api/v1/refs/{ns}/{bucket}/{key}/finalize/{blobIdentifier}");
+            using HttpRequestMessage putObjectRequest = BuildHttpRequest(HttpMethod.Post, new Uri($"api/v1/refs/{ns}/{bucket}/{key}/finalize/{blobIdentifier}", UriKind.Relative));
             putObjectRequest.Headers.Add("Accept", MediaTypeNames.Application.Json);
 
             HttpResponseMessage response = await HttpClient.SendAsync(putObjectRequest);
diff --git a/Engine/Source/Programs/Horde/HordeStorage/Horde.Storage/Implementation/Objects/MemoryReferencesStore.cs b/Engine/Source/Programs/Horde/HordeStorage/Horde.Storage/Implementation/Objects/MemoryReferencesStore.cs
--- a/Engine/Source/Programs/Horde/HordeStorage/Horde.Storage/Implementation/Objects/MemoryReferencesStore.cs
+++ b/Engine/Source/Programs/Horde/HordeStorage/Horde.Storage/Implementation/Objects/MemoryReferencesStore.cs
@@ -83,7 +83,7 @@ namespace Horde.Storage.Implementation
 
         public Task<bool> Delete(NamespaceId ns, BucketId bucket, IoHashKey key)
         {
-            if (!_objects.TryRemove(BuildKey(ns, bucket, key), out MemoryStoreObject? o))
+            if (!_objects.TryRemove(BuildKey(ns, bucket, key), out MemoryStoreObject? _))
             {
                 throw new ObjectNotFoundException(ns, bucket, key);
             }
@@ -186,6 +186,5 @@ namespace Horde.Storage.Implementation
             bool includePayload = (fieldFlags & IReferencesStore.FieldFlags.IncludePayload) != 0;
             return new ObjectRecord(Namespace, Bucket, Name, LastAccessTime, includePayload ? Blob : null, BlobHash, IsFinalized);
         }
-
     }
 }
diff --git a/Engine/Source/Programs/Horde/HordeStorage/Horde.Storage/Implementation/Objects/MongoReferencesStore.cs b/Engine/Source/Programs/Horde/HordeStorage/Horde.Storage/Implementation/Objects/MongoReferencesStore.cs
--- a/Engine/Source/Programs/Horde/HordeStorage/Horde.Storage/Implementation/Objects/MongoReferencesStore.cs
+++ b/Engine/Source/Programs/Horde/HordeStorage/Horde.Storage/Implementation/Objects/MongoReferencesStore.cs
@@ -8,14 +8,11 @@ using Jupiter.Implementation;
 using Microsoft.Extensions.Options;
 using MongoDB.Bson.Serialization.Attributes;
 using MongoDB.Driver;
-using Serilog;
 
 namespace Horde.Storage.Implementation
 {
     public class MongoReferencesStore : MongoStore, IReferencesStore
     {
-        private readonly ILogger _logger = Log.ForContext<MongoReferencesStore>();
-
         public MongoReferencesStore(IOptionsMonitor<MongoSettings> settings, string? overrideDatabaseName = null) : base(settings, overrideDatabaseName)
         {
             CreateCollectionIfNotExists<MongoReferencesModelV0>().Wait();
@@ -49,7 +46,9 @@ namespace Horde.Storage.Implementation
             IAsyncCursor<MongoReferencesModelV0>? cursor = await collection.FindAsync(m => m.Ns == ns.ToString() && m.Bucket == bucket.ToString() && m.Key == key.ToString());
             MongoReferencesModelV0? model = await cursor.FirstOrDefaultAsync();
             if (model == null)
+            {
                 throw new ObjectNotFoundException(ns, bucket, key);
+            }
 
             if (!includePayload)
             {
@@ -262,7 +261,6 @@ namespace Horde.Storage.Implementation
         [BsonRequired]
         public string Ns { get; set; }
 
-
         [BsonConstructor]
         public MongoNamespacesModelV0(string ns)
         {
@@ -273,6 +271,5 @@ namespace Horde.Storage.Implementation
         {
             Ns = ns.ToString();
         }
-
     }
 }
diff --git a/Engine/Source/Programs/Horde/HordeStorage/Horde.Storage/Implementation/Objects/ScyllaReferencesStore.cs b/Engine/Source/Programs/Horde/HordeStorage/Horde.Storage/Implementation/Objects/ScyllaReferencesStore.cs
--- a/Engine/Source/Programs/Horde/HordeStorage/Horde.Storage/Implementation/Objects/ScyllaReferencesStore.cs
+++ b/Engine/Source/Programs/Horde/HordeStorage/Horde.Storage/Implementation/Objects/ScyllaReferencesStore.cs
@@ -72,7 +72,9 @@ namespace Horde.Storage.Implementation
             }
 
             if (o == null)
+            {
                 throw new ObjectNotFoundException(ns, bucket, name);
+            }
 
             try
             {
@@ -105,7 +107,6 @@ namespace Horde.Storage.Implementation
             await addBucketTask;
         }
 
-
         public async Task Finalize(NamespaceId ns, BucketId bucket, IoHashKey name, BlobIdentifier blobIdentifier)
         {
             using IScope scope = Tracer.Instance.StartActive("scylla.finalize");
@@ -129,7 +130,9 @@ namespace Horde.Storage.Implementation
             foreach (Row row in rowSet)
             {
                 if (rowSet.GetAvailableWithoutFetching() == 0)
+                {
                     await rowSet.FetchMoreResultsAsync();
+                }
 
                 string bucket = row.GetValue<string>("bucket");
                 string name = row.GetValue<string>("name");
@@ -137,7 +140,10 @@ namespace Horde.Storage.Implementation
 
                 // skip any names that are not conformant to io hash
                 if (name.Length != 40)
+                {
                     continue;
+                }
+
                 yield return (new BucketId(bucket), new IoHashKey(name), lastAccessTime);
             }
         }
@@ -150,7 +156,9 @@ namespace Horde.Storage.Implementation
             foreach (Row row in rowSet)
             {
                 if (rowSet.GetAvailableWithoutFetching() == 0)
+                {
                     await rowSet.FetchMoreResultsAsync();
+                }
 
                 yield return new NamespaceId(row.GetValue<string>(0));
             }
@@ -164,7 +172,9 @@ namespace Horde.Storage.Implementation
             AppliedInfo<ScyllaObject> info = await _mapper.DeleteIfAsync<ScyllaObject>("WHERE namespace=? AND bucket=? AND name=? IF EXISTS", ns.ToString(), bucket.ToString(), key.ToString());
 
             if (info.Applied)
+            {
                 return true;
+            }
 
             return false;
         }
@@ -311,15 +321,29 @@ namespace Horde.Storage.Implementation
         public void ThrowIfRequiredFieldIsMissing(bool includePayload)
         {
             if (string.IsNullOrEmpty(Namespace))
-                throw new ArgumentException("Namespace was not valid", nameof(Namespace));
+            {
+                throw new InvalidOperationException("Namespace was not valid");
+            }
+
             if (string.IsNullOrEmpty(Bucket))
-                throw new ArgumentException("Bucket was not valid", nameof(Bucket));
+            {
+                throw new InvalidOperationException("Bucket was not valid");
+            }
+
             if (string.IsNullOrEmpty(Name))
-                throw new ArgumentException("Name was not valid", nameof(Name));
+            {
+                throw new InvalidOperationException("Name was not valid");
+            }
+
             if (PayloadHash == null && includePayload)
-                throw new ArgumentException("PayloadHash was not valid", nameof(PayloadHash));
+            {
+                throw new ArgumentException("PayloadHash was not valid");
+            }
+
             if (!IsFinalized.HasValue)
-                throw new ArgumentException("IsFinalized was not valid", nameof(IsFinalized));
+            {
+                throw new ArgumentException("IsFinalized was not valid");
+            }
         }
     }
 
@@ -340,6 +364,7 @@ namespace Horde.Storage.Implementation
         [Cassandra.Mapping.Attributes.PartitionKey]
         public string? Namespace { get; set; }
 
+        [System.Diagnostics.CodeAnalysis.SuppressMessage("Usage", "CA2227:Collection properties should be read only", Justification = "Used by serialization")]
         public List<string> Buckets { get; set; } = new List<string>();
     }
 
@@ -362,7 +387,7 @@ namespace Horde.Storage.Implementation
         [Cassandra.Mapping.Attributes.Column("accessed_at")]
         public LocalDate? AccessedAt { get; set; }
 
-
+        [System.Diagnostics.CodeAnalysis.SuppressMessage("Usage", "CA2227:Collection properties should be read only", Justification = "Used by serialization")]
         public List<string> Objects { get; set; } = new List<string>();
 
     }
diff --git a/Engine/Source/Programs/Horde/HordeStorage/Horde.Storage/Implementation/Objects/ScyllaSettings.cs b/Engine/Source/Programs/Horde/HordeStorage/Horde.Storage/Implementation/Objects/ScyllaSettings.cs
--- a/Engine/Source/Programs/Horde/HordeStorage/Horde.Storage/Implementation/Objects/ScyllaSettings.cs
+++ b/Engine/Source/Programs/Horde/HordeStorage/Horde.Storage/Implementation/Objects/ScyllaSettings.cs
@@ -17,6 +17,7 @@ public class ScyllaSettings
     /// Note that this only applies when creating the keyspace
     /// To modify a existing keyspace see https://docs.datastax.com/en/dse/6.7/dse-admin/datastax_enterprise/operations/opsChangeKSStrategy.html
     /// </summary>
+    [System.Diagnostics.CodeAnalysis.SuppressMessage("Usage", "CA2227:Collection properties should be read only", Justification = "Used by serialization")]
     public Dictionary<string, string>? KeyspaceReplicationStrategy { get; set; }
 
     /// <summary>
@@ -24,6 +25,7 @@ public class ScyllaSettings
     /// Note that this only applies when creating the keyspace
     /// To modify a existing keyspace see https://docs.datastax.com/en/dse/6.7/dse-admin/datastax_enterprise/operations/opsChangeKSStrategy.html
     /// </summary>
+    [System.Diagnostics.CodeAnalysis.SuppressMessage("Usage", "CA2227:Collection properties should be read only", Justification = "Used by serialization")]
     public Dictionary<string, string>? LocalKeyspaceReplicationStrategy { get; set; }
 
     /// <summary>
diff --git a/Engine/Source/Programs/Horde/HordeStorage/Horde.Storage/Implementation/PeerStatusService.cs b/Engine/Source/Programs/Horde/HordeStorage/Horde.Storage/Implementation/PeerStatusService.cs
--- a/Engine/Source/Programs/Horde/HordeStorage/Horde.Storage/Implementation/PeerStatusService.cs
+++ b/Engine/Source/Programs/Horde/HordeStorage/Horde.Storage/Implementation/PeerStatusService.cs
@@ -12,20 +12,20 @@ using Microsoft.Extensions.Options;
 
 namespace Horde.Storage.Implementation
 {
-    public interface IPeerStatusService
+    public class PeerStatus
     {
-        public class PeerStatus
+        public PeerStatus(PeerSettings peerSettings)
         {
-            public PeerStatus(PeerSettings peerSettings)
-            {
-                Endpoints = peerSettings.Endpoints;
-            }
-
-            public List<PeerEndpoints> Endpoints { get; set; }
-            public int Latency { get; set; }
-            public bool Reachable { get; set; }
+            Endpoints = peerSettings.Endpoints;
         }
 
+        public List<PeerEndpoints> Endpoints { get; }
+        public int Latency { get; set; }
+        public bool Reachable { get; set; }
+    }
+
+    public interface IPeerStatusService
+    {
         PeerStatus? GetPeerStatus(string regionName);
 
         /// <summary>
@@ -36,21 +36,20 @@ namespace Horde.Storage.Implementation
         List<(int, string)> GetPeersByLatency(IEnumerable<string> peerNames);
     }
 
-    public class PeerStatusService : PollingService<PeerStatusService.PeerStatusServiceState>, IPeerStatusService
+    public class PeerStatusServiceState
+    {
+    }
+    public class PeerStatusService : PollingService<PeerStatusServiceState>, IPeerStatusService
     {
         private readonly IOptionsMonitor<ClusterSettings> _clusterSettings;
         private readonly IOptionsMonitor<JupiterSettings> _jupiterSettings;
         private readonly IHttpClientFactory _clientFactory;
-        private readonly Dictionary<string, IPeerStatusService.PeerStatus> _peers = new Dictionary<string, IPeerStatusService.PeerStatus>(StringComparer.InvariantCultureIgnoreCase);
+        private readonly Dictionary<string, PeerStatus> _peers = new Dictionary<string, PeerStatus>(StringComparer.InvariantCultureIgnoreCase);
         private volatile bool _alreadyPolling = false;
 
-        public class PeerStatusServiceState
-        {
-        }
-
-        public IPeerStatusService.PeerStatus? GetPeerStatus(string regionName)
+        public PeerStatus? GetPeerStatus(string regionName)
         {
-            if (_peers.TryGetValue(regionName, out IPeerStatusService.PeerStatus? peerStatus))
+            if (_peers.TryGetValue(regionName, out PeerStatus? peerStatus))
             {
                 return peerStatus;
             }
@@ -64,12 +63,17 @@ namespace Horde.Storage.Implementation
             foreach (string peerName in peerNames)
             {
                 // skip the local site
-                if (string.Equals(peerName, _jupiterSettings.CurrentValue.CurrentSite, StringComparison.InvariantCultureIgnoreCase))
+                if (string.Equals(peerName, _jupiterSettings.CurrentValue.CurrentSite, StringComparison.OrdinalIgnoreCase))
+                {
                     continue;
+                }
 
-                IPeerStatusService.PeerStatus? peerStatus = GetPeerStatus(peerName);
+                PeerStatus? peerStatus = GetPeerStatus(peerName);
                 if (peerStatus == null)
+                {
                     continue;
+                }
+
                 peers.Add((peerStatus.Latency, peerName));
             }
 
@@ -86,10 +90,12 @@ namespace Horde.Storage.Implementation
             foreach (PeerSettings peerSettings in clusterSettings.CurrentValue.Peers)
             {
                 // skip the local site
-                if (string.Equals(peerSettings.Name, jupiterSettings.CurrentValue.CurrentSite, StringComparison.InvariantCultureIgnoreCase))
+                if (string.Equals(peerSettings.Name, jupiterSettings.CurrentValue.CurrentSite, StringComparison.OrdinalIgnoreCase))
+                {
                     continue;
+                }
 
-                _peers[peerSettings.Name] = new IPeerStatusService.PeerStatus(peerSettings)
+                _peers[peerSettings.Name] = new PeerStatus(peerSettings)
                 {
                     Latency = int.MaxValue
                 };
@@ -99,7 +105,9 @@ namespace Horde.Storage.Implementation
         public override async Task<bool> OnPoll(PeerStatusServiceState state, CancellationToken cancellationToken)
         {
             if (_alreadyPolling)
+            {
                 return false;
+            }
 
             _alreadyPolling = true;
 
@@ -112,10 +120,12 @@ namespace Horde.Storage.Implementation
             foreach (PeerSettings peerSettings in _clusterSettings.CurrentValue.Peers)
             {
                 // skip the local site
-                if (string.Equals(peerSettings.Name, _jupiterSettings.CurrentValue.CurrentSite, StringComparison.InvariantCultureIgnoreCase))
+                if (string.Equals(peerSettings.Name, _jupiterSettings.CurrentValue.CurrentSite, StringComparison.OrdinalIgnoreCase))
+                {
                     continue;
+                }
 
-                IPeerStatusService.PeerStatus status = _peers[peerSettings.Name];
+                PeerStatus status = _peers[peerSettings.Name];
 
                 int bestLatency = int.MaxValue;
                 bool reachable = false;
@@ -126,7 +136,9 @@ namespace Horde.Storage.Implementation
                     bestLatency = Math.Min(latency, bestLatency);
 
                     if (latency != int.MaxValue)
-                       reachable = true;
+                    {
+                        reachable = true;
+                    }
                 });
 
                 status.Reachable = reachable;
@@ -148,7 +160,9 @@ namespace Horde.Storage.Implementation
                 HttpResponseMessage result = await client.GetAsync(uri);
                 // ignore error responses as they may not have reached the actual instance
                 if (!result.IsSuccessStatusCode)
+                {
                     return int.MaxValue;
+                }
 
                 return (int)stopwatch.ElapsedMilliseconds;
             }
@@ -157,7 +171,6 @@ namespace Horde.Storage.Implementation
                 // error reaching the endpoint is just considered to max latency
                 return int.MaxValue;
             }
-
         }
     }
 }
diff --git a/Engine/Source/Programs/Horde/HordeStorage/Horde.Storage/Implementation/Refs/CachedRefStore.cs b/Engine/Source/Programs/Horde/HordeStorage/Horde.Storage/Implementation/Refs/CachedRefStore.cs
--- a/Engine/Source/Programs/Horde/HordeStorage/Horde.Storage/Implementation/Refs/CachedRefStore.cs
+++ b/Engine/Source/Programs/Horde/HordeStorage/Horde.Storage/Implementation/Refs/CachedRefStore.cs
@@ -15,7 +15,7 @@ using Serilog.Context;
 
 namespace Horde.Storage.Implementation
 {
-    public class CachedRefStore : IRefsStore
+    public class CachedRefStore : IRefsStore, IDisposable
     {
         private readonly IRefsStore _backingStore;
         private readonly IOptionsMonitor<MemoryCacheRefSettings> _memoryOptions;
@@ -26,7 +26,7 @@ namespace Horde.Storage.Implementation
         private readonly ILogger _logger = Log.ForContext<CachedRefStore>();
 
         // ReSharper disable once NotAccessedField.Local
-        private Timer _cacheInfoTimer;
+        private readonly Timer _cacheInfoTimer;
 
         public CachedRefStore(IRefsStore backingStore, IOptionsMonitor<MemoryCacheRefSettings> memoryOptions)
         {
@@ -39,10 +39,7 @@ namespace Horde.Storage.Implementation
             _cacheInfoTimer = new Timer(LogCacheInfo, null, TimeSpan.Zero, TimeSpan.FromSeconds(60));
         }
 
-        public IRefsStore BackingStore
-        {
-            get { return _backingStore; }
-        }
+        public IRefsStore BackingStore => _backingStore;
 
         private void LogCacheInfo(object? state)
         {
@@ -75,7 +72,9 @@ namespace Horde.Storage.Implementation
             
             // do not refs a miss as it could get added at any time
             if (backendRecord == null)
+            {
                 return null;
+            }
 
             CreateMemoryCacheEntry(cacheKey, backendRecord, ns, bucket);
 
@@ -159,6 +158,21 @@ namespace Horde.Storage.Implementation
 
             return _backingStore.DropNamespace(ns);
         }
+
+        protected virtual void Dispose(bool disposing)
+        {
+            if (disposing)
+            {
+                _memoryCache.Dispose();
+                _cacheInfoTimer.Dispose();
+            }
+        }
+
+        public void Dispose()
+        {
+            Dispose(true);
+            GC.SuppressFinalize(this);
+        }
     }
 
     class BoolChangeToken : IChangeToken
@@ -170,10 +184,6 @@ namespace Horde.Storage.Implementation
 
         public bool HasChanged { get; set; }
 
-        public bool ActiveChangeCallbacks
-        {
-            get { return false; }
-        }
+        public bool ActiveChangeCallbacks => false;
     }
-
 }
diff --git a/Engine/Source/Programs/Horde/HordeStorage/Horde.Storage/Implementation/Refs/CosmosRefsStore.cs b/Engine/Source/Programs/Horde/HordeStorage/Horde.Storage/Implementation/Refs/CosmosRefsStore.cs
--- a/Engine/Source/Programs/Horde/HordeStorage/Horde.Storage/Implementation/Refs/CosmosRefsStore.cs
+++ b/Engine/Source/Programs/Horde/HordeStorage/Horde.Storage/Implementation/Refs/CosmosRefsStore.cs
@@ -23,7 +23,6 @@ namespace Horde.Storage.Implementation
             _cosmosSettings = cosmosSettings;
         }
 
-
         public override async Task Add(RefRecord record)
         {
             string dbName = GetDatabaseName(record.Namespace);
@@ -31,7 +30,7 @@ namespace Horde.Storage.Implementation
             bool collectionExists = await CollectionExistsAsync(dbName, ModelName);
             if (!collectionExists)
             {
-                var createDbCommand = new BsonDocument
+                BsonDocument createDbCommand = new BsonDocument
                 {
                     {
                         "customAction","CreateCollection"
@@ -43,7 +42,7 @@ namespace Horde.Storage.Implementation
                         "offerThroughput",_cosmosSettings.CurrentValue.DefaultRU
                     }
                 };
-                var db = _client.GetDatabase(dbName);
+                IMongoDatabase? db = Client.GetDatabase(dbName);
                 BsonDocument? result;
                 try
                 {
@@ -69,7 +68,9 @@ namespace Horde.Storage.Implementation
                     //https://docs.microsoft.com/en-us/azure/cosmos-db/mongodb-custom-commands#default-output
                     int ok = result["ok"].AsInt32;
                     if (ok == 0) // failure
+                    {
                         throw new Exception($"Failed to create MongoDB {dbName} in Cosmos. Response: {result}");
+                    }
                 }
             }
 
@@ -78,11 +79,11 @@ namespace Horde.Storage.Implementation
 
         private async Task<bool> CollectionExistsAsync(string dbName, string collectionName)
         {
-            var db = _client.GetDatabase(dbName);
+            IMongoDatabase? db = Client.GetDatabase(dbName);
 
-            var filter = new BsonDocument("name", collectionName);
+            BsonDocument filter = new BsonDocument("name", collectionName);
             //filter by collection name
-            var collections = await db.ListCollectionsAsync(new ListCollectionsOptions { Filter = filter });
+            IAsyncCursor<BsonDocument>? collections = await db.ListCollectionsAsync(new ListCollectionsOptions { Filter = filter });
             //check for existence
             return await collections.AnyAsync();
         }
diff --git a/Engine/Source/Programs/Horde/HordeStorage/Horde.Storage/Implementation/Refs/DynamoDbRefsStore.cs b/Engine/Source/Programs/Horde/HordeStorage/Horde.Storage/Implementation/Refs/DynamoDbRefsStore.cs
--- a/Engine/Source/Programs/Horde/HordeStorage/Horde.Storage/Implementation/Refs/DynamoDbRefsStore.cs
+++ b/Engine/Source/Programs/Horde/HordeStorage/Horde.Storage/Implementation/Refs/DynamoDbRefsStore.cs
@@ -91,16 +91,24 @@ namespace Horde.Storage.Implementation
             // LoadAsync can return null records even if its not annotated for it
             // ReSharper disable once ConditionIsAlwaysTrueOrFalse
             if (baseRefRecord == null)
+            {
                 return null; // this key is not valid
+            }
 
             if (!baseRefRecord.IsValid())
+            {
                 throw new InvalidOperationException("Refs record was not fully populated, version mismatch?");
+            }
 
             if (!lastAccessRefRecord?.IsValid() ?? false)
+            {
                 throw new InvalidOperationException("LastAccess Refs record was not fully populated, version mismatch?");
+            }
 
             if (!metadataRefRecord?.IsValid() ?? false)
+            {
                 throw new InvalidOperationException("Metadata refs record was not fully populated, version mismatch?");
+            }
 
             return new DynamoRefRecord(new NamespaceId(baseRefRecord.Namespace!), new BucketId(baseRefRecord.Bucket!), new KeyId(baseRefRecord.Name!), 
                 baseRefRecord.Blobs!.Select(s => new BlobIdentifier(s)).ToArray(), 
@@ -176,17 +184,16 @@ namespace Horde.Storage.Implementation
             return countOfDeletedItems;
         }
 
-
         public async Task<long> Delete(NamespaceId ns, BucketId bucket, KeyId key)
         {
             await Initialize();
 
-            var deleteBaseTask = new DeleteItemRequest(MainTableName, new Dictionary<string, AttributeValue> {{"Id", new AttributeValue(DynamoBaseRefRecord.BuildCacheKey(ns, bucket, key)) }})
+            DeleteItemRequest deleteBaseTask = new(MainTableName, new Dictionary<string, AttributeValue> {{"Id", new AttributeValue(DynamoBaseRefRecord.BuildCacheKey(ns, bucket, key)) }})
             {
                 ConditionExpression = "attribute_exists(Id)"
             };
-            var deleteLastAccessTask = new DeleteItemRequest(MainTableName, new Dictionary<string, AttributeValue> {{"Id", new AttributeValue(DynamoLastAccessRefRecord.BuildCacheKey(ns, bucket, key)) }});
-            var deleteMetadataTask = new DeleteItemRequest(MainTableName, new Dictionary<string, AttributeValue> {{"Id", new AttributeValue(DynamoMetadataRefRecord.BuildCacheKey(ns, bucket, key)) }});
+            DeleteItemRequest deleteLastAccessTask = new(MainTableName, new Dictionary<string, AttributeValue> {{"Id", new AttributeValue(DynamoLastAccessRefRecord.BuildCacheKey(ns, bucket, key)) }});
+            DeleteItemRequest deleteMetadataTask = new(MainTableName, new Dictionary<string, AttributeValue> {{"Id", new AttributeValue(DynamoMetadataRefRecord.BuildCacheKey(ns, bucket, key)) }});
 
             try
             {
@@ -201,7 +208,6 @@ namespace Horde.Storage.Implementation
             return 1;
         }
 
-
         public async Task UpdateLastAccessTime(RefRecord record, DateTime lastAccessTime)
         {
             await Initialize();
@@ -215,7 +221,10 @@ namespace Horde.Storage.Implementation
             Document? item = await table.GetItemAsync(DynamoLastAccessRefRecord.BuildCacheKey(record.Namespace, record.Bucket, record.RefName));
             // if the item has been deleted after it was scheduled to be updated it will be null
             if (item == null)
+            {
                 return;
+            }
+
             item["LastAccessTime"] = lastAccessTime;
             await table.UpdateItemAsync(item);
         }
@@ -224,7 +233,7 @@ namespace Horde.Storage.Implementation
         {
             await Initialize();
 
-            await foreach (var ns in _namespaceStore.GetNamespaces(DynamoNamespaceStore.NamespaceUsage.Cache))
+            await foreach (NamespaceId ns in _namespaceStore.GetNamespaces(DynamoNamespaceStore.NamespaceUsage.Cache))
             {
                 yield return ns;
             }
@@ -235,7 +244,7 @@ namespace Horde.Storage.Implementation
             await Initialize();
 
             DateTime cutOffDate = DateTime.Now - oldRecordCutoff;
-            using var context = new DynamoDBContext(Client, new DynamoDBContextConfig { Conversion = DynamoDBEntryConversion.V2 });
+            using DynamoDBContext context = new DynamoDBContext(Client, new DynamoDBContextConfig { Conversion = DynamoDBEntryConversion.V2 });
             AsyncSearch<DynamoLastAccessRefRecord> search = context.ScanAsync<DynamoLastAccessRefRecord>(new[]
             {
                 new ScanCondition("LastAccessTime", ScanOperator.Between, DateTime.MinValue, cutOffDate)
@@ -247,14 +256,23 @@ namespace Horde.Storage.Implementation
                 foreach (DynamoLastAccessRefRecord document in newSet)
                 {
                     if (document == null)
+                    {
                         continue;
+                    }
+
                     string parentKey = DynamoLastAccessRefRecord.ToBaseRefRecordsKey(document.Id!);
 
                     DynamoBaseRefRecord? baseRecord = await context.LoadAsync<DynamoBaseRefRecord>(parentKey);
                     if (baseRecord == null || baseRecord.Namespace == null || baseRecord.Bucket == null || baseRecord.Name == null)
+                    {
                         continue;
+                    }
+
                     if (baseRecord.Namespace != ns.ToString())
+                    {
                         continue;
+                    }
+
                     yield return new OldRecord(new NamespaceId(baseRecord.Namespace), new BucketId(baseRecord.Bucket), new KeyId(baseRecord.Name));
                 }
             } while (!search.IsDone);
@@ -262,7 +280,7 @@ namespace Horde.Storage.Implementation
 
         public async Task DropNamespace(NamespaceId ns)
         {
-            using var context = new DynamoDBContext(Client, new DynamoDBContextConfig { Conversion = DynamoDBEntryConversion.V2 });
+            using DynamoDBContext context = new DynamoDBContext(Client, new DynamoDBContextConfig { Conversion = DynamoDBEntryConversion.V2 });
             try
             {
                 AsyncSearch<DynamoBaseRefRecord> search = context.ScanAsync<DynamoBaseRefRecord>(new[]
@@ -328,7 +346,6 @@ namespace Horde.Storage.Implementation
             Id = BuildCacheKey(record.Namespace, record.Bucket, record.RefName);
         }
 
-
         [DynamoDBHashKey] 
         public string? Id { get; set; }
 
@@ -339,6 +356,7 @@ namespace Horde.Storage.Implementation
         public string? Name { get; set; }
 
         [DynamoDBProperty]
+        [System.Diagnostics.CodeAnalysis.SuppressMessage("Usage", "CA2227:Collection properties should be read only", Justification = "Required by serialization")]
         public List<string>? Blobs { get; set; }
 
         public string? ContentHash { get; set; }
@@ -351,17 +369,34 @@ namespace Horde.Storage.Implementation
         public bool IsValid()
         {
             if (Id == null)
+            {
                 return false;
+            }
+
             if (Namespace == null)
+            {
                 return false;
+            }
+
             if (Bucket == null)
+            {
                 return false;
+            }
+
             if (Name == null)
+            {
                 return false;
+            }
+
             if (Blobs == null)
+            {
                 return false;
+            }
+
             if (ContentHash == null)
+            {
                 return false;
+            }
 
             return true;
         }
@@ -407,18 +442,25 @@ namespace Horde.Storage.Implementation
 
         public static string ToBaseRefRecordsKey(string cacheKey)
         {
-            return cacheKey.Replace("#last-access", "");
+            return cacheKey.Replace("#last-access", "", StringComparison.OrdinalIgnoreCase);
         }
 
-
         public bool IsValid()
         {
             if (Id == null)
+            {
                 return false;
+            }
+
             if (LastAccessTime == null)
+            {
                 return false;
+            }
+
             if (LastAccessIndex == null)
+            {
                 return false;
+            }
 
             return true;
         }
@@ -443,6 +485,7 @@ namespace Horde.Storage.Implementation
         public string? Id { get; set; }
 
         [DynamoDBProperty(typeof(DynamoJsonConverter))]
+        [System.Diagnostics.CodeAnalysis.SuppressMessage("Usage", "CA2227:Collection properties should be read only", Justification = "Used by serialization")]
         public Dictionary<string, object>? Metadata { get; set; }
 
         public static string BuildCacheKey(NamespaceId ns, BucketId bucket, KeyId key)
@@ -453,7 +496,9 @@ namespace Horde.Storage.Implementation
         public bool IsValid()
         {
             if (Id == null)
+            {
                 return false;
+            }
 
             return true;
         }
diff --git a/Engine/Source/Programs/Horde/HordeStorage/Horde.Storage/Implementation/Refs/DynamoNamespaceStore.cs b/Engine/Source/Programs/Horde/HordeStorage/Horde.Storage/Implementation/Refs/DynamoNamespaceStore.cs
--- a/Engine/Source/Programs/Horde/HordeStorage/Horde.Storage/Implementation/Refs/DynamoNamespaceStore.cs
+++ b/Engine/Source/Programs/Horde/HordeStorage/Horde.Storage/Implementation/Refs/DynamoNamespaceStore.cs
@@ -1,6 +1,5 @@
 // Copyright Epic Games, Inc. All Rights Reserved.
 
-using System;
 using System.Collections.Generic;
 using System.Threading.Tasks;
 using Amazon.DynamoDBv2;
@@ -8,7 +7,6 @@ using Amazon.DynamoDBv2.DataModel;
 using Amazon.DynamoDBv2.DocumentModel;
 using Amazon.DynamoDBv2.Model;
 using EpicGames.Horde.Storage;
-using Jupiter.Implementation;
 using Microsoft.Extensions.Options;
 
 namespace Horde.Storage.Implementation
@@ -45,7 +43,6 @@ namespace Horde.Storage.Implementation
             }
         }
 
-
         public async IAsyncEnumerable<NamespaceId> GetNamespaces(NamespaceUsage usage)
         {
             await Initialize();
@@ -61,7 +58,10 @@ namespace Horde.Storage.Implementation
                 foreach (NamespaceMappingRecord document in newSet)
                 {
                     if (document.Namespace == null)
+                    {
                         continue;
+                    }
+
                     yield return new NamespaceId(document.Namespace);
                 }
             } while (!search.IsDone);
@@ -103,8 +103,8 @@ namespace Horde.Storage.Implementation
             Tree
         }
 
-
         [DynamoDBTable(NamespaceTableName)]
+        [System.Diagnostics.CodeAnalysis.SuppressMessage("Design", "CA1034:Nested types should not be visible", Justification = "Used by serialization")]
         public class NamespaceMappingRecord
         {
             [DynamoDBHashKey] public string Namespace { get; set; } = null!;
diff --git a/Engine/Source/Programs/Horde/HordeStorage/Horde.Storage/Implementation/Refs/DynamoStore.cs b/Engine/Source/Programs/Horde/HordeStorage/Horde.Storage/Implementation/Refs/DynamoStore.cs
--- a/Engine/Source/Programs/Horde/HordeStorage/Horde.Storage/Implementation/Refs/DynamoStore.cs
+++ b/Engine/Source/Programs/Horde/HordeStorage/Horde.Storage/Implementation/Refs/DynamoStore.cs
@@ -15,7 +15,7 @@ using Newtonsoft.Json;
 
 namespace Horde.Storage.Implementation
 {
-    public abstract class DynamoStore
+    public abstract class DynamoStore: IDisposable
     {
         private readonly IOptionsMonitor<DynamoDbSettings> _settings;
         private readonly SemaphoreSlim _initializeSemaphore = new SemaphoreSlim(1, 1);
@@ -37,12 +37,17 @@ namespace Horde.Storage.Implementation
         public async Task Initialize()
         {
             if (_initialized)
+            {
                 return;
+            }
 
             await _initializeSemaphore.WaitAsync();
 
             if (_initialized)
+            {
                 return;
+            }
+
             try
             {
                 await CreateTables();
@@ -61,9 +66,11 @@ namespace Horde.Storage.Implementation
 
         protected async Task CreateTable(string tableName, KeySchemaElement[] keySchemaElements, AttributeDefinition[] attributeDefinitions, ProvisionedThroughput provisionedThroughput, IEnumerable<GlobalSecondaryIndex>? secondaryIndices = null)
         {
-            var tables = await Client.ListTablesAsync();
+            ListTablesResponse? tables = await Client.ListTablesAsync();
             if (tables.TableNames.Contains(tableName))
+            {
                 return;
+            }
 
             bool shouldCreateTable = _settings.CurrentValue.CreateTablesOnDemand;
 
@@ -101,6 +108,19 @@ namespace Horde.Storage.Implementation
             }
         }
 
+        protected virtual void Dispose(bool disposing)
+        {
+            if (disposing)
+            {
+                _initializeSemaphore.Dispose();
+            }
+        }
+
+        public void Dispose()
+        {
+            Dispose(true);
+            GC.SuppressFinalize(this);
+        }
     }
 
     public class DynamoJsonConverter : IPropertyConverter
@@ -108,15 +128,18 @@ namespace Horde.Storage.Implementation
         public object? FromEntry(DynamoDBEntry entry)
         {
             Primitive? primitive = entry as Primitive;
-            if (primitive == null || !(primitive.Value is String) || string.IsNullOrEmpty((string)primitive.Value))
-                throw new ArgumentOutOfRangeException();
-            Dictionary<string, object>? ret = JsonConvert.DeserializeObject<Dictionary<string, object>>((primitive.Value as string)!);
+            if (primitive == null || primitive.Value is not string value || string.IsNullOrEmpty(value))
+            {
+                throw new ArgumentOutOfRangeException(nameof(entry));
+            }
+
+            Dictionary<string, object>? ret = JsonConvert.DeserializeObject<Dictionary<string, object>>(value);
             return ret;
         }
 
         public DynamoDBEntry ToEntry(object value)
         {
-            var jsonString = JsonConvert.SerializeObject(value);
+            string jsonString = JsonConvert.SerializeObject(value);
             DynamoDBEntry ret = new Primitive(jsonString);
             return ret;
         }
diff --git a/Engine/Source/Programs/Horde/HordeStorage/Horde.Storage/Implementation/Refs/MemoryRefsStore.cs b/Engine/Source/Programs/Horde/HordeStorage/Horde.Storage/Implementation/Refs/MemoryRefsStore.cs
--- a/Engine/Source/Programs/Horde/HordeStorage/Horde.Storage/Implementation/Refs/MemoryRefsStore.cs
+++ b/Engine/Source/Programs/Horde/HordeStorage/Horde.Storage/Implementation/Refs/MemoryRefsStore.cs
@@ -5,8 +5,6 @@ using System.Collections.Generic;
 using System.Threading.Tasks;
 using Dasync.Collections;
 using Jupiter.Implementation;
-using Serilog;
-
 using RefContainer = System.Collections.Concurrent.ConcurrentDictionary<Jupiter.Implementation.KeyId, Horde.Storage.Implementation.RefRecord>;
 using BucketContainer = System.Collections.Concurrent.ConcurrentDictionary<EpicGames.Horde.Storage.BucketId, System.Collections.Concurrent.ConcurrentDictionary<Jupiter.Implementation.KeyId, Horde.Storage.Implementation.RefRecord>>;
 using NamespaceContainer = System.Collections.Concurrent.ConcurrentDictionary<EpicGames.Horde.Storage.NamespaceId, System.Collections.Concurrent.ConcurrentDictionary<EpicGames.Horde.Storage.BucketId, System.Collections.Concurrent.ConcurrentDictionary<Jupiter.Implementation.KeyId, Horde.Storage.Implementation.RefRecord>>>;
@@ -16,7 +14,6 @@ namespace Horde.Storage.Implementation
 {
     internal class MemoryRefsStore : IRefsStore
     {
-        private readonly ILogger _logger = Log.ForContext<MemoryRefsStore>();
         private readonly NamespaceContainer _refs = new NamespaceContainer();
 
         public Task<RefRecord?> Get(NamespaceId ns, BucketId bucket, KeyId key, IRefsStore.ExtraFieldsFlag fields)
@@ -56,7 +53,7 @@ namespace Horde.Storage.Implementation
             BucketContainer bucketContainer = _refs.GetOrAdd(ns, new BucketContainer());
             RefContainer refContainer = bucketContainer.GetOrAdd(bucket, new RefContainer());
 
-            if (refContainer.TryRemove(key, out RefRecord? refRecord))
+            if (refContainer.TryRemove(key, out RefRecord? _))
             {
                 return Task.FromResult(1L);
             }
@@ -95,7 +92,9 @@ namespace Horde.Storage.Implementation
                 {
                     RefRecord refRecord = refContainer[keyId];
                     if (refRecord.LastAccessTime < cutoff)
+                    {
                         yield return new OldRecord(ns, bucket, keyId);
+                    }
                 }
             }
         }
diff --git a/Engine/Source/Programs/Horde/HordeStorage/Horde.Storage/Implementation/Refs/MongoRefsStore.cs b/Engine/Source/Programs/Horde/HordeStorage/Horde.Storage/Implementation/Refs/MongoRefsStore.cs
--- a/Engine/Source/Programs/Horde/HordeStorage/Horde.Storage/Implementation/Refs/MongoRefsStore.cs
+++ b/Engine/Source/Programs/Horde/HordeStorage/Horde.Storage/Implementation/Refs/MongoRefsStore.cs
@@ -89,7 +89,7 @@ namespace Horde.Storage.Implementation
 
     public class MongoRefsStore : IRefsStore
     {
-        protected readonly MongoClient _client;
+        protected MongoClient Client { get; }
         protected const string ModelName = "Cache";
         private const string MetadataModel = "Metadata";
         private readonly ILogger _logger = Log.ForContext<MongoRefsStore>();
@@ -102,10 +102,10 @@ namespace Horde.Storage.Implementation
             );
             if (settings.RequireTls12)
             {
-                mongoClientSettings.SslSettings = new SslSettings {EnabledSslProtocols = SslProtocols.Tls12};
+                mongoClientSettings.SslSettings = new SslSettings {EnabledSslProtocols = SslProtocols.None};
             }
 
-            _client = new MongoClient(mongoClientSettings);
+            Client = new MongoClient(mongoClientSettings);
 
             SetupClassMaps();
         }
@@ -131,12 +131,12 @@ namespace Horde.Storage.Implementation
         protected IMongoCollection<T> GetCollection<T>(NamespaceId ns, string collectionName)
         {
             string dbName = GetDatabaseName(ns);
-            return _client.GetDatabase(dbName).GetCollection<T>(collectionName);
+            return Client.GetDatabase(dbName).GetCollection<T>(collectionName);
         }
 
         protected string GetDatabaseName(NamespaceId ns)
         {
-            string cleanedNs = ns.ToString().Replace(".", "_");
+            string cleanedNs = ns.ToString().Replace(".", "_", StringComparison.OrdinalIgnoreCase);
             string dbName = $"Europa_{cleanedNs}";
             return dbName;
         }
@@ -229,16 +229,18 @@ namespace Horde.Storage.Implementation
 
         public async IAsyncEnumerable<NamespaceId> GetNamespaces()
         {
-            using IAsyncCursor<BsonDocument> cursor = await _client.ListDatabasesAsync();
+            using IAsyncCursor<BsonDocument> cursor = await Client.ListDatabasesAsync();
             while (await cursor.MoveNextAsync())
             {
                 foreach (BsonDocument doc in cursor.Current)
                 {
                     string dbName = doc["name"].AsString;
-                    MongoMetadata? result = await _client.GetDatabase(dbName).GetCollection<MongoMetadata>("Metadata")
+                    MongoMetadata? result = await Client.GetDatabase(dbName).GetCollection<MongoMetadata>("Metadata")
                         .Find(FilterDefinition<MongoMetadata>.Empty).FirstOrDefaultAsync();
                     if (result == null)
+                    {
                         continue;
+                    }
 
                     NamespaceId ns;
                     try
@@ -269,7 +271,7 @@ namespace Horde.Storage.Implementation
         public async Task DropNamespace(NamespaceId ns)
         {
             string dbName = GetDatabaseName(ns);
-            await _client.GetDatabase(dbName).DropCollectionAsync(ModelName);
+            await Client.GetDatabase(dbName).DropCollectionAsync(ModelName);
         }
 
         private async Task<V> ExecuteWithRetries<V>(Func<Task<V>> function)
@@ -291,8 +293,7 @@ namespace Horde.Storage.Implementation
                     // https://feedback.azure.com/forums/263030-azure-cosmos-db/suggestions/32956231-better-handling-of-request-limit-too-large
                     // so lets just guess that we should retry after a while
                     return TimeSpan.FromMilliseconds(250);
-
-                };
+                }
 
                 TimeSpan sleepTime;
                 try
diff --git a/Engine/Source/Programs/Horde/HordeStorage/Horde.Storage/Implementation/RelayStore.cs b/Engine/Source/Programs/Horde/HordeStorage/Horde.Storage/Implementation/RelayStore.cs
--- a/Engine/Source/Programs/Horde/HordeStorage/Horde.Storage/Implementation/RelayStore.cs
+++ b/Engine/Source/Programs/Horde/HordeStorage/Horde.Storage/Implementation/RelayStore.cs
@@ -12,10 +12,7 @@ namespace Horde.Storage.Implementation
         private readonly IServiceCredentials _serviceCredentials;
         private readonly HttpClient _httpClient;
 
-        protected HttpClient HttpClient
-        {
-            get { return _httpClient; }
-        }
+        protected HttpClient HttpClient => _httpClient;
 
         protected RelayStore(IOptionsMonitor<UpstreamRelaySettings> settings, IHttpClientFactory httpClientFactory, IServiceCredentials serviceCredentials)
         {
@@ -26,15 +23,16 @@ namespace Horde.Storage.Implementation
             _httpClient.BaseAddress = new Uri(_settings.CurrentValue.ConnectionString);
         }
 
-        protected HttpRequestMessage BuildHttpRequest(HttpMethod method, string uri)
+        protected HttpRequestMessage BuildHttpRequest(HttpMethod method, Uri uri)
         {
             string? token = _serviceCredentials.GetToken();
             HttpRequestMessage request = new HttpRequestMessage(method, uri);
             if (!string.IsNullOrEmpty(token))
+            {
                 request.Headers.Add("Authorization", "Bearer " + token);
+            }
 
             return request;
         }
-
     }
 }
diff --git a/Engine/Source/Programs/Horde/HordeStorage/Horde.Storage/Implementation/Replication/IReplicator.cs b/Engine/Source/Programs/Horde/HordeStorage/Horde.Storage/Implementation/Replication/IReplicator.cs
--- a/Engine/Source/Programs/Horde/HordeStorage/Horde.Storage/Implementation/Replication/IReplicator.cs
+++ b/Engine/Source/Programs/Horde/HordeStorage/Horde.Storage/Implementation/Replication/IReplicator.cs
@@ -3,40 +3,38 @@
 using System;
 using System.Threading.Tasks;
 using EpicGames.Horde.Storage;
-using Jupiter.Implementation;
 
 namespace Horde.Storage.Implementation
 {
-    public interface IReplicator : IDisposable
+    public class ReplicatorState
     {
-        public class ReplicatorState
-        {
-            public long? ReplicatorOffset { get; set; }
-            public Guid? ReplicatingGeneration { get; set; }
-            public string? LastBucket { get; set; }
-            public Guid? LastEvent { get; set; }
-        }
-
-        public class ReplicatorInfo
+        public long? ReplicatorOffset { get; set; }
+        public Guid? ReplicatingGeneration { get; set; }
+        public string? LastBucket { get; set; }
+        public Guid? LastEvent { get; set; }
+    }
+    public class ReplicatorInfo
+    {
+        public ReplicatorInfo(string replicatorName, NamespaceId namespaceToReplicate, ReplicatorState state)
         {
-            public ReplicatorInfo(string replicatorName, NamespaceId namespaceToReplicate, ReplicatorState state)
-            {
-                ReplicatorName = replicatorName;
-                NamespaceToReplicate = namespaceToReplicate;
-                State = state;
+            ReplicatorName = replicatorName;
+            NamespaceToReplicate = namespaceToReplicate;
+            State = state;
 
-                LastRun = DateTime.MinValue;
-                CountOfRunningReplications = 0;
-            }
+            LastRun = DateTime.MinValue;
+            CountOfRunningReplications = 0;
+        }
 
-            public string ReplicatorName { get; }
-            public NamespaceId NamespaceToReplicate { get; }
-            public ReplicatorState State { get; }
+        public string ReplicatorName { get; }
+        public NamespaceId NamespaceToReplicate { get; }
+        public ReplicatorState State { get; }
 
-            public DateTime LastRun { get; set; }
-            public int CountOfRunningReplications { get; set; }
-        }
+        public DateTime LastRun { get; set; }
+        public int CountOfRunningReplications { get; set; }
+    }
 
+    public interface IReplicator : IDisposable
+    {
         /// <summary>
         /// Attempt to run a new replication, if a replication is already in flight for this replicator this will early exist.
         /// </summary>
diff --git a/Engine/Source/Programs/Horde/HordeStorage/Horde.Storage/Implementation/Replication/RefsReplicator.cs b/Engine/Source/Programs/Horde/HordeStorage/Horde.Storage/Implementation/Replication/RefsReplicator.cs
--- a/Engine/Source/Programs/Horde/HordeStorage/Horde.Storage/Implementation/Replication/RefsReplicator.cs
+++ b/Engine/Source/Programs/Horde/HordeStorage/Horde.Storage/Implementation/Replication/RefsReplicator.cs
@@ -55,7 +55,7 @@ namespace Horde.Storage.Implementation
             _httpClient.BaseAddress = new Uri(replicatorSettings.ConnectionString);
 
             string stateFileName = $"{_name}.json";
-            DirectoryInfo stateRoot = new DirectoryInfo(replicationSettings.CurrentValue.GetStateRoot());
+            DirectoryInfo stateRoot = new DirectoryInfo(replicationSettings.CurrentValue.StateRoot);
             _stateFile = new FileInfo(Path.Combine(stateRoot.FullName, stateFileName));
 
             ReplicatorState? replicatorState = _replicationLog.GetReplicatorState(_replicatorSettings.NamespaceToReplicate, _name).Result;
@@ -79,18 +79,31 @@ namespace Horde.Storage.Implementation
                 };
             }
 
-            Info = new IReplicator.ReplicatorInfo(replicatorSettings.ReplicatorName, replicatorSettings.NamespaceToReplicate, _refsState);
+            Info = new ReplicatorInfo(replicatorSettings.ReplicatorName, replicatorSettings.NamespaceToReplicate, _refsState);
         }
 
         public void Dispose()
+        {
+            Dispose(true);
+            GC.SuppressFinalize(this);
+        }
+
+        protected virtual void Dispose(bool disposing)
         {
             if (_disposed)
+            {
                 return;
+            }
 
-            SaveState(_stateFile, _refsState);
+            if (disposing)
+            {
+                SaveState(_stateFile, _refsState);
+
+                _replicationFinishedEvent.Wait();
+                _replicationTokenSource.Dispose();
 
-            _replicationFinishedEvent.Wait();
-            _replicationTokenSource.Dispose();
+                _httpClient.Dispose();
+            }
 
             _disposed = true;
         }
@@ -102,7 +115,9 @@ namespace Horde.Storage.Implementation
             JsonSerializer serializer = new JsonSerializer();
             RefsState? state =  serializer.Deserialize<RefsState>(reader);
             if (state == null)
+            {
                 throw new Exception("Failed to read state");
+            }
 
             return state;
         }
@@ -150,7 +165,6 @@ namespace Horde.Storage.Implementation
                 };
             }
 
-
             LogReplicationHeartbeat(0);
 
             bool hasRun;
@@ -261,7 +275,6 @@ namespace Horde.Storage.Implementation
                         SaveState(_stateFile, _refsState);
                         retry = true;
                     }
-                    
                 } while (retry);
                 hasRun = countOfReplicationsDone != 0;
             }
@@ -279,14 +292,16 @@ namespace Horde.Storage.Implementation
             // determine latest snapshot if no specific blob was specified
             if (snapshotBlob == null)
             {
-                HttpRequestMessage snapshotRequest = BuildHttpRequest(HttpMethod.Get, $"api/v1/replication-log/snapshots/{ns}");
+                using HttpRequestMessage snapshotRequest = BuildHttpRequest(HttpMethod.Get, new Uri($"api/v1/replication-log/snapshots/{ns}", UriKind.Relative));
                 HttpResponseMessage snapshotResponse = await _httpClient.SendAsync(snapshotRequest, cancellationToken);
                 snapshotResponse.EnsureSuccessStatusCode();
 
                 string s = await snapshotResponse.Content.ReadAsStringAsync(cancellationToken);
                 ReplicationLogSnapshots? snapshots = JsonConvert.DeserializeObject<ReplicationLogSnapshots>(s);
                 if (snapshots == null)
+                {
                     throw new NotImplementedException();
+                }
 
                 SnapshotInfo? snapshotInfo = snapshots.Snapshots?.FirstOrDefault();
                 if (snapshotInfo == null)
@@ -302,7 +317,7 @@ namespace Horde.Storage.Implementation
             // fetch the snapshot from the remote blob store
             ReplicationLogSnapshot snapshot;
             {
-                HttpRequestMessage request = BuildHttpRequest(HttpMethod.Get, $"api/v1/blobs/{blobNamespace}/{snapshotBlob}");
+                using HttpRequestMessage request = BuildHttpRequest(HttpMethod.Get, new Uri($"api/v1/blobs/{blobNamespace}/{snapshotBlob}", UriKind.Relative));
                 HttpResponseMessage response = await _httpClient.SendAsync(request, HttpCompletionOption.ResponseHeadersRead, cancellationToken);
                 response.EnsureSuccessStatusCode();
                 await using Stream blobStream = await response.Content.ReadAsStreamAsync(cancellationToken);
@@ -310,7 +325,9 @@ namespace Horde.Storage.Implementation
             }
 
             if (!snapshot.LastEvent.HasValue)
+            {
                 throw new Exception("No last event found");
+            }
 
             Guid snapshotEvent = snapshot.LastEvent.Value;
             string snapshotBucket = snapshot.LastBucket!;
@@ -329,7 +346,9 @@ namespace Horde.Storage.Implementation
                     LogReplicationHeartbeat(countOfObjectsCurrentlyReplicating);
 
                     if (countOfObjectsReplicated % 100 == 0)
+                    {
                         _logger.Information("{Name} Snapshot replication still running, replicated {CountOfObjects} of {TotalCountOfObjects}. Replicating {Namespace}", _name, countOfObjectsReplicated, snapshot.LiveObjectsCount, ns);
+                    }
 
                     Info.CountOfRunningReplications = countOfObjectsCurrentlyReplicating;
                     Info.LastRun = DateTime.Now;
@@ -346,7 +365,6 @@ namespace Horde.Storage.Implementation
                     Interlocked.Increment(ref countOfObjectsReplicated);
                     Interlocked.Decrement(ref countOfObjectsCurrentlyReplicating);
                 }
-
             }, maxParallelism, cancellationToken);
 
             snapshot.Dispose();
@@ -355,12 +373,15 @@ namespace Horde.Storage.Implementation
             return (snapshotBucket, snapshotEvent, countOfObjectsReplicated);
         }
 
-        private HttpRequestMessage BuildHttpRequest(HttpMethod httpMethod, string uri)
+        private HttpRequestMessage BuildHttpRequest(HttpMethod httpMethod, Uri uri)
         {
             string? token = _serviceCredentials.GetToken();
             HttpRequestMessage request = new HttpRequestMessage(httpMethod, uri);
             if (!string.IsNullOrEmpty(token))
+            {
                 request.Headers.Add("Authorization", "Bearer " + token);
+            }
+
             return request;
         }
 
@@ -376,11 +397,13 @@ namespace Horde.Storage.Implementation
 
             _logger.Information("{Name} Starting incremental replication maxParallelism: {MaxParallelism} Last event: {LastEvent} Last Bucket {LastBucket}", _name, maxParallelism, lastEvent, lastBucket);
 
-            CancellationTokenSource cancellationTokenSource = new CancellationTokenSource();
-            CancellationTokenSource linkedTokenSource = CancellationTokenSource.CreateLinkedTokenSource(cancellationTokenSource.Token, replicationToken);
+            using CancellationTokenSource cancellationTokenSource = new CancellationTokenSource();
+            using CancellationTokenSource linkedTokenSource = CancellationTokenSource.CreateLinkedTokenSource(cancellationTokenSource.Token, replicationToken);
 
             if (replicationToken.IsCancellationRequested)
+            {
                 return countOfReplicationsDone;
+            }
 
             await GetRefEvents(ns, lastBucket, lastEvent, replicationToken).ParallelForEachAsync(async (ReplicationLogEvent @event) =>
             {
@@ -413,7 +436,9 @@ namespace Horde.Storage.Implementation
                     if (@event.Op != ReplicationLogEvent.OpType.Deleted)
                     {
                         if (@event.Blob == null)
+                        {
                             throw new Exception($"Event: [REDACTED_EMAIL]} [REDACTED_EMAIL]} in namespace [REDACTED_EMAIL]} was missing a blob, unable to replicate it");
+                        }
 
                         await ReplicateOp(@event.Namespace, @event.Blob, replicationToken);
                     }
@@ -469,7 +494,7 @@ namespace Horde.Storage.Implementation
             //if (await _blobService.Exists(ns, blob))
             //    return currentOffset;
 
-            HttpRequestMessage referencesRequest = BuildHttpRequest(HttpMethod.Get, $"api/v1/objects/{ns}/{blob}/references");
+            using HttpRequestMessage referencesRequest = BuildHttpRequest(HttpMethod.Get, new Uri($"api/v1/objects/{ns}/{blob}/references", UriKind.Relative));
             HttpResponseMessage referencesResponse = await _httpClient.SendAsync(referencesRequest, cancellationToken);
             string body = await referencesResponse.Content.ReadAsStringAsync(cancellationToken);
 
@@ -482,7 +507,9 @@ namespace Horde.Storage.Implementation
 
             ResolvedReferencesResult? refs = JsonConvert.DeserializeObject<ResolvedReferencesResult>(body);
             if (refs == null)
+            {
                 throw new Exception($"Unable to resolve references for object {blob} in namespace {ns}");
+            }
 
             BlobIdentifier[] potentialBlobs = new BlobIdentifier[refs.References.Length + 1];
             Array.Copy(refs.References, potentialBlobs, refs.References.Length);
@@ -495,7 +522,7 @@ namespace Horde.Storage.Implementation
                 BlobIdentifier blobToReplicate = missingBlobs[i];
                 blobReplicationTasks[i] = Task.Run(async () =>
                 {
-                    HttpRequestMessage blobRequest = BuildHttpRequest(HttpMethod.Get, $"api/v1/blobs/{ns}/{blobToReplicate}");
+                    HttpRequestMessage blobRequest = BuildHttpRequest(HttpMethod.Get, new Uri($"api/v1/blobs/{ns}/{blobToReplicate}", UriKind.Relative));
                     HttpResponseMessage blobResponse = await _httpClient.SendAsync(blobRequest, HttpCompletionOption.ResponseHeadersRead, cancellationToken);
 
                     if (blobResponse.StatusCode == HttpStatusCode.NotFound)
@@ -516,7 +543,7 @@ namespace Horde.Storage.Implementation
                     using IBufferedPayload payload = contentLength is null or > int.MaxValue ? await FilesystemBufferedPayload.Create(s) : await MemoryBufferedPayload.Create(s);
 
                     await _blobService.PutObject(ns, payload, blobToReplicate);
-                });
+                }, cancellationToken);
             }
 
             await Task.WhenAll(blobReplicationTasks);
@@ -529,35 +556,46 @@ namespace Horde.Storage.Implementation
             do
             {
                 if (cancellationToken.IsCancellationRequested)
+                {
                     yield break;
+                }
 
                 if (hasRunOnce && (lastBucket == null || lastEvent == null))
                 {
                     throw new Exception($"Failed to find state to resume from after first page of ref events, lastBucket: {lastBucket} lastEvent: {lastEvent}");
                 }
-                StringBuilder url = new StringBuilder($"api/v1/replication-log/incremental/{ns}");
+                StringBuilder url = new StringBuilder($"/api/v1/replication-log/incremental/{ns}");
                 // its okay for last bucket and last event to be null incase we have never run before, but after the first iteration we need them to keep track of where we were
                 if (lastBucket != null)
+                {
                     url.Append($"?lastBucket={lastBucket}");
+                }
+
                 if (lastEvent != null)
+                {
                     url.Append($"&lastEvent={lastEvent}");
+                }
 
                 hasRunOnce = true;
-                HttpRequestMessage request = BuildHttpRequest(HttpMethod.Get, url.ToString());
+                using HttpRequestMessage request = BuildHttpRequest(HttpMethod.Get, new Uri(url.ToString(), UriKind.Relative));
                 HttpResponseMessage response = await _httpClient.SendAsync(request, cancellationToken);
                 string body = await response.Content.ReadAsStringAsync(cancellationToken);
                 if (response.StatusCode == HttpStatusCode.BadRequest)
                 {
                     ProblemDetails? problemDetails = JsonConvert.DeserializeObject<ProblemDetails>(body);
                     if (problemDetails == null)
+                    {
                         throw new Exception($"Unknown bad request body when reading incremental replication log. Body: {body}");
+                    }
 
                     if (problemDetails.Type == ProblemTypes.UseSnapshot)
                     {
                         ProblemDetailsWithSnapshots? problemDetailsWithSnapshots = JsonConvert.DeserializeObject<ProblemDetailsWithSnapshots>(body);
 
                         if (problemDetailsWithSnapshots == null)
+                        {
                             throw new Exception($"Unable to cast the problem details to a snapshot version. Body: {body}");
+                        }
 
                         BlobIdentifier snapshotBlob = problemDetailsWithSnapshots.SnapshotId;
                         throw new UseSnapshotException(snapshotBlob);
@@ -567,7 +605,9 @@ namespace Horde.Storage.Implementation
                 response.EnsureSuccessStatusCode();
                 ReplicationLogEvents? e = JsonConvert.DeserializeObject<ReplicationLogEvents>(body);
                 if (e == null)
+                {
                     throw new Exception($"Unknown error when deserializing replication log events {ns} {lastBucket} {lastEvent}");
+                }
 
                 logEvents = e;
                 foreach (ReplicationLogEvent logEvent in logEvents.Events)
@@ -577,7 +617,6 @@ namespace Horde.Storage.Implementation
                     lastBucket = logEvent.TimeBucket;
                     lastEvent = logEvent.EventId;
                 }
-
             } while (logEvents.Events.Any());
         }
 
@@ -599,18 +638,17 @@ namespace Horde.Storage.Implementation
         public async Task StopReplicating()
         {
             if (_disposed)
+            {
                 return;
+            }
 
             _replicationTokenSource.Cancel(true);
             await _replicationFinishedEvent.WaitAsync();
         }
 
-        public IReplicator.ReplicatorState State
-        {
-            get { return _refsState; }
-        }
+        public ReplicatorState State => _refsState;
 
-        public IReplicator.ReplicatorInfo Info { get; private set; }
+        public ReplicatorInfo Info { get; private set; }
 
         public Task DeleteState()
         {
@@ -628,7 +666,7 @@ namespace Horde.Storage.Implementation
         }
     }
 
-    internal class UseSnapshotException : Exception
+    public class UseSnapshotException : Exception
     {
         public BlobIdentifier SnapshotBlob { get; }
 
@@ -636,17 +674,16 @@ namespace Horde.Storage.Implementation
         {
             SnapshotBlob = snapshotBlob;
         }
-
     }
 
-    internal class NoSnapshotAvailableException : Exception
+    public class NoSnapshotAvailableException : Exception
     {
         public NoSnapshotAvailableException(string message) : base(message)
         {
         }
     }
 
-    public class RefsState : IReplicator.ReplicatorState
+    public class RefsState : ReplicatorState
     {
         public RefsState()
         {
diff --git a/Engine/Source/Programs/Horde/HordeStorage/Horde.Storage/Implementation/Replication/ReplicationService.cs b/Engine/Source/Programs/Horde/HordeStorage/Horde.Storage/Implementation/Replication/ReplicationService.cs
--- a/Engine/Source/Programs/Horde/HordeStorage/Horde.Storage/Implementation/Replication/ReplicationService.cs
+++ b/Engine/Source/Programs/Horde/HordeStorage/Horde.Storage/Implementation/Replication/ReplicationService.cs
@@ -8,43 +8,41 @@ using System.Threading;
 using System.Threading.Tasks;
 using EpicGames.Horde.Storage;
 using Jupiter;
-using Jupiter.Implementation;
 using Microsoft.Extensions.DependencyInjection;
 using Microsoft.Extensions.Options;
 using Serilog;
 
 namespace Horde.Storage.Implementation
 {
+    public class ReplicationState
+    {
+        public List<IReplicator> Replicators { get; } = new List<IReplicator>();
+    }
+
     // ReSharper disable once ClassNeverInstantiated.Global
-    public class ReplicationService : PollingService<ReplicationService.ReplicationState>
+    public class ReplicationService : PollingService<ReplicationState>
     {
         private readonly IOptionsMonitor<ReplicationSettings> _settings;
         private readonly ILeaderElection _leaderElection;
         private readonly ILogger _logger = Log.ForContext<ReplicationService>();
         private readonly Dictionary<string, Task<bool>> _currentReplications = new Dictionary<string, Task<bool>>();
 
-        public class ReplicationState
-        {
-            public List<IReplicator> Replicators { get; } = new List<IReplicator>();
-        }
-
         protected override bool ShouldStartPolling()
         {
             return _settings.CurrentValue.Enabled;
         }
 
-        public ReplicationService(IOptionsMonitor<ReplicationSettings> settings, IServiceProvider provider, ILeaderElection leaderElection) :
-            base(serviceName: nameof(ReplicationService), TimeSpan.FromSeconds(settings.CurrentValue.ReplicationPollFrequencySeconds), new ReplicationState(), startAtRandomTime: true)
+        public ReplicationService(IOptionsMonitor<ReplicationSettings> settings, IServiceProvider provider, ILeaderElection leaderElection) : base(serviceName: nameof(ReplicationService), TimeSpan.FromSeconds(settings.CurrentValue.ReplicationPollFrequencySeconds), new ReplicationState(), startAtRandomTime: true)
         {
             _settings = settings;
             _leaderElection = leaderElection;
 
             _leaderElection.OnLeaderChanged += OnLeaderChanged;
 
-            DirectoryInfo di = new DirectoryInfo(settings.CurrentValue.GetStateRoot());
+            DirectoryInfo di = new DirectoryInfo(settings.CurrentValue.StateRoot);
             Directory.CreateDirectory(di.FullName);
 
-            foreach (var replicator in settings.CurrentValue.Replicators)
+            foreach (ReplicatorSettings replicator in settings.CurrentValue.Replicators)
             {
                 try
                 {
@@ -55,13 +53,14 @@ namespace Horde.Storage.Implementation
                     _logger.Error(e, "Failed to create replicator {Name}", replicator.ReplicatorName);
                 }
             }
-
         }
 
-        private void OnLeaderChanged(object? sender, ILeaderElection.OnLeaderChangedEventArgs e)
+        private void OnLeaderChanged(object? sender, OnLeaderChangedEventArgs e)
         {
             if (e.IsLeader)
+            {
                 return;
+            }
 
             // if we are no longer the leader cancel any pending replications
             // TODO: Reset replication token if we are the new leader
@@ -79,7 +78,7 @@ namespace Horde.Storage.Implementation
                 case ReplicatorVersion.Refs:
                     return ActivatorUtilities.CreateInstance<RefsReplicator>(provider, replicatorSettings);
                 default:
-                    throw new ArgumentOutOfRangeException();
+                    throw new NotImplementedException($"Unknown replicator version: {replicatorSettings.Version}");
             }
         }
 
diff --git a/Engine/Source/Programs/Horde/HordeStorage/Horde.Storage/Implementation/Replication/ReplicationSettings.cs b/Engine/Source/Programs/Horde/HordeStorage/Horde.Storage/Implementation/Replication/ReplicationSettings.cs
--- a/Engine/Source/Programs/Horde/HordeStorage/Horde.Storage/Implementation/Replication/ReplicationSettings.cs
+++ b/Engine/Source/Programs/Horde/HordeStorage/Horde.Storage/Implementation/Replication/ReplicationSettings.cs
@@ -1,16 +1,18 @@
 // Copyright Epic Games, Inc. All Rights Reserved.
 
+using System;
 using System.Collections.Generic;
 using System.ComponentModel.DataAnnotations;
 using EpicGames.Horde.Storage;
 using Horde.Storage.Implementation;
-using Jupiter.Implementation;
 
 namespace Horde.Storage
 {
     // ReSharper disable once ClassNeverInstantiated.Global
     public class ReplicationSettings
     {
+        private string _stateRoot = "";
+
         /// <summary>
         /// Enable to start a replicating another Jupiter instance into this one
         /// </summary>
@@ -20,7 +22,11 @@ namespace Horde.Storage
         /// Path to a directory were the local state can be kept
         /// </summary>
         [Required]
-        public string StateRoot { get; set; } = "";
+        public string StateRoot
+        {
+            get => PathUtil.ResolvePath(_stateRoot);
+            set => _stateRoot = value;
+        }
 
         /// <summary>
         /// Path to a directory where state used to be stored. State is migrated from old state root to the new during startup.
@@ -35,13 +41,9 @@ namespace Horde.Storage
         public int ReplicationPollFrequencySeconds { get; set; } = 15;
         
         [Required]
+        [System.Diagnostics.CodeAnalysis.SuppressMessage("Usage", "CA2227:Collection properties should be read only", Justification = "Used by serialization")]
         // ReSharper disable once CollectionNeverUpdated.Global
         public List<ReplicatorSettings> Replicators { get; set; } = new List<ReplicatorSettings>();
-
-        public string GetStateRoot()
-        {
-            return PathUtil.ResolvePath(StateRoot);
-        }
     }
 
     public class ReplicatorSettings
@@ -109,7 +111,7 @@ namespace Horde.Storage
         /// <summary>
         /// The url to login for a token to use to connect to the other services. Set to empty to disable login, assuming a unprotected service.
         /// </summary>
-        public string OAuthLoginUrl { get; set; } = "";
+        public Uri? OAuthLoginUrl { get; set; } = null;
 
         /// <summary>
         /// The scope to request
diff --git a/Engine/Source/Programs/Horde/HordeStorage/Horde.Storage/Implementation/Replication/Replicator.cs b/Engine/Source/Programs/Horde/HordeStorage/Horde.Storage/Implementation/Replication/Replicator.cs
--- a/Engine/Source/Programs/Horde/HordeStorage/Horde.Storage/Implementation/Replication/Replicator.cs
+++ b/Engine/Source/Programs/Horde/HordeStorage/Horde.Storage/Implementation/Replication/Replicator.cs
@@ -40,29 +40,29 @@ namespace Horde.Storage.Implementation
         protected override IAsyncEnumerable<TransactionEvent> GetCallistoOp(long stateReplicatorOffset, Guid? stateReplicatingGeneration, string currentSite,
             OpsEnumerationState enumerationState, CancellationToken replicationToken)
         {
-            CallistoReader remoteCallistoReader = new CallistoReader(_restClient, _replicatorSettings.NamespaceToReplicate);
+            CallistoReader remoteCallistoReader = new CallistoReader(Client, ReplicatorSettings.NamespaceToReplicate);
 
-            return remoteCallistoReader.GetOps(stateReplicatorOffset, stateReplicatingGeneration, currentSite, enumerationState: enumerationState, cancellationToken: replicationToken, maxOffsetsAttempted: _replicatorSettings.MaxOffsetsAttempted);
+            return remoteCallistoReader.GetOps(stateReplicatorOffset, stateReplicatingGeneration, currentSite, enumerationState: enumerationState, cancellationToken: replicationToken, maxOffsetsAttempted: ReplicatorSettings.MaxOffsetsAttempted);
         }
 
         protected override async Task ReplicateOp(IRestClient remoteClient, TransactionEvent op, CancellationToken replicationToken)
         {
             using IScope scope = Tracer.Instance.StartActive("replicator.replicate_blobs");
-            NamespaceId ns = _replicatorSettings.NamespaceToReplicate;
+            NamespaceId ns = ReplicatorSettings.NamespaceToReplicate;
 
             // now we replicate the blobs
             switch (op)
             {
                 case AddTransactionEvent addEvent:
                     BlobIdentifier[] blobs = addEvent.Blobs;
-                    await ReplicateBlobs(remoteClient, ns, blobs, replicationToken);
+                    await ReplicateBlobs(ns, blobs, replicationToken);
                     break;
                 case RemoveTransactionEvent removeEvent:
                     // TODO: Do we even want to do anything? we can not delete the blob in the store because it may be used by something else
                     // so we wait for the GC to remove it.
                     break;
                 default:
-                    throw new ArgumentOutOfRangeException();
+                    throw new NotImplementedException("Unknown op type {op}");
             }
         }
 
@@ -71,77 +71,71 @@ namespace Horde.Storage.Implementation
             using IScope scope = Tracer.Instance.StartActive("replicator.replicate_inline");
             //scope.Span.ResourceName = op.transactionId;
 
-            NamespaceId ns = _replicatorSettings.NamespaceToReplicate;
+            NamespaceId ns = ReplicatorSettings.NamespaceToReplicate;
 
             // Put a copy of the event in the local callisto store
             // Make sure this is put before the io put request as we want to pin the reference before uploading the content to avoid the gc removing it right away
             long newTransactionId = await _transactionLogWriter.Add(ns, op);
 
-            _logger.Information("{Name} Replicated Op {@Op} to local store as transaction id {@TransactionId}", _name, op, newTransactionId);
+            Logger.Information("{Name} Replicated Op {@Op} to local store as transaction id {@TransactionId}", Name, op, newTransactionId);
 
             // TODO: Add Refstore replication
             //_refStore.Add();
 
             return newTransactionId;
         }
-
     }
 
     public abstract class Replicator<T> : IReplicator
         where T: IReplicationEvent
     {
         private readonly IBlobService _blobService;
-        protected readonly string _name;
         private readonly string _currentSite;
-        protected readonly ReplicatorSettings _replicatorSettings;
-        protected readonly ILogger _logger = Log.ForContext<Replicator<T>>();
 
         private readonly FileInfo _stateFile;
         private bool _replicationRunning;
         private readonly AsyncManualResetEvent _replicationFinishedEvent = new AsyncManualResetEvent(true);
         private readonly CancellationTokenSource _replicationTokenSource = new CancellationTokenSource();
-        private readonly IReplicator.ReplicatorInfo _replicatorInfo;
-        protected readonly IRestClient _restClient;
         private readonly IServiceCredentials _serviceCredentials;
         private readonly HttpClient _httpClient;
 
+        protected string Name { get; set; }
+        protected ReplicatorSettings ReplicatorSettings { get; init; }
 
-        public IReplicator.ReplicatorState State { get; private set; }
-
-        public IReplicator.ReplicatorInfo Info
-        {
-            get { return _replicatorInfo; }
-        }
+        public ReplicatorState State { get; private set; }
+        protected ILogger Logger { get; } = Log.ForContext<Replicator<T>>();
+        public ReplicatorInfo Info { get; }
+        protected IRestClient Client { get; }
 
         protected static IRestClient CreateRemoteClient(ReplicatorSettings replicatorSettings, IServiceCredentials serviceCredentials)
         {
-            var remoteClient = new RestClient(replicatorSettings.ConnectionString).UseSerializer(() => new JsonNetSerializer());
+            IRestClient remoteClient = new RestClient(replicatorSettings.ConnectionString).UseSerializer(() => new JsonNetSerializer());
             remoteClient.Authenticator = serviceCredentials.GetAuthenticator();
             return remoteClient;
         }
 
         protected Replicator(ReplicatorSettings replicatorSettings, IOptionsMonitor<ReplicationSettings> replicationSettings, IOptionsMonitor<JupiterSettings> jupiterSettings, IBlobService blobService, IRestClient remoteClient, IServiceCredentials serviceCredentials, IHttpClientFactory httpClientFactory)
         {
-            _name = replicatorSettings.ReplicatorName;
+            Name = replicatorSettings.ReplicatorName;
             _blobService = blobService;
             _currentSite = jupiterSettings.CurrentValue.CurrentSite;
-            _replicatorSettings = replicatorSettings;
+            ReplicatorSettings = replicatorSettings;
 
-            string stateFileName = $"{_name}.json";
-            DirectoryInfo stateRoot = new DirectoryInfo(replicationSettings.CurrentValue.GetStateRoot());
+            string stateFileName = $"{Name}.json";
+            DirectoryInfo stateRoot = new DirectoryInfo(replicationSettings.CurrentValue.StateRoot);
             _stateFile = new FileInfo(Path.Combine(stateRoot.FullName, stateFileName));
 
             if (_stateFile.Exists)
             {
-                State = ReadState(_stateFile) ?? new IReplicator.ReplicatorState {ReplicatorOffset = null, ReplicatingGeneration = null};
+                State = ReadState(_stateFile) ?? new ReplicatorState {ReplicatorOffset = null, ReplicatingGeneration = null};
             }
             else
             {
-                State = new IReplicator.ReplicatorState {ReplicatorOffset = null, ReplicatingGeneration = null};
+                State = new ReplicatorState {ReplicatorOffset = null, ReplicatingGeneration = null};
             }
 
-            _replicatorInfo = new IReplicator.ReplicatorInfo(_replicatorSettings.ReplicatorName, _replicatorSettings.NamespaceToReplicate, State);
-            _restClient = remoteClient;
+            Info = new ReplicatorInfo(ReplicatorSettings.ReplicatorName, ReplicatorSettings.NamespaceToReplicate, State);
+            Client = remoteClient;
 
             _serviceCredentials = serviceCredentials;
 
@@ -151,7 +145,7 @@ namespace Horde.Storage.Implementation
 
         ~Replicator()
         {
-            Dispose();
+            Dispose(false);
         }
 
         /// <summary>
@@ -182,27 +176,27 @@ namespace Horde.Storage.Implementation
                 _stateFile.Refresh();
                 if (_stateFile.Exists)
                 {
-                    State = ReadState(_stateFile) ?? new IReplicator.ReplicatorState {ReplicatorOffset = null, ReplicatingGeneration = null};
+                    State = ReadState(_stateFile) ?? new ReplicatorState {ReplicatorOffset = null, ReplicatingGeneration = null};
                 }
                 else
                 {
-                    State = new IReplicator.ReplicatorState {ReplicatorOffset = null, ReplicatingGeneration = null};
+                    State = new ReplicatorState {ReplicatorOffset = null, ReplicatingGeneration = null};
                 }
 
-                LogReplicationHeartbeat(_name, State, 0);
-                _logger.Information("{Name} Looking for new transaction. Previous state: {@State}", _name, State);
+                LogReplicationHeartbeat(Name, State, 0);
+                Logger.Information("{Name} Looking for new transaction. Previous state: {@State}", Name, State);
 
                 SortedList<long, Task> replicationTasks = new();
                 await foreach (T op in GetCallistoOp(State.ReplicatorOffset ?? 0, State.ReplicatingGeneration, _currentSite, enumerationState, replicationToken))
                 {
                     hasRun = true;
-                    _logger.Information("{Name} New transaction to replicate found. New op: {@Op} . Count of running replications: {CurrentReplications}", _name, op, replicationTasks.Count);
+                    Logger.Information("{Name} New transaction to replicate found. New op: {@Op} . Count of running replications: {CurrentReplications}", Name, op, replicationTasks.Count);
 
                     Info.CountOfRunningReplications = replicationTasks.Count;
 
                     if (!op.Identifier.HasValue || !op.NextIdentifier.HasValue)
                     {
-                        _logger.Error("{Name} Missing identifier in {@Op}", _name, op);
+                        Logger.Error("{Name} Missing identifier in {@Op}", Name, op);
                         throw new Exception("Missing identifier in op. Version mismatch?");
                     }
                     
@@ -212,10 +206,10 @@ namespace Horde.Storage.Implementation
                     try
                     {
                         // first replicate the things which are order sensitive like the transaction log
-                        long? _ = await ReplicateOpInline(_restClient, op, replicationToken);
+                        long? _ = await ReplicateOpInline(Client, op, replicationToken);
 
                         // next we replicate the large things that overlap with other replications like blobs
-                        Task currentReplicationTask = ReplicateOp(_restClient, op, replicationToken);
+                        Task currentReplicationTask = ReplicateOp(Client, op, replicationToken);
 
                         Task OnDone(Task task)
                         {
@@ -246,7 +240,7 @@ namespace Horde.Storage.Implementation
                                 Info.LastRun = DateTime.Now;
                                 Info.CountOfRunningReplications = replicationTasks.Count;
 
-                                LogReplicationHeartbeat(_name, State, replicationTasks.Count);
+                                LogReplicationHeartbeat(Name, State, replicationTasks.Count);
 
                                 return task;
                             }
@@ -261,19 +255,19 @@ namespace Horde.Storage.Implementation
 
                         lock (replicationTasks)
                         {
-                            Task<Task> commitStateTask = currentReplicationTask.ContinueWith(OnDone, replicationToken);
+                            Task<Task> commitStateTask = currentReplicationTask.ContinueWith(OnDone, replicationToken, TaskContinuationOptions.None, TaskScheduler.Current);
 
                             replicationTasks.Add(currentOffset, commitStateTask);
                         }
                     }
                     catch (BlobNotFoundException)
                     {
-                        _logger.Warning("{Name} Failed to replicate {@Op} in {Namespace} because blob was not present in remote store. Skipping.", _name, op, Info.NamespaceToReplicate);
+                        Logger.Warning("{Name} Failed to replicate {@Op} in {Namespace} because blob was not present in remote store. Skipping.", Name, op, Info.NamespaceToReplicate);
                     }
 
                     // if we have reached the max amount of parallel replications we wait for one of them to finish before starting a new one
                     // if max replications is set to -1 we do not limit the concurrency
-                    if ( _replicatorSettings.MaxParallelReplications != -1 && replicationTasks.Count >= _replicatorSettings.MaxParallelReplications)
+                    if ( ReplicatorSettings.MaxParallelReplications != -1 && replicationTasks.Count >= ReplicatorSettings.MaxParallelReplications)
                     {
                         Task[] currentReplicationTasks;
                         lock (replicationTasks)
@@ -284,7 +278,7 @@ namespace Horde.Storage.Implementation
                     }
 
                     Info.CountOfRunningReplications = replicationTasks.Count;
-                    LogReplicationHeartbeat(_name, State, replicationTasks.Count);
+                    LogReplicationHeartbeat(Name, State, replicationTasks.Count);
 
                     if (replicationToken.IsCancellationRequested)
                     {
@@ -303,7 +297,7 @@ namespace Horde.Storage.Implementation
             }
 
             DateTime endedReplicationAt = DateTime.Now;
-            _logger.Information("{Name} {Namespace} Replication Finished at {EndTime} and last run was at {LastTime}, replicated {CountOfEvents} events. Thus was {TimeDifference} minutes behind.", _name, Info.NamespaceToReplicate, endedReplicationAt, startedReplicationAt, countOfReplicatedEvents, (endedReplicationAt - startedReplicationAt).TotalMinutes);
+            Logger.Information("{Name} {Namespace} Replication Finished at {EndTime} and last run was at {LastTime}, replicated {CountOfEvents} events. Thus was {TimeDifference} minutes behind.", Name, Info.NamespaceToReplicate, endedReplicationAt, startedReplicationAt, countOfReplicatedEvents, (endedReplicationAt - startedReplicationAt).TotalMinutes);
 
             if (!hasRun)
             {
@@ -318,23 +312,25 @@ namespace Horde.Storage.Implementation
 
         protected abstract IAsyncEnumerable<T> GetCallistoOp(long stateReplicatorOffset, Guid? stateReplicatingGeneration, string currentSite, OpsEnumerationState enumerationState, CancellationToken replicationToken);
         
-        private void LogReplicationHeartbeat(string name, IReplicator.ReplicatorState? state, int countOfCurrentReplications)
+        private void LogReplicationHeartbeat(string name, ReplicatorState? state, int countOfCurrentReplications)
         {
             if (state == null)
+            {
                 return;
+            }
 
             // log message used to generate metric for how many replications are currently running
-            _logger.Information("{Name} replication has run . Count of running replications: {CurrentReplications}", _name, countOfCurrentReplications);
+            Logger.Information("{Name} replication has run . Count of running replications: {CurrentReplications}", Name, countOfCurrentReplications);
 
             // log message used to verify replicators are actually running
-            _logger.Information("{Name} starting replication. Last transaction was {TransactionId} {Generation}", name, state.ReplicatorOffset.GetValueOrDefault(0L), state.ReplicatingGeneration.GetValueOrDefault(Guid.Empty) );
+            Logger.Information("{Name} starting replication. Last transaction was {TransactionId} {Generation}", name, state.ReplicatorOffset.GetValueOrDefault(0L), state.ReplicatingGeneration.GetValueOrDefault(Guid.Empty) );
         }
 
         protected abstract Task<long?> ReplicateOpInline(IRestClient remoteClient, T op, CancellationToken replicationToken);
 
         protected abstract Task ReplicateOp(IRestClient remoteClient, T op, CancellationToken replicationToken);
 
-        protected async Task ReplicateBlobs(IRestClient remoteClient, NamespaceId ns, BlobIdentifier[] blobs, CancellationToken replicationToken)
+        protected async Task ReplicateBlobs(NamespaceId ns, BlobIdentifier[] blobs, CancellationToken replicationToken)
         {
             Task[] blobReplicateTasks = new Task[blobs.Length];
             for (int index = 0; index < blobs.Length; index++)
@@ -344,7 +340,9 @@ namespace Horde.Storage.Implementation
                 {
                     // check if this blob exists locally before replicating
                     if (await _blobService.Exists(ns, blob))
+                    {
                         return;
+                    }
 
                     // attempt to replicate for a few tries with some delay in between
                     // this because new transactions are written to callisto first (to establish the GC handle) before content is uploaded to io.
@@ -359,26 +357,26 @@ namespace Horde.Storage.Implementation
 
                         if (response.StatusCode == HttpStatusCode.NotFound)
                         {
-                            _logger.Information("{Blob} was not found in remote blob store. Retry attempt {Attempts}", blob, attempts);
+                            Logger.Information("{Blob} was not found in remote blob store. Retry attempt {Attempts}", blob, attempts);
                             await Task.Delay(TimeSpan.FromSeconds(5), replicationToken);
                             continue;
                         }
 
                         if (response.StatusCode == HttpStatusCode.GatewayTimeout)
                         {
-                            _logger.Information("GatewayTimeout while replicating {Blob}, retrying. Retry attempt {Attempts}", blob, attempts);
+                            Logger.Information("GatewayTimeout while replicating {Blob}, retrying. Retry attempt {Attempts}", blob, attempts);
                             continue;
                         }
 
                         if (response.StatusCode == HttpStatusCode.BadGateway)
                         {
-                            _logger.Information("BadGateway while replicating {Blob}, retrying. Retry attempt {Attempts}", blob, attempts);
+                            Logger.Information("BadGateway while replicating {Blob}, retrying. Retry attempt {Attempts}", blob, attempts);
                             continue;
                         }
 
                         if (response.StatusCode == HttpStatusCode.BadRequest || response.StatusCode == HttpStatusCode.NotFound)
                         {
-                            _logger.Warning("Remote blob {Blob} missing, unable to replicate.", blob);
+                            Logger.Warning("Remote blob {Blob} missing, unable to replicate.", blob);
                             return;
                         }
 
@@ -395,7 +393,7 @@ namespace Horde.Storage.Implementation
                             }
                             if (!blob.Equals(calculatedBlob))
                             {
-                                _logger.Warning("Mismatching blob when replicating {Blob}. Determined Hash was {Hash} size was {Size} HttpStatusCode {StatusCode}",
+                                Logger.Warning("Mismatching blob when replicating {Blob}. Determined Hash was {Hash} size was {Size} HttpStatusCode {StatusCode}",
                                     blob, calculatedBlob, payload.Length, response.StatusCode);
                                 continue; // attempt to replicate again
                             }
@@ -404,25 +402,23 @@ namespace Horde.Storage.Implementation
                             break;
                         }
 
-                        throw new Exception($"Replicator \"{_name}\" failed to replicate blob {blob} due unsuccessful status from remote blob store: {response.StatusCode} . Error message: \"{response.ReasonPhrase}\"");
+                        throw new Exception($"Replicator \"{Name}\" failed to replicate blob {blob} due unsuccessful status from remote blob store: {response.StatusCode} . Error message: \"{response.ReasonPhrase}\"");
                     }
-
-
                 }, replicationToken);
             }
 
             await Task.WhenAll(blobReplicateTasks);
         }
 
-        private static IReplicator.ReplicatorState? ReadState(FileInfo stateFile)
+        private static ReplicatorState? ReadState(FileInfo stateFile)
         {
             using StreamReader streamReader = stateFile.OpenText();
             using JsonReader reader = new JsonTextReader(streamReader);
             JsonSerializer serializer = new JsonSerializer();
-            return serializer.Deserialize<IReplicator.ReplicatorState>(reader);
+            return serializer.Deserialize<ReplicatorState>(reader);
         }
 
-        private static void SaveState(FileInfo stateFile, IReplicator.ReplicatorState newState)
+        private static void SaveState(FileInfo stateFile, ReplicatorState newState)
         {
             using StreamWriter writer = stateFile.CreateText();
             JsonSerializer serializer = new JsonSerializer();
@@ -434,14 +430,28 @@ namespace Horde.Storage.Implementation
             string? token = _serviceCredentials.GetToken();
             HttpRequestMessage request = new HttpRequestMessage(httpMethod, uri);
             if (!string.IsNullOrEmpty(token))
+            {
                 request.Headers.Add("Authorization", "Bearer " + token);
+            }
+
             return request;
         }
+
         public void Dispose()
         {
-            SaveState(_stateFile, State);
+            Dispose(true);
+            GC.SuppressFinalize(this);
+        }
 
-            _replicationTokenSource.Dispose();
+        protected virtual void Dispose(bool disposing)
+        {
+            if (disposing)
+            {
+                SaveState(_stateFile, State);
+
+                _replicationTokenSource.Dispose();
+                _httpClient.Dispose();
+            }
         }
 
         public async Task StopReplicating()
@@ -458,7 +468,7 @@ namespace Horde.Storage.Implementation
 
         public Task DeleteState()
         {
-            State = new IReplicator.ReplicatorState()
+            State = new ReplicatorState()
             {
                 ReplicatingGeneration = null,
                 ReplicatorOffset = 0
diff --git a/Engine/Source/Programs/Horde/HordeStorage/Horde.Storage/Implementation/ScyllaSessionManager.cs b/Engine/Source/Programs/Horde/HordeStorage/Horde.Storage/Implementation/ScyllaSessionManager.cs
--- a/Engine/Source/Programs/Horde/HordeStorage/Horde.Storage/Implementation/ScyllaSessionManager.cs
+++ b/Engine/Source/Programs/Horde/HordeStorage/Horde.Storage/Implementation/ScyllaSessionManager.cs
@@ -1,9 +1,5 @@
 ﻿// Copyright Epic Games, Inc. All Rights Reserved.
 
-using System;
-using System.Collections.Generic;
-using System.Linq;
-using System.Threading.Tasks;
 using Cassandra;
 
 namespace Horde.Storage.Implementation
@@ -35,6 +31,5 @@ namespace Horde.Storage.Implementation
         {
             return _localSession;
         }
-
     }
 }
