Horde Storage - Updated editorconfig to more closely match what we run for Horde and updated code to fix all static analysis errors

diff --git a/Engine/Source/Programs/Horde/HordeStorage/Horde.Storage/Implementation/ServiceCredentials.cs b/Engine/Source/Programs/Horde/HordeStorage/Horde.Storage/Implementation/ServiceCredentials.cs
--- a/Engine/Source/Programs/Horde/HordeStorage/Horde.Storage/Implementation/ServiceCredentials.cs
+++ b/Engine/Source/Programs/Horde/HordeStorage/Horde.Storage/Implementation/ServiceCredentials.cs
@@ -1,9 +1,6 @@
 ﻿// Copyright Epic Games, Inc. All Rights Reserved.
 
 using System;
-using System.Collections.Generic;
-using System.Linq;
-using System.Threading.Tasks;
 using Jupiter.Common.Implementation;
 using Jupiter.Implementation;
 using Microsoft.Extensions.Options;
@@ -20,20 +17,24 @@ namespace Horde.Storage.Implementation
 
     public class ServiceCredentials : IServiceCredentials
     {
-        private readonly IOptionsMonitor<ServiceCredentialSettings> _settings;
         private readonly ClientCredentialOAuthAuthenticator? _authenticator;
 
         public ServiceCredentials(IOptionsMonitor<ServiceCredentialSettings> settings, ISecretResolver secretResolver)
         {
-            _settings = settings;
-            if (!string.IsNullOrEmpty(settings.CurrentValue.OAuthLoginUrl))
+            if (settings.CurrentValue.OAuthLoginUrl != null)
             {
                 string? clientId = secretResolver.Resolve(settings.CurrentValue.OAuthClientId);
                 if (string.IsNullOrEmpty(clientId))
+                {
                     throw new ArgumentException("ClientId must be set when using a service credential");
+                }
+
                 string? clientSecret = secretResolver.Resolve(settings.CurrentValue.OAuthClientSecret);
                 if (string.IsNullOrEmpty(clientSecret))
+                {
                     throw new ArgumentException("ClientSecret must be set when using a service credential");
+                }
+
                 _authenticator = new ClientCredentialOAuthAuthenticator(settings.CurrentValue.OAuthLoginUrl, clientId, clientSecret, settings.CurrentValue.OAuthScope);
             }
         }
@@ -48,5 +49,4 @@ namespace Horde.Storage.Implementation
             return _authenticator?.Authenticate();
         }
     }
-
 }
diff --git a/Engine/Source/Programs/Horde/HordeStorage/Horde.Storage/Implementation/TransactionLog/CallistoTransactionLogWriter.cs b/Engine/Source/Programs/Horde/HordeStorage/Horde.Storage/Implementation/TransactionLog/CallistoTransactionLogWriter.cs
--- a/Engine/Source/Programs/Horde/HordeStorage/Horde.Storage/Implementation/TransactionLog/CallistoTransactionLogWriter.cs
+++ b/Engine/Source/Programs/Horde/HordeStorage/Horde.Storage/Implementation/TransactionLog/CallistoTransactionLogWriter.cs
@@ -22,7 +22,6 @@ namespace Horde.Storage.Implementation
             _client.Authenticator = serviceCredentials.GetAuthenticator();
         }
 
-
         public async Task<long> Add(NamespaceId ns, TransactionEvent @event)
         {
             const int retryAttempts = 15;
@@ -38,7 +37,7 @@ namespace Horde.Storage.Implementation
                 if (!response.IsSuccessful)
                 {
                     // if callisto is down we wait a while and try again
-                    if (response.StatusCode == 0 && response.ErrorMessage.StartsWith("Connection refused"))
+                    if (response.StatusCode == 0 && response.ErrorMessage.StartsWith("Connection refused", StringComparison.OrdinalIgnoreCase))
                     {
                         await Task.Delay(1000);
                         continue; // retry
@@ -56,8 +55,8 @@ namespace Horde.Storage.Implementation
 
         public async Task<long> Delete(NamespaceId ns, BucketId bucket, KeyId key)
         {
-            RestRequest request = new RestRequest("api/v1/t/{ns}");
-            var record = new CallistoRequest("Delete", key, bucket);
+            RestRequest request = new("api/v1/t/{ns}");
+            CallistoRequest record = new("Delete", key, bucket);
             request.AddUrlSegment("ns", ns);
             request.AddJsonBody(record);
             
@@ -66,7 +65,10 @@ namespace Horde.Storage.Implementation
             CallistoResponse response = await _client.PostAsync<CallistoResponse>(request);
 
             if (response == null)
+            {
                 throw new Exception("Failed to put delete record in Callisto.");
+            }
+
             return response.Offset;
         }
     }
diff --git a/Engine/Source/Programs/Horde/HordeStorage/Horde.Storage/Implementation/TransactionLog/IReplicationLog.cs b/Engine/Source/Programs/Horde/HordeStorage/Horde.Storage/Implementation/TransactionLog/IReplicationLog.cs
--- a/Engine/Source/Programs/Horde/HordeStorage/Horde.Storage/Implementation/TransactionLog/IReplicationLog.cs
+++ b/Engine/Source/Programs/Horde/HordeStorage/Horde.Storage/Implementation/TransactionLog/IReplicationLog.cs
@@ -21,17 +21,10 @@ namespace Horde.Storage.Implementation
         Task<SnapshotInfo?> GetLatestSnapshot(NamespaceId ns);
         IAsyncEnumerable<SnapshotInfo> GetSnapshots(NamespaceId ns);
 
-
         Task UpdateReplicatorState(NamespaceId ns, string replicatorName, ReplicatorState newState);
         Task<ReplicatorState?> GetReplicatorState(NamespaceId ns, string replicatorName);
     }
 
-    public class ReplicatorState
-    {
-        public string? LastBucket { get; set; }
-        public Guid? LastEvent { get; set; }
-    }
-
     public class SnapshotInfo
     {
         [JsonConstructor]
diff --git a/Engine/Source/Programs/Horde/HordeStorage/Horde.Storage/Implementation/TransactionLog/MemoryReplicationLog.cs b/Engine/Source/Programs/Horde/HordeStorage/Horde.Storage/Implementation/TransactionLog/MemoryReplicationLog.cs
--- a/Engine/Source/Programs/Horde/HordeStorage/Horde.Storage/Implementation/TransactionLog/MemoryReplicationLog.cs
+++ b/Engine/Source/Programs/Horde/HordeStorage/Horde.Storage/Implementation/TransactionLog/MemoryReplicationLog.cs
@@ -86,7 +86,9 @@ namespace Horde.Storage.Implementation
             {
                 // if we have no buckets we are done
                 if (!buckets.Any())
+                {
                     yield break;
+                }
 
                 lastBucket = buckets.First().Key;
             }
@@ -172,7 +174,6 @@ namespace Horde.Storage.Implementation
             {
                 yield return snapshot;
             }
-
         }
 
         public Task UpdateReplicatorState(NamespaceId ns, string replicatorName,
diff --git a/Engine/Source/Programs/Horde/HordeStorage/Horde.Storage/Implementation/TransactionLog/ReplicationLogSnapshot.cs b/Engine/Source/Programs/Horde/HordeStorage/Horde.Storage/Implementation/TransactionLog/ReplicationLogSnapshot.cs
--- a/Engine/Source/Programs/Horde/HordeStorage/Horde.Storage/Implementation/TransactionLog/ReplicationLogSnapshot.cs
+++ b/Engine/Source/Programs/Horde/HordeStorage/Horde.Storage/Implementation/TransactionLog/ReplicationLogSnapshot.cs
@@ -2,20 +2,31 @@
 
 using System;
 using System.Collections.Generic;
-using System.Diagnostics;
 using System.IO;
 using System.IO.Compression;
-using System.Linq;
 using System.Text;
-using System.Threading.Tasks;
 using EpicGames.Core;
 using EpicGames.Horde.Storage;
+using JetBrains.Annotations;
 using Jupiter.Common.Implementation;
 using Jupiter.Implementation;
 using Newtonsoft.Json;
 
 namespace Horde.Storage.Implementation.TransactionLog
 {
+    [UsedImplicitly]
+    public class ReplicationLogSnapshotState
+    {
+        public ReplicationLogSnapshotState(SnapshotHeader header, List<SnapshotLiveObject> liveObjects)
+        {
+            Header = header;
+            LiveObjects = liveObjects;
+        }
+
+        public SnapshotHeader Header { get; set; }
+        public List<SnapshotLiveObject> LiveObjects { get; }
+    }
+
     public class JsonReplicationLogSnapshot : ReplicationLogSnapshot
     {
         private readonly List<SnapshotLiveObject> _liveObjects;
@@ -83,11 +94,15 @@ namespace Horde.Storage.Implementation.TransactionLog
             byte[] bytes = new byte[4];
             int bytesRead = stream.Read(bytes, 0, 4);
             if (bytesRead != 4)
+            {
                 throw new Exception("Partial stream read when deserializing snapshot");
+            }
 
             bool hasMagic = bytes[0] == 'S' && bytes[1] == 'N' && bytes[2] == 'A' && bytes[3] == 'P';
             if (!hasMagic)
+            {
                 throw new Exception("Did not find magic prefix when deserializing snapshot, incorrect format?");
+            }
 
             using BinaryReader reader = new BinaryReader(stream, Encoding.ASCII, leaveOpen: true);
             string ns = reader.ReadString();
@@ -98,11 +113,12 @@ namespace Horde.Storage.Implementation.TransactionLog
             return (new NamespaceId(ns), lastBucket, lastEvent, countOfObjects);
         }
 
-
         public override IEnumerable<SnapshotLiveObject> GetLiveObjects()
         {
             if (_filesystemBufferedPayload == null)
+            {
                 yield break;
+            }
 
             using Stream payloadStream = _filesystemBufferedPayload.GetStream();
 
@@ -131,12 +147,16 @@ namespace Horde.Storage.Implementation.TransactionLog
         public static ReplicationLogSnapshot DeserializeSnapshotFromStream(Stream stream)
         {
             if (!stream.CanSeek)
+            {
                 throw new Exception("Seekable stream required");
+            }
 
             byte[] bytes = new byte[4];
             int bytesRead = stream.Read(bytes, 0, 4);
             if (bytesRead != 4)
+            {
                 throw new Exception("Partial stream read when deserializing snapshot");
+            }
 
             bool hasMagic = bytes[0] == 'S' && bytes[1] == 'N' && bytes[2] == 'A' && bytes[3] == 'P';
             stream.Seek(0, SeekOrigin.Begin);
@@ -155,45 +175,33 @@ namespace Horde.Storage.Implementation.TransactionLog
         }
     }
 
-    public abstract class ReplicationLogSnapshot: IDisposable
+    public class SnapshotHeader
     {
-        private readonly List<SnapshotLiveObject> _addedObjects = new List<SnapshotLiveObject>();
-        private readonly HashSet<(BucketId, IoHashKey)> _removedObjects = new HashSet<(BucketId, IoHashKey)>();
-
-        public class ReplicationLogSnapshotState
-        {
-            public ReplicationLogSnapshotState(SnapshotHeader header, List<SnapshotLiveObject> liveObjects)
-            {
-                Header = header;
-                LiveObjects = liveObjects;
-            }
-
-            public SnapshotHeader Header { get; set; }
-            public List<SnapshotLiveObject> LiveObjects { get; set; }
-        }
+        public NamespaceId Namespace { get; set; }
+        public string LastBucket { get; set; } = null!;
+        public Guid LastEvent { get; set; }
+    }
 
-        public class SnapshotHeader
+    public class SnapshotLiveObject
+    {
+        public SnapshotLiveObject(BucketId bucket, IoHashKey key, BlobIdentifier blob)
         {
-            public NamespaceId Namespace { get; set; }
-            public string LastBucket { get; set; } = null!;
-            public Guid LastEvent { get; set; }
+            Bucket = bucket;
+            Key = key;
+            Blob = blob;
         }
 
-        public class SnapshotLiveObject
-        {
-            public SnapshotLiveObject(BucketId bucket, IoHashKey key, BlobIdentifier blob)
-            {
-                Bucket = bucket;
-                Key = key;
-                Blob = blob;
-            }
+        public BucketId Bucket { get; set; }
+        public IoHashKey Key { get; set; }
+        public BlobIdentifier Blob { get; set; }
+    }
 
-            public BucketId Bucket { get; set; }
-            public IoHashKey Key { get; set; }
-            public BlobIdentifier Blob { get; set; }
-        }
+    public abstract class ReplicationLogSnapshot
+    {
+        private readonly List<SnapshotLiveObject> _addedObjects = new List<SnapshotLiveObject>();
+        private readonly HashSet<(BucketId, IoHashKey)> _removedObjects = new HashSet<(BucketId, IoHashKey)>();
 
-        public ReplicationLogSnapshot(NamespaceId ns)
+        protected ReplicationLogSnapshot(NamespaceId ns)
         {
             Namespace = ns;
         }
@@ -216,9 +224,14 @@ namespace Horde.Storage.Implementation.TransactionLog
         public void Serialize(Stream stream)
         {
             if (LastBucket == null)
+            {
                 throw new Exception("No last bucket found when serializing state, did you really have events?");
+            }
+
             if (LastEvent == null)
+            {
                 throw new Exception("No last event found when serializing state, did you really have events?");
+            }
 
             void WriteHeader(Stream s, NamespaceId ns, string lastBucket, Guid lastEvent, ulong countOfObjects)
             {
@@ -275,7 +288,6 @@ namespace Horde.Storage.Implementation.TransactionLog
             WriteHeader(stream, Namespace, LastBucket, LastEvent.Value, countOfLiveObjects);
         }
 
-
         public void ProcessEvent(ReplicationLogEvent entry)
         {
             LastBucket = entry.TimeBucket;
diff --git a/Engine/Source/Programs/Horde/HordeStorage/Horde.Storage/Implementation/TransactionLog/ReplicationLogSnapshotBuilder.cs b/Engine/Source/Programs/Horde/HordeStorage/Horde.Storage/Implementation/TransactionLog/ReplicationLogSnapshotBuilder.cs
--- a/Engine/Source/Programs/Horde/HordeStorage/Horde.Storage/Implementation/TransactionLog/ReplicationLogSnapshotBuilder.cs
+++ b/Engine/Source/Programs/Horde/HordeStorage/Horde.Storage/Implementation/TransactionLog/ReplicationLogSnapshotBuilder.cs
@@ -33,7 +33,9 @@ namespace Horde.Storage.Implementation.TransactionLog
             SnapshotInfo? snapshotInfo = await _replicationLog.GetLatestSnapshot(ns);
 
             if (cancellationToken.IsCancellationRequested)
+            {
                 throw new TaskCanceledException();
+            }
 
             ReplicationLogSnapshot snapshot;
             string? lastBucket;
@@ -43,7 +45,10 @@ namespace Horde.Storage.Implementation.TransactionLog
                 // append to the previous snapshot if one is available
                 await using BlobContents blobContents = await _blobService.GetObject(snapshotInfo.BlobNamespace, snapshotInfo.SnapshotBlob);
                 if (cancellationToken.IsCancellationRequested)
+                {
                     throw new TaskCanceledException();
+                }
+
                 using FilesystemBufferedPayload snapshotPayload = await FilesystemBufferedPayload.Create(blobContents.Stream);
                 await using Stream s = snapshotPayload.GetStream();
                 snapshot = ReplicationLogFactory.DeserializeSnapshotFromStream(s);
@@ -59,17 +64,20 @@ namespace Horde.Storage.Implementation.TransactionLog
             }
 
             if (cancellationToken.IsCancellationRequested)
+            {
                 throw new TaskCanceledException();
+            }
 
             await foreach (ReplicationLogEvent entry in _replicationLog.Get(ns, lastBucket, lastEvent))
             {
                 if (cancellationToken.IsCancellationRequested)
+                {
                     throw new TaskCanceledException();
+                }
 
                 snapshot.ProcessEvent(entry);
             }
 
-
             FileInfo tempFile = new FileInfo(Path.GetTempFileName());
             {
                 await using FileStream fs = tempFile.OpenWrite();
@@ -95,7 +103,9 @@ namespace Horde.Storage.Implementation.TransactionLog
                 BlobIdentifier cbBlobId = BlobIdentifier.FromBlob(cbObjectBytes);
 
                 if (cancellationToken.IsCancellationRequested)
+                {
                     throw new TaskCanceledException();
+                }
 
                 // upload the attachment first so we are not missing any references when we go to create the ref
                 await _blobService.PutObject(storeInNamespace, payload, blobIdentifier);
@@ -105,10 +115,14 @@ namespace Horde.Storage.Implementation.TransactionLog
                 List<ContentHash> missingHashes = new List<ContentHash>(missingContentIds);
                 missingHashes.AddRange(missingBlobs);
                 if (missingHashes.Count != 0)
+                {
                     throw new Exception($"Failed to upload snapshot to object service, missing references {string.Join(',' , missingHashes.Select(b => b.ToString()))}");
+                }
 
                 if (cancellationToken.IsCancellationRequested)
+                {
                     throw new TaskCanceledException();
+                }
 
                 // update the replication log with the new snapshot
                 await _replicationLog.AddSnapshot(new SnapshotInfo(ns, storeInNamespace, blobIdentifier, DateTime.Now));
diff --git a/Engine/Source/Programs/Horde/HordeStorage/Horde.Storage/Implementation/TransactionLog/ReplicationSnapshotService.cs b/Engine/Source/Programs/Horde/HordeStorage/Horde.Storage/Implementation/TransactionLog/ReplicationSnapshotService.cs
--- a/Engine/Source/Programs/Horde/HordeStorage/Horde.Storage/Implementation/TransactionLog/ReplicationSnapshotService.cs
+++ b/Engine/Source/Programs/Horde/HordeStorage/Horde.Storage/Implementation/TransactionLog/ReplicationSnapshotService.cs
@@ -1,7 +1,6 @@
 ﻿// Copyright Epic Games, Inc. All Rights Reserved.
 
 using System;
-using System.ComponentModel.DataAnnotations;
 using System.Threading;
 using System.Threading.Tasks;
 using Dasync.Collections;
@@ -16,13 +15,13 @@ using Serilog;
 
 namespace Horde.Storage.Implementation
 {
-    // ReSharper disable once ClassNeverInstantiated.Global
-    public class ReplicationSnapshotService : PollingService<ReplicationSnapshotService.SnapshotState>
+    public class SnapshotState
     {
-        public class SnapshotState
-        {
-        }
+    }
 
+    // ReSharper disable once ClassNeverInstantiated.Global
+    public class ReplicationSnapshotService : PollingService<SnapshotState>
+    {
         private readonly IServiceProvider _provider;
         private readonly IOptionsMonitor<SnapshotSettings> _settings;
         private readonly IReplicationLog _replicationLog;
@@ -36,8 +35,7 @@ namespace Horde.Storage.Implementation
             return _settings.CurrentValue.Enabled;
         }
 
-        public ReplicationSnapshotService(IServiceProvider provider, IOptionsMonitor<SnapshotSettings> settings, IReplicationLog replicationLog, ILeaderElection leaderElection) :
-            base(serviceName: nameof(ReplicationSnapshotService), TimeSpan.FromMinutes(15), new SnapshotState())
+        public ReplicationSnapshotService(IServiceProvider provider, IOptionsMonitor<SnapshotSettings> settings, IReplicationLog replicationLog, ILeaderElection leaderElection) : base(serviceName: nameof(ReplicationSnapshotService), TimeSpan.FromMinutes(15), new SnapshotState())
         {
             _provider = provider;
             _settings = settings;
@@ -98,7 +96,9 @@ namespace Horde.Storage.Implementation
             _cancellationTokenSource.Cancel();
 
             if (_snapshotBuildTask != null)
+            {
                 await _snapshotBuildTask;
+            }
         }
     }
 
diff --git a/Engine/Source/Programs/Horde/HordeStorage/Horde.Storage/Implementation/TransactionLog/ScyllaReplicationLog.cs b/Engine/Source/Programs/Horde/HordeStorage/Horde.Storage/Implementation/TransactionLog/ScyllaReplicationLog.cs
--- a/Engine/Source/Programs/Horde/HordeStorage/Horde.Storage/Implementation/TransactionLog/ScyllaReplicationLog.cs
+++ b/Engine/Source/Programs/Horde/HordeStorage/Horde.Storage/Implementation/TransactionLog/ScyllaReplicationLog.cs
@@ -124,7 +124,9 @@ namespace Horde.Storage.Implementation
                 readReplicationScope.Span.ResourceName = bucket;
 
                 if (lastBucket != null && bucket != lastBucket)
+                {
                     continue;
+                }
 
                 // at least one bucket was found
                 bucketFound = true;
@@ -144,7 +146,9 @@ namespace Horde.Storage.Implementation
                     }
 
                     if (skipEvents)
+                    {
                         continue;
+                    }
 
                     yield return new ReplicationLogEvent(
                         new NamespaceId(scyllaReplicationLog.Namespace),
@@ -218,14 +222,16 @@ namespace Horde.Storage.Implementation
 
                 bucketTime = bucketTime.AddHours(1.0);
             }
-
         }
 
-        private long FromReplicationBucketIdentifier(string bucket)
+        private static long FromReplicationBucketIdentifier(string bucket)
         {
-            if (!bucket.StartsWith("rep-"))
+            if (!bucket.StartsWith("rep-", StringComparison.OrdinalIgnoreCase))
+            {
                 throw new ArgumentException($"Invalid bucket identifier: \"{bucket}\"", nameof(bucket));
-            string timestamp = bucket.Substring(bucket.IndexOf("-") + 1);
+            }
+
+            string timestamp = bucket.Substring(bucket.IndexOf("-", StringComparison.OrdinalIgnoreCase) + 1);
             long filetime = long.Parse(timestamp);
             return filetime;
         }
@@ -349,8 +355,6 @@ namespace Horde.Storage.Implementation
         }
     }
 
-
-    
     [Cassandra.Mapping.Attributes.Table("replication_snapshot")]
     class ScyllaSnapshot
     {
@@ -370,8 +374,7 @@ namespace Horde.Storage.Implementation
         }
 
         [Cassandra.Mapping.Attributes.PartitionKey]
-        public string Namespace { get;set; }
-
+        public string Namespace { get; set; }
 
         [Cassandra.Mapping.Attributes.Column("id")]
         [Cassandra.Mapping.Attributes.ClusteringKey]
@@ -419,11 +422,9 @@ namespace Horde.Storage.Implementation
             LastEvent = lastEvent;
         }
 
-
         [Cassandra.Mapping.Attributes.PartitionKey]
         public string Namespace { get; set; }
 
-
         [Cassandra.Mapping.Attributes.ClusteringKey]
         public string Name { get; set; }
 
diff --git a/Engine/Source/Programs/Horde/HordeStorage/Horde.Storage/Implementation/VersionFile.cs b/Engine/Source/Programs/Horde/HordeStorage/Horde.Storage/Implementation/VersionFile.cs
--- a/Engine/Source/Programs/Horde/HordeStorage/Horde.Storage/Implementation/VersionFile.cs
+++ b/Engine/Source/Programs/Horde/HordeStorage/Horde.Storage/Implementation/VersionFile.cs
@@ -15,7 +15,7 @@ namespace Horde.Storage.Implementation
     {
         private readonly VersionFileContents _versionFileContents;
 
-        public class VersionFileContents
+        private class VersionFileContents
         {
             public string? Version { get; set; }
         }
@@ -32,9 +32,6 @@ namespace Horde.Storage.Implementation
             _versionFileContents = deserializer.Deserialize<VersionFileContents>(tr);
         }
 
-        public string? VersionString
-        {
-            get { return _versionFileContents.Version; }
-        }
+        public string? VersionString => _versionFileContents.Version;
     }
 }
diff --git a/Engine/Source/Programs/Horde/HordeStorage/Horde.Storage/Settings.cs b/Engine/Source/Programs/Horde/HordeStorage/Horde.Storage/Settings.cs
--- a/Engine/Source/Programs/Horde/HordeStorage/Horde.Storage/Settings.cs
+++ b/Engine/Source/Programs/Horde/HordeStorage/Horde.Storage/Settings.cs
@@ -4,7 +4,6 @@ using System;
 using System.Collections.Generic;
 using System.ComponentModel.DataAnnotations;
 using System.Linq;
-using Horde.Storage.Implementation;
 using Microsoft.Extensions.Caching.Memory;
 
 namespace Horde.Storage
@@ -83,8 +82,8 @@ namespace Horde.Storage
             Static, 
             Kubernetes
         }
-        
-        public class ValidStorageBackend : ValidationAttribute
+
+        private sealed class ValidStorageBackend : ValidationAttribute
         {
             public override string FormatErrorMessage(string name)
             {
@@ -94,12 +93,17 @@ namespace Horde.Storage
 
             public override bool IsValid(object? value)
             {
-                if (value == null) return true;
+                if (value == null)
+                {
+                    return true;
+                }
+
                 return value is IEnumerable<string> backends && backends.All(x => Enum.TryParse(typeof(StorageBackendImplementations), x, true, out _));
             }
         }
 
         [ValidStorageBackend]
+        [System.Diagnostics.CodeAnalysis.SuppressMessage("Naming", "CA1721:Property names should not match get methods", Justification = "This pattern is used to work around limitations in dotnet configurations support for enums in arrays")]
         public string[]? StorageImplementations { get; set; }
 
         public IEnumerable<StorageBackendImplementations> GetStorageImplementations()
@@ -147,7 +151,6 @@ namespace Horde.Storage
         public bool CreateDatabaseIfMissing { get; set; } = true;
     }
 
-
     public class DynamoDbSettings
     {
         [Required] public string ConnectionString { get; set; } = "";
@@ -170,7 +173,7 @@ namespace Horde.Storage
 
         public static (string, int) ParseDaxEndpointAsHostPort(string endpoint)
         {
-            if (!endpoint.Contains(":"))
+            if (!endpoint.Contains(':', StringComparison.InvariantCultureIgnoreCase))
             {
                 return (endpoint, 8111);
             }
@@ -181,10 +184,9 @@ namespace Horde.Storage
         }
     }
 
-
     public class CosmosSettings
     {
-        [Range(400, 10_000)] public int DefaultRU = 400;
+        [Range(400, 10_000)] public int DefaultRU { get; set; } = 400;
     }
 
     public class CallistoTransactionLogSettings
@@ -214,7 +216,6 @@ namespace Horde.Storage
         public int SlidingExpirationMinutes { get; set; } = 60;
     }
 
-
     public class AzureSettings
     {
         [Required] public string ConnectionString { get; set; } = string.Empty;
@@ -241,9 +242,6 @@ namespace Horde.Storage
         public bool SetBucketPolicies { get; set; } = true;
 
         public bool UseBlobIndexForExistsCheck { get; set; } = false;
-
-        public List<string> NamespacesThatUseBlobIndexForExistsCheck { get; set; } = new List<string>();
-
     }
 
     public class GCSettings
@@ -268,6 +266,7 @@ namespace Horde.Storage
 
     public class ClusterSettings
     {
+        [System.Diagnostics.CodeAnalysis.SuppressMessage("Usage", "CA2227:Collection properties should be read only", Justification = "Used by serialization")]
         public List<PeerSettings> Peers { get; set; } = new List<PeerSettings>();
     }
 
@@ -277,12 +276,13 @@ namespace Horde.Storage
 
         [Required] public string FullName { get; set; } = null!;
 
+        [System.Diagnostics.CodeAnalysis.SuppressMessage("Usage", "CA2227:Collection properties should be read only", Justification = "Used by serialization")]
         public List<PeerEndpoints> Endpoints { get; set; } = new List<PeerEndpoints>();
     }
 
     public class PeerEndpoints
     { 
-        [Required] public string Url { get; set; } = null!;
+        [Required] public Uri Url { get; set; } = null!;
 
         public bool IsInternal { get; set; } = false;
     }
diff --git a/Engine/Source/Programs/Horde/HordeStorage/Horde.Storage/Tests/Functional/Blob/BlobIndexTests.cs b/Engine/Source/Programs/Horde/HordeStorage/Horde.Storage/Tests/Functional/Blob/BlobIndexTests.cs
--- a/Engine/Source/Programs/Horde/HordeStorage/Horde.Storage/Tests/Functional/Blob/BlobIndexTests.cs
+++ b/Engine/Source/Programs/Horde/HordeStorage/Horde.Storage/Tests/Functional/Blob/BlobIndexTests.cs
@@ -31,10 +31,9 @@ namespace Horde.Storage.FunctionalTests.Storage
 {
     public abstract class BlobIndexTests
     {
-        protected TestServer? _server;
-        protected HttpClient? _httpClient;
-
-        protected readonly NamespaceId TestNamespaceName = new NamespaceId("testbucket");
+        private TestServer? _server;
+        private HttpClient? _httpClient;
+        private readonly NamespaceId _testNamespaceName = new NamespaceId("testbucket");
 
         [TestInitialize]
         public async Task Setup()
@@ -68,7 +67,6 @@ namespace Horde.Storage.FunctionalTests.Storage
         protected abstract Task Seed(IServiceProvider serverServices);
         protected abstract Task Teardown(IServiceProvider serverServices);
 
-
         [TestCleanup]
         public async Task MyTeardown()
         {
@@ -83,12 +81,12 @@ namespace Horde.Storage.FunctionalTests.Storage
             using ByteArrayContent requestContent = new ByteArrayContent(payload);
             requestContent.Headers.ContentType = new MediaTypeHeaderValue(MediaTypeNames.Application.Octet);
             BlobIdentifier contentHash = BlobIdentifier.FromBlob(payload);
-            HttpResponseMessage response = await _httpClient!.PutAsync(requestUri: $"api/v1/s/{TestNamespaceName}/{contentHash}", requestContent);
+            HttpResponseMessage response = await _httpClient!.PutAsync(new Uri($"api/v1/s/{_testNamespaceName}/{contentHash}", UriKind.Relative), requestContent);
             response.EnsureSuccessStatusCode();
 
             IBlobIndex? index = _server!.Services.GetService<IBlobIndex>();
             Assert.IsNotNull(index);
-            IBlobIndex.BlobInfo? blobInfo = await index.GetBlobInfo(TestNamespaceName, contentHash);
+            BlobInfo? blobInfo = await index.GetBlobInfo(_testNamespaceName, contentHash);
 
             Assert.IsNotNull(blobInfo);
             Assert.IsTrue(blobInfo.Regions.Contains("test"));
@@ -105,16 +103,16 @@ namespace Horde.Storage.FunctionalTests.Storage
             byte[] objectData = writer.ToByteArray();
             BlobIdentifier objectHash = BlobIdentifier.FromBlob(objectData);
 
-            HttpContent requestContent = new ByteArrayContent(objectData);
+            using HttpContent requestContent = new ByteArrayContent(objectData);
             requestContent.Headers.ContentType = new MediaTypeHeaderValue(CustomMediaTypeNames.UnrealCompactBinary);
             requestContent.Headers.Add(CommonHeaders.HashHeaderName, objectHash.ToString());
             IoHashKey putKey = IoHashKey.FromName("newReferenceUploadObject");
-            HttpResponseMessage result = await _httpClient!.PutAsync(requestUri: $"api/v1/refs/{TestNamespaceName}/bucket/{putKey}.uecb", requestContent);
+            HttpResponseMessage result = await _httpClient!.PutAsync(new Uri($"api/v1/refs/{_testNamespaceName}/bucket/{putKey}.uecb", UriKind.Relative), requestContent);
             result.EnsureSuccessStatusCode();
 
             IBlobIndex? index = _server!.Services.GetService<IBlobIndex>();
             Assert.IsNotNull(index);
-            IBlobIndex.BlobInfo? blobInfo = await index.GetBlobInfo(TestNamespaceName, objectHash);
+            BlobInfo? blobInfo = await index.GetBlobInfo(_testNamespaceName, objectHash);
 
             Assert.IsNotNull(blobInfo);
             Assert.IsTrue(blobInfo.Regions.Contains("test"));
@@ -135,30 +133,29 @@ namespace Horde.Storage.FunctionalTests.Storage
             {
                 using ByteArrayContent requestContent = new ByteArrayContent(payload);
                 requestContent.Headers.ContentType = new MediaTypeHeaderValue(MediaTypeNames.Application.Octet);
-                HttpResponseMessage response = await _httpClient!.PutAsync(requestUri: $"api/v1/s/{TestNamespaceName}/{contentHash}", requestContent);
+                HttpResponseMessage response = await _httpClient!.PutAsync(new Uri($"api/v1/s/{_testNamespaceName}/{contentHash}", UriKind.Relative), requestContent);
                 response.EnsureSuccessStatusCode();
             }
             // verify its present in the blob index
             IBlobIndex? index = _server!.Services.GetService<IBlobIndex>();
             Assert.IsNotNull(index);
-            IBlobIndex.BlobInfo? blobInfo = await index.GetBlobInfo(TestNamespaceName, contentHash);
+            BlobInfo? blobInfo = await index.GetBlobInfo(_testNamespaceName, contentHash);
             Assert.IsNotNull(blobInfo);
             Assert.IsTrue(blobInfo.Regions.Any());
 
             // delete the blob
             {
-                HttpResponseMessage response = await _httpClient!.DeleteAsync(requestUri: $"api/v1/s/{TestNamespaceName}/{contentHash}");
+                HttpResponseMessage response = await _httpClient!.DeleteAsync(new Uri($"api/v1/s/{_testNamespaceName}/{contentHash}", UriKind.Relative));
                 response.EnsureSuccessStatusCode();
             }
 
-            IBlobIndex.BlobInfo? deletedBlobInfo = await index.GetBlobInfo(TestNamespaceName, contentHash);
+            BlobInfo? deletedBlobInfo = await index.GetBlobInfo(_testNamespaceName, contentHash);
 
             bool hasRegions = deletedBlobInfo != null && deletedBlobInfo.Regions.Any();
             // but the blob info will not contain the current region
             Assert.IsTrue(!hasRegions);
         }
 
-
         [TestMethod]
         public async Task EnumerateAllBlobs()
         {
@@ -166,7 +163,7 @@ namespace Horde.Storage.FunctionalTests.Storage
             Assert.IsNotNull(index);
             {
                 // verify the blob info list is empty at the start
-                IBlobIndex.BlobInfo[]? blobInfos =  await index.GetAllBlobs().ToArrayAsync();
+                BlobInfo[]? blobInfos =  await index.GetAllBlobs().ToArrayAsync();
                 Assert.AreEqual(0, blobInfos.Length);
             }
 
@@ -179,7 +176,7 @@ namespace Horde.Storage.FunctionalTests.Storage
                 {
                     using ByteArrayContent requestContent = new ByteArrayContent(payload);
                     requestContent.Headers.ContentType = new MediaTypeHeaderValue(MediaTypeNames.Application.Octet);
-                    HttpResponseMessage response = await _httpClient!.PutAsync(requestUri: $"api/v1/blobs/{TestNamespaceName}/{contentHash}", requestContent);
+                    HttpResponseMessage response = await _httpClient!.PutAsync(new Uri($"api/v1/blobs/{_testNamespaceName}/{contentHash}", UriKind.Relative), requestContent);
                     response.EnsureSuccessStatusCode();
                 }
             }
@@ -191,22 +188,20 @@ namespace Horde.Storage.FunctionalTests.Storage
                 compressedPayloadIdentifier = BlobIdentifier.FromBlob(texturePayload);
                 BlobIdentifier uncompressedPayloadIdentifier = new BlobIdentifier("DEA81B6C3B565BB5089695377C98CE0F1C13B0C3");
 
-                ByteArrayContent content = new ByteArrayContent(texturePayload);
+                using ByteArrayContent content = new ByteArrayContent(texturePayload);
                 content.Headers.ContentType = new MediaTypeHeaderValue(CustomMediaTypeNames.UnrealCompressedBuffer);
-                HttpResponseMessage result = await _httpClient!.PutAsync($"api/v1/compressed-blobs/{TestNamespaceName}/{uncompressedPayloadIdentifier}", content);
+                HttpResponseMessage result = await _httpClient!.PutAsync(new Uri($"api/v1/compressed-blobs/{_testNamespaceName}/{uncompressedPayloadIdentifier}", UriKind.Relative), content);
                 result.EnsureSuccessStatusCode();
             }
 
             {
-                IBlobIndex.BlobInfo[]? blobInfos =  await index.GetAllBlobs().ToArrayAsync();
+                BlobInfo[]? blobInfos =  await index.GetAllBlobs().ToArrayAsync();
                 Assert.AreEqual(2, blobInfos.Length);
 
                 Assert.IsNotNull(blobInfos.FirstOrDefault(info => info.BlobIdentifier.Equals(compressedPayloadIdentifier)));
                 Assert.IsNotNull(blobInfos.FirstOrDefault(info => info.BlobIdentifier.Equals(contentHash)));
             }
         }
-
-     
     }
 
     [TestClass()]
diff --git a/Engine/Source/Programs/Horde/HordeStorage/Horde.Storage/Tests/Functional/Blob/BlobReplicationTests.cs b/Engine/Source/Programs/Horde/HordeStorage/Horde.Storage/Tests/Functional/Blob/BlobReplicationTests.cs
--- a/Engine/Source/Programs/Horde/HordeStorage/Horde.Storage/Tests/Functional/Blob/BlobReplicationTests.cs
+++ b/Engine/Source/Programs/Horde/HordeStorage/Horde.Storage/Tests/Functional/Blob/BlobReplicationTests.cs
@@ -1,5 +1,6 @@
 // Copyright Epic Games, Inc. All Rights Reserved.
 
+using System;
 using System.Collections.Generic;
 using System.Linq;
 using System.Net;
@@ -24,10 +25,9 @@ namespace Horde.Storage.FunctionalTests.Storage
 {
     [TestClass]
     public class BlobReplicationTests
-    {
-        protected readonly NamespaceId TestNamespaceName = new NamespaceId("test-namespace");
+    { 
+        protected NamespaceId TestNamespaceName { get; } = new NamespaceId("test-namespace");
 
-        
         [TestMethod]
         public async Task ReplicateBlobFromRegion()
         {
@@ -60,7 +60,7 @@ namespace Horde.Storage.FunctionalTests.Storage
                 .ReadFrom.Configuration(configuration)
                 .CreateLogger();
 
-            TestServer server = new TestServer(new WebHostBuilder()
+            using TestServer server = new TestServer(new WebHostBuilder()
                 .UseConfiguration(configuration)
                 .UseEnvironment("Testing")
                 .UseSerilog(logger)
@@ -78,12 +78,12 @@ namespace Horde.Storage.FunctionalTests.Storage
                                 {
                                     new PeerEndpoints()
                                     {
-                                        Url = "http://siteA.com/internal",
+                                        Url = new Uri("http://siteA.com/internal"),
                                         IsInternal = true
                                     },
                                     new PeerEndpoints()
                                     {
-                                        Url = "http://siteA.com/public"
+                                        Url = new Uri("http://siteA.com/public")
                                     },
                                 }.ToList()
                             },
@@ -95,12 +95,12 @@ namespace Horde.Storage.FunctionalTests.Storage
                                 {
                                     new PeerEndpoints()
                                     {
-                                        Url = "http://siteB.com/internal",
+                                        Url = new Uri("http://siteB.com/internal"),
                                         IsInternal = true
                                     },
                                     new PeerEndpoints()
                                     {
-                                        Url = "http://siteB.com/public"
+                                        Url = new Uri("http://siteB.com/public")
                                     },
                                 }.ToList()
                             },
@@ -109,10 +109,10 @@ namespace Horde.Storage.FunctionalTests.Storage
                     });
                     collection.Configure<NamespaceSettings>(settings =>
                     {
-                        settings.Policies = new Dictionary<string, NamespaceSettings.PerNamespaceSettings>()
+                        settings.Policies = new Dictionary<string, NamespacePolicy>()
                         {
                             {
-                                TestNamespaceName.ToString(), new NamespaceSettings.PerNamespaceSettings()
+                                TestNamespaceName.ToString(), new NamespacePolicy()
                                 {
                                     OnDemandReplication = true
                                 }
@@ -131,7 +131,8 @@ namespace Horde.Storage.FunctionalTests.Storage
             await blobIndex.AddBlobToIndex(TestNamespaceName, blobIdentifier, "siteB");
 
             HttpClient httpClient = server.CreateClient();
-            HttpResponseMessage response = await httpClient!.GetAsync(requestUri: $"api/v1/blobs/{TestNamespaceName}/{blobIdentifier}");
+
+            HttpResponseMessage response = await httpClient!.GetAsync(new Uri($"api/v1/blobs/{TestNamespaceName}/{blobIdentifier}", UriKind.Relative));
             response.EnsureSuccessStatusCode();
 
             byte[] responseData = await response.Content.ReadAsByteArrayAsync();
@@ -163,7 +164,7 @@ namespace Horde.Storage.FunctionalTests.Storage
                 .ReadFrom.Configuration(configuration)
                 .CreateLogger();
 
-            TestServer server = new TestServer(new WebHostBuilder()
+            using TestServer server = new TestServer(new WebHostBuilder()
                 .UseConfiguration(configuration)
                 .UseEnvironment("Testing")
                 .UseSerilog(logger)
@@ -181,12 +182,12 @@ namespace Horde.Storage.FunctionalTests.Storage
                                 {
                                     new PeerEndpoints()
                                     {
-                                        Url = "http://siteA.com/internal",
+                                        Url = new Uri("http://siteA.com/internal"),
                                         IsInternal = true
                                     },
                                     new PeerEndpoints()
                                     {
-                                        Url = "http://siteA.com/public"
+                                        Url = new Uri("http://siteA.com/public")
                                     },
                                 }.ToList()
                             },
@@ -198,12 +199,12 @@ namespace Horde.Storage.FunctionalTests.Storage
                                 {
                                     new PeerEndpoints()
                                     {
-                                        Url = "http://siteB.com/internal",
+                                        Url = new Uri("http://siteB.com/internal"),
                                         IsInternal = true
                                     },
                                     new PeerEndpoints()
                                     {
-                                        Url = "http://siteB.com/public"
+                                        Url = new Uri("http://siteB.com/public")
                                     },
                                 }.ToList()
                             },
@@ -216,7 +217,7 @@ namespace Horde.Storage.FunctionalTests.Storage
             );
 
             HttpClient httpClient = server.CreateClient();
-            HttpResponseMessage response = await httpClient!.GetAsync(requestUri: $"api/v1/blobs/{TestNamespaceName}/{blobIdentifier}");
+            HttpResponseMessage response = await httpClient!.GetAsync(new Uri($"api/v1/blobs/{TestNamespaceName}/{blobIdentifier}", UriKind.Relative));
             Assert.IsTrue(response.StatusCode == HttpStatusCode.NotFound);
         }
     }
diff --git a/Engine/Source/Programs/Horde/HordeStorage/Horde.Storage/Tests/Functional/ContentId/CompressedBlobTests.cs b/Engine/Source/Programs/Horde/HordeStorage/Horde.Storage/Tests/Functional/ContentId/CompressedBlobTests.cs
--- a/Engine/Source/Programs/Horde/HordeStorage/Horde.Storage/Tests/Functional/ContentId/CompressedBlobTests.cs
+++ b/Engine/Source/Programs/Horde/HordeStorage/Horde.Storage/Tests/Functional/ContentId/CompressedBlobTests.cs
@@ -10,13 +10,11 @@ using System.Net.Mime;
 using System.Text;
 using System.Threading.Tasks;
 using Horde.Storage.FunctionalTests.Storage;
-using Horde.Storage.Implementation;
 using Jupiter;
 using Jupiter.Implementation;
 using Microsoft.AspNetCore.TestHost;
 using Microsoft.AspNetCore.Hosting;
 using Microsoft.Extensions.Configuration;
-using Microsoft.Extensions.DependencyInjection;
 using Microsoft.VisualStudio.TestTools.UnitTesting;
 using Serilog;
 using Serilog.Core;
@@ -73,12 +71,10 @@ namespace Horde.Storage.FunctionalTests.CompressedBlobs
         }
     }
 
-
-
     public abstract class CompressedBlobTests
     {
-        protected TestServer? _server;
-        protected HttpClient? _httpClient;
+        private TestServer? Server { get; set; }
+        private HttpClient? Client { get; set; }
 
         protected const string TestNamespace = "test-namespace";
 
@@ -102,26 +98,23 @@ namespace Horde.Storage.FunctionalTests.CompressedBlobs
                 .UseSerilog(logger)
                 .UseStartup<HordeStorageStartup>()
             );
-            _httpClient = server.CreateClient();
-            _server = server;
+            Client = server.CreateClient();
+            Server = server;
 
             // Seed storage
-            await Seed(_server.Services);
+            await Seed(Server.Services);
         }
 
         protected abstract IEnumerable<KeyValuePair<string, string>> GetSettings();
-
         protected abstract Task Seed(IServiceProvider serverServices);
         protected abstract Task Teardown();
 
-
         [TestCleanup]
         public async Task MyTeardown()
         {
             await Teardown();
         }
 
-
         [TestMethod]
         [DataTestMethod]
         [DataRow("7835c353d7dc67e8a0531c88fbc75ddfda10dee4", "7835c353d7dc67e8a0531c88fbc75ddfda10dee4")]
@@ -139,9 +132,9 @@ namespace Horde.Storage.FunctionalTests.CompressedBlobs
             BlobIdentifier compressedPayloadIdentifier = BlobIdentifier.FromBlob(texturePayload);
             BlobIdentifier uncompressedPayloadIdentifier = new BlobIdentifier(uncompressedHash);
 
-            ByteArrayContent content = new ByteArrayContent(texturePayload);
+            using ByteArrayContent content = new(texturePayload);
             content.Headers.ContentType = new MediaTypeHeaderValue(CustomMediaTypeNames.UnrealCompressedBuffer);
-            HttpResponseMessage result = await _httpClient!.PutAsync($"api/v1/compressed-blobs/{TestNamespace}/{uncompressedPayloadIdentifier}", content);
+            HttpResponseMessage result = await Client!.PutAsync(new Uri($"api/v1/compressed-blobs/{TestNamespace}/{uncompressedPayloadIdentifier}", UriKind.Relative), content);
             result.EnsureSuccessStatusCode();
 
             InsertResponse response = await result.Content.ReadAsAsync<InsertResponse>();
@@ -157,9 +150,9 @@ namespace Horde.Storage.FunctionalTests.CompressedBlobs
             ContentId uncompressedPayloadIdentifier = new ContentId("DEA81B6C3B565BB5089695377C98CE0F1C13B0C3");
 
             {
-                ByteArrayContent content = new ByteArrayContent(texturePayload);
+                using ByteArrayContent content = new(texturePayload);
                 content.Headers.ContentType = new MediaTypeHeaderValue(CustomMediaTypeNames.UnrealCompressedBuffer);
-                HttpResponseMessage result = await _httpClient!.PutAsync($"api/v1/compressed-blobs/{TestNamespace}/{uncompressedPayloadIdentifier}", content);
+                HttpResponseMessage result = await Client!.PutAsync(new Uri($"api/v1/compressed-blobs/{TestNamespace}/{uncompressedPayloadIdentifier}", UriKind.Relative), content);
                 result.EnsureSuccessStatusCode();
 
                 InsertResponse response = await result.Content.ReadAsAsync<InsertResponse>();
@@ -168,9 +161,8 @@ namespace Horde.Storage.FunctionalTests.CompressedBlobs
                 Assert.AreEqual(uncompressedPayloadIdentifier, ContentId.FromBlobIdentifier(response.Identifier));
             }
 
-
             {
-                HttpResponseMessage result = await _httpClient!.GetAsync($"api/v1/compressed-blobs/{TestNamespace}/{uncompressedPayloadIdentifier}");
+                HttpResponseMessage result = await Client!.GetAsync(new Uri($"api/v1/compressed-blobs/{TestNamespace}/{uncompressedPayloadIdentifier}", UriKind.Relative));
                 result.EnsureSuccessStatusCode();
                 Assert.AreEqual(CustomMediaTypeNames.UnrealCompressedBuffer, result.Content.Headers.ContentType!.MediaType);
 
@@ -180,7 +172,7 @@ namespace Horde.Storage.FunctionalTests.CompressedBlobs
 
             {
                 // verify the compressed blob can be retrieved in the blob store
-                HttpResponseMessage result = await _httpClient!.GetAsync($"api/v1/blobs/{TestNamespace}/{compressedPayloadIdentifier}");
+                HttpResponseMessage result = await Client!.GetAsync(new Uri($"api/v1/blobs/{TestNamespace}/{compressedPayloadIdentifier}", UriKind.Relative));
                 result.EnsureSuccessStatusCode();
                 Assert.AreEqual(MediaTypeNames.Application.Octet, result.Content.Headers.ContentType!.MediaType);
 
@@ -190,7 +182,7 @@ namespace Horde.Storage.FunctionalTests.CompressedBlobs
 
             {
                 // the uncompressed payload should not be valid in the blob endpoint
-                HttpResponseMessage result = await _httpClient!.GetAsync($"api/v1/blobs/{TestNamespace}/{uncompressedPayloadIdentifier}");
+                HttpResponseMessage result = await Client!.GetAsync(new Uri($"api/v1/blobs/{TestNamespace}/{uncompressedPayloadIdentifier}", UriKind.Relative));
                 Assert.AreEqual(HttpStatusCode.NotFound, result.StatusCode);
             }
         }
@@ -200,19 +192,17 @@ namespace Horde.Storage.FunctionalTests.CompressedBlobs
         {
             byte[] texturePayload = await File.ReadAllBytesAsync("ContentId/Payloads/UncompressedTexture_CAS_dea81b6c3b565bb5089695377c98ce0f1c13b0c3.udd");
             BlobIdentifier compressedPayloadIdentifier = BlobIdentifier.FromBlob(texturePayload);
-            BlobIdentifier uncompressedPayloadIdentifier = new BlobIdentifier("DEA81B6C3B565BB5089695377C98CE0F1C13B0C3");
 
             {
-                ByteArrayContent content = new ByteArrayContent(texturePayload);
+                using ByteArrayContent content = new(texturePayload);
                 content.Headers.ContentType = new MediaTypeHeaderValue(CustomMediaTypeNames.UnrealCompressedBuffer);
                 // we purposefully use the compressed identifier here which is not what is expected
-                HttpResponseMessage result = await _httpClient!.PutAsync($"api/v1/compressed-blobs/{TestNamespace}/{compressedPayloadIdentifier}", content);
+                HttpResponseMessage result = await Client!.PutAsync(new Uri($"api/v1/compressed-blobs/{TestNamespace}/{compressedPayloadIdentifier}", UriKind.Relative), content);
 
                 Assert.AreEqual(HttpStatusCode.BadRequest, result.StatusCode);
             }
         }
 
-
         [TestMethod]
         public async Task GetUncompressedContent()
         {
@@ -222,15 +212,14 @@ namespace Horde.Storage.FunctionalTests.CompressedBlobs
 
             // upload a uncompressed blob
             {
-                ByteArrayContent content = new ByteArrayContent(payload);
+                using ByteArrayContent content = new(payload);
                 content.Headers.ContentType = new MediaTypeHeaderValue(MediaTypeNames.Application.Octet);
-                HttpResponseMessage result = await _httpClient!.PutAsync($"api/v1/blobs/{TestNamespace}/{blobIdentifier}", content);
+                HttpResponseMessage result = await Client!.PutAsync(new Uri($"api/v1/blobs/{TestNamespace}/{blobIdentifier}", UriKind.Relative), content);
                 result.EnsureSuccessStatusCode();
             }
 
-
             {
-                HttpResponseMessage result = await _httpClient!.GetAsync($"api/v1/compressed-blobs/{TestNamespace}/{blobIdentifier}");
+                HttpResponseMessage result = await Client!.GetAsync(new Uri($"api/v1/compressed-blobs/{TestNamespace}/{blobIdentifier}", UriKind.Relative));
                 result.EnsureSuccessStatusCode();
 
                 // verify that we return the uncompressed content but also that the media type indicates that
@@ -250,28 +239,27 @@ namespace Horde.Storage.FunctionalTests.CompressedBlobs
 
             // upload a uncompressed blob
             {
-                ByteArrayContent content = new ByteArrayContent(payload);
+                using ByteArrayContent content = new(payload);
                 content.Headers.ContentType = new MediaTypeHeaderValue(MediaTypeNames.Application.Octet);
-                HttpResponseMessage result = await _httpClient!.PutAsync($"api/v1/blobs/{TestNamespace}/{blobIdentifier}", content);
+                HttpResponseMessage result = await Client!.PutAsync(new Uri($"api/v1/blobs/{TestNamespace}/{blobIdentifier}", UriKind.Relative), content);
                 result.EnsureSuccessStatusCode();
             }
 
-
             {
-                HttpRequestMessage request = new HttpRequestMessage(HttpMethod.Get, $"api/v1/compressed-blobs/{TestNamespace}/{blobIdentifier}");
+                using HttpRequestMessage request = new HttpRequestMessage(HttpMethod.Get, $"api/v1/compressed-blobs/{TestNamespace}/{blobIdentifier}");
                 request.Headers.Add("Accept", CustomMediaTypeNames.UnrealCompactBinary);
 
                 // asking for a compressed buffer when we only have the uncompressed content results in a 415 (unsupported media type) as transcoding this is to complicated when the compressed buffer header requires the full blake3 hash of the content
-                HttpResponseMessage result = await _httpClient!.SendAsync(request);
+                HttpResponseMessage result = await Client!.SendAsync(request);
                 Assert.AreEqual(HttpStatusCode.UnsupportedMediaType, result.StatusCode);
             }
 
             {
                 // ask for any accept type
-                HttpRequestMessage request = new HttpRequestMessage(HttpMethod.Get, $"api/v1/compressed-blobs/{TestNamespace}/{blobIdentifier}");
+                using HttpRequestMessage request = new HttpRequestMessage(HttpMethod.Get, $"api/v1/compressed-blobs/{TestNamespace}/{blobIdentifier}");
                 request.Headers.Add("Accept", "*/*");
 
-                HttpResponseMessage result = await _httpClient!.SendAsync(request);
+                HttpResponseMessage result = await Client!.SendAsync(request);
                 result.EnsureSuccessStatusCode();
                 Assert.AreEqual(MediaTypeNames.Application.Octet, result.Content.Headers.ContentType!.MediaType);
 
diff --git a/Engine/Source/Programs/Horde/HordeStorage/Horde.Storage/Tests/Functional/GC/GCBlobTests.cs b/Engine/Source/Programs/Horde/HordeStorage/Horde.Storage/Tests/Functional/GC/GCBlobTests.cs
--- a/Engine/Source/Programs/Horde/HordeStorage/Horde.Storage/Tests/Functional/GC/GCBlobTests.cs
+++ b/Engine/Source/Programs/Horde/HordeStorage/Horde.Storage/Tests/Functional/GC/GCBlobTests.cs
@@ -4,7 +4,6 @@ using System;
 using System.Collections.Generic;
 using System.Linq;
 using System.Net;
-using System.Net.Http;
 using System.Threading;
 using System.Threading.Tasks;
 using Dasync.Collections;
@@ -31,9 +30,6 @@ namespace Horde.Storage.FunctionalTests.GC
     [TestClass]
     public class GCBlobTests
     {
-        private TestServer? _server;
-        private HttpClient? _httpClient;
-
         private readonly BlobIdentifier object0id = new BlobIdentifier("0000000000000000000000000000000000000000");
         private readonly BlobIdentifier object1id = new BlobIdentifier("1111111111111111111111111111111111111111");
         private readonly BlobIdentifier object2id = new BlobIdentifier("2222222222222222222222222222222222222222");
@@ -60,8 +56,7 @@ namespace Horde.Storage.FunctionalTests.GC
                 .ReadFrom.Configuration(configuration)
                 .CreateLogger();
 
-
-            _callistoBlobMock = GetCallistoMock(new CallistoReader.CallistoGetResponse(new TransactionEvent[]
+            _callistoBlobMock = GetCallistoMock(new CallistoGetResponse(new TransactionEvent[]
             {
                 new AddTransactionEvent("object0", "default", new[] {object0id}, identifier: 0, nextIdentifier: 1),
                 new AddTransactionEvent("object2", "default", new[] {object2id}, identifier: 1, nextIdentifier: 2),
@@ -70,7 +65,7 @@ namespace Horde.Storage.FunctionalTests.GC
                 new AddTransactionEvent("object6", "default", new[] {object6id}, identifier: 4, nextIdentifier: 5)
             }.ToList(), Guid.Empty, 10));
 
-            TestServer server = new TestServer(new WebHostBuilder()
+           using  TestServer server = new TestServer(new WebHostBuilder()
                 .UseConfiguration(configuration)
                 .UseEnvironment("Testing")
                 .UseSerilog(logger)
@@ -88,21 +83,15 @@ namespace Horde.Storage.FunctionalTests.GC
 
                         NamespacePolicyResolver policyResolver = Mock.Of<NamespacePolicyResolver>();
 
-
                         IBlobService blobService = provider.GetService<IBlobService>()!;
-                        return new OrphanBlobCleanup(blobService!, new LeaderElectionStub(true), _callistoBlobMock.Object, gcSettingsMon, policyResolver);
+                        return new OrphanBlobCleanup(blobService!, new LeaderElectionStub(true), _callistoBlobMock.Object, policyResolver);
                     });
-
                 })
             );
-
-            _httpClient = server.CreateClient();
-            _server = server;
-
             _blobService = server.Services.GetService<IBlobService>()!;
 
             MemoryBlobStore memoryBlobStore = (MemoryBlobStore) ((BlobService)_blobService).BlobStore.First();
-            byte[] emptyContents = new byte[0];
+            byte[] emptyContents = Array.Empty<byte>();
             await memoryBlobStore.PutObject(TestNamespace, emptyContents, object0id);
             await memoryBlobStore.PutObject(TestNamespace, emptyContents, object1id);// this is not in callisto
             await memoryBlobStore.PutObject(TestNamespace, emptyContents, object2id);
@@ -121,7 +110,7 @@ namespace Horde.Storage.FunctionalTests.GC
             memoryBlobStore.SetLastModifiedTime(TestNamespace, object6id, DateTime.Now.AddDays(-2));
         }
 
-        private Mock<IRestClient> GetCallistoMock(CallistoReader.CallistoGetResponse responseData)
+        private Mock<IRestClient> GetCallistoMock(CallistoGetResponse responseData)
         {
             // the callisto reader will do 2 requests to fetch the events (one to determine the generation and the other to actually fetch the objects)
             // after that we insert a empty response to indicate the end of the log was reached
@@ -139,27 +128,27 @@ namespace Horde.Storage.FunctionalTests.GC
                     It.IsAny<CancellationToken>())).ReturnsAsync(listNsResponse.Object).Verifiable();
             
             
-            Mock<IRestResponse<CallistoReader.CallistoGetResponse>> response = new Mock<IRestResponse<CallistoReader.CallistoGetResponse>>();
+            Mock<IRestResponse<CallistoGetResponse>> response = new Mock<IRestResponse<CallistoGetResponse>>();
             response.Setup(_ => _.StatusCode).Returns(HttpStatusCode.OK);
             response.Setup(_ => _.IsSuccessful).Returns(true);
             response.Setup(_ => _.Data).Returns(responseData);
             response.Setup(_ => _.ResponseUri).Returns(new Uri("http://localhost"));
             
             mock.Setup(x =>
-                x.ExecuteGetAsync<CallistoReader.CallistoGetResponse>(It.IsAny<IRestRequest>(),
+                x.ExecuteGetAsync<CallistoGetResponse>(It.IsAny<IRestRequest>(),
                     It.IsAny<CancellationToken>())).ReturnsAsync(response.Object).Verifiable();
 
             // when we do the second call to callisto with a offset we return a empty object indicating we have reached the end
-            Mock<IRestResponse<CallistoReader.CallistoGetResponse>> responseSecondPage = new Mock<IRestResponse<CallistoReader.CallistoGetResponse>>();
+            Mock<IRestResponse<CallistoGetResponse>> responseSecondPage = new Mock<IRestResponse<CallistoGetResponse>>();
             responseSecondPage.Setup(_ => _.StatusCode).Returns(HttpStatusCode.OK);
             responseSecondPage.Setup(_ => _.IsSuccessful).Returns(true);
-            responseSecondPage.Setup(_ => _.Data).Returns(new CallistoReader.CallistoGetResponse(
+            responseSecondPage.Setup(_ => _.Data).Returns(new CallistoGetResponse(
                 new List<TransactionEvent>(),
                 responseData.Generation,
                 responseData.CurrentOffset + 10 // set the offset to some arbitrary higher number to indicate we are at the end of the log
             ));
             responseSecondPage.Setup(_ => _.ResponseUri).Returns(new Uri("http://localhost"));
-            mock.Setup(x => x.ExecuteGetAsync<CallistoReader.CallistoGetResponse>(
+            mock.Setup(x => x.ExecuteGetAsync<CallistoGetResponse>(
                     It.Is<IRestRequest>(request =>
                         request.Parameters.Any(parameter =>
                             parameter.Name == "offset" && parameter.Value!.ToString() != "0")),
@@ -180,16 +169,16 @@ namespace Horde.Storage.FunctionalTests.GC
             IOptionsMonitor<GCSettings> gcSettingsMon = Mock.Of<IOptionsMonitor<GCSettings>>(_ => _.CurrentValue == gcSettings);
 
             Mock<INamespacePolicyResolver> policyResolverMock = new Mock<INamespacePolicyResolver>();
-            policyResolverMock.Setup(resolver => resolver.GetPoliciesForNs(TestNamespace)).Returns(new NamespaceSettings.PerNamespaceSettings()
+            policyResolverMock.Setup(resolver => resolver.GetPoliciesForNs(TestNamespace)).Returns(new NamespacePolicy()
             {
                 IsLegacyNamespace = true,
             });
 
-            OrphanBlobCleanup cleanup = new OrphanBlobCleanup(_blobService!, new LeaderElectionStub(true), _callistoBlobMock.Object, gcSettingsMon, policyResolverMock.Object);
+            OrphanBlobCleanup cleanup = new OrphanBlobCleanup(_blobService!, new LeaderElectionStub(true), _callistoBlobMock.Object, policyResolverMock.Object);
             List<NamespaceId> namespaces = await cleanup.ListNamespaces().ToListAsync();
             Assert.AreEqual(1, namespaces.Count);
 
-            CancellationTokenSource cts = new CancellationTokenSource();
+            using CancellationTokenSource cts = new CancellationTokenSource();
             ulong removedBlobsCount = await cleanup.Cleanup(cts.Token);
             Assert.AreEqual(4ul, removedBlobsCount);
             
@@ -219,7 +208,6 @@ namespace Horde.Storage.FunctionalTests.GC
     public abstract class GCBlobTestsRefs
     {
         private TestServer? _server;
-        private HttpClient? _httpClient;
 
         private readonly BlobIdentifier object0id = new BlobIdentifier("0000000000000000000000000000000000000000");
         private readonly BlobIdentifier object1id = new BlobIdentifier("1111111111111111111111111111111111111111");
@@ -261,10 +249,10 @@ namespace Horde.Storage.FunctionalTests.GC
                 {
                     collection.Configure<NamespaceSettings>(settings =>
                     {
-                        settings.Policies = new Dictionary<string, NamespaceSettings.PerNamespaceSettings>()
+                        settings.Policies = new Dictionary<string, NamespacePolicy>()
                         {
                             {
-                                TestNamespace.ToString(), new NamespaceSettings.PerNamespaceSettings()
+                                TestNamespace.ToString(), new NamespacePolicy()
                                 {
                                     IsLegacyNamespace = false
                                 }
@@ -275,13 +263,12 @@ namespace Horde.Storage.FunctionalTests.GC
                 .UseStartup<HordeStorageStartup>()
             );
 
-            _httpClient = server.CreateClient();
             _server = server;
 
             _blobService = server.Services.GetService<IBlobService>()!;
 
             MemoryBlobStore memoryBlobStore = (MemoryBlobStore) ((BlobService)_blobService).BlobStore.First();
-            byte[] emptyContents = new byte[0];
+            byte[] emptyContents = Array.Empty<byte>();
             await memoryBlobStore.PutObject(TestNamespace, emptyContents, object0id);
             await memoryBlobStore.PutObject(TestNamespace, emptyContents, object1id);// this is not in the index
             await memoryBlobStore.PutObject(TestNamespace, emptyContents, object2id);
@@ -330,13 +317,11 @@ namespace Horde.Storage.FunctionalTests.GC
             await blobIndex.AddRefToBlobs(TestNamespace, testBucket, IoHashKey.FromName("object3"), new [] {object3id });
             await blobIndex.AddBlobToIndex(TestNamespace, object6id);
             await blobIndex.AddRefToBlobs(TestNamespace, testBucket, IoHashKey.FromName("object6"), new [] {object6id });
-
-
         }
 
         protected abstract string GetImplementation();
 
-        private (BlobIdentifier, CbObject) GetCBWithAttachment(BlobIdentifier blobIdentifier)
+        private static (BlobIdentifier, CbObject) GetCBWithAttachment(BlobIdentifier blobIdentifier)
         {
             CbWriter writer = new CbWriter();
             writer.BeginObject();
@@ -353,7 +338,7 @@ namespace Horde.Storage.FunctionalTests.GC
             OrphanBlobCleanupRefs? cleanup = _server!.Services.GetService<OrphanBlobCleanupRefs>();
             Assert.IsNotNull(cleanup);
 
-            CancellationTokenSource cts = new CancellationTokenSource();
+            using CancellationTokenSource cts = new();
             ulong countOfRemovedBlobs = await cleanup.Cleanup(cts.Token);
             Assert.AreEqual(3u, countOfRemovedBlobs);
 
diff --git a/Engine/Source/Programs/Horde/HordeStorage/Horde.Storage/Tests/Functional/GC/GCRefTests.cs b/Engine/Source/Programs/Horde/HordeStorage/Horde.Storage/Tests/Functional/GC/GCRefTests.cs
--- a/Engine/Source/Programs/Horde/HordeStorage/Horde.Storage/Tests/Functional/GC/GCRefTests.cs
+++ b/Engine/Source/Programs/Horde/HordeStorage/Horde.Storage/Tests/Functional/GC/GCRefTests.cs
@@ -2,13 +2,11 @@
 
 using System;
 using System.Collections.Generic;
-using System.Linq;
 using System.Net.Http;
 using System.Threading.Tasks;
 using Dasync.Collections;
 using Horde.Storage.Controllers;
 using Horde.Storage.Implementation;
-using Horde.Storage.Implementation.Blob;
 using Jupiter.Implementation;
 using Microsoft.AspNetCore.TestHost;
 using Microsoft.AspNetCore.Hosting;
@@ -26,14 +24,14 @@ namespace Horde.Storage.FunctionalTests.GC
 {
     [TestClass]
 
-    public class GCRefTests
+    public class GCRefTests : IDisposable
     {
-        private TestServer? _server;
         private HttpClient? _httpClient;
         private Mock<IRefsStore>? _refMock;
 
         private readonly NamespaceId TestNamespace = new NamespaceId("test-namespace");
         private readonly BucketId DefaultBucket = new BucketId("default");
+        private TestServer? _server;
 
         [TestInitialize]
         public void Setup()
@@ -58,7 +56,7 @@ namespace Horde.Storage.FunctionalTests.GC
             _refMock.Setup(store => store.GetOldRecords(TestNamespace, It.IsAny<TimeSpan>())).Returns(oldRecords.ToAsyncEnumerable()).Verifiable();
             _refMock.Setup(store => store.Delete(TestNamespace, DefaultBucket, It.IsAny<KeyId>())).Verifiable();
 
-            TestServer server = new TestServer(new WebHostBuilder()
+           _server = new TestServer(new WebHostBuilder()
                 .UseConfiguration(configuration)
                 .UseEnvironment("Testing")
                 .UseSerilog(logger)
@@ -70,10 +68,10 @@ namespace Horde.Storage.FunctionalTests.GC
 
                     collection.Configure<NamespaceSettings>(settings =>
                     {
-                        settings.Policies = new Dictionary<string, NamespaceSettings.PerNamespaceSettings>()
+                        settings.Policies = new Dictionary<string, NamespacePolicy>()
                         {
                             {
-                                TestNamespace.ToString(), new NamespaceSettings.PerNamespaceSettings()
+                                TestNamespace.ToString(), new NamespacePolicy()
                                 {
                                     IsLegacyNamespace = true
                                 }
@@ -82,33 +80,39 @@ namespace Horde.Storage.FunctionalTests.GC
                     });
                 })
             );
-            _httpClient = server.CreateClient();
-            _server = server;
-            
+            _httpClient = _server.CreateClient();
         }
 
-
         [TestMethod]
         public async Task RunRefCleanup()
         {
             // trigger the cleanup
-            var cleanupResponse = await _httpClient!.PostAsync($"api/v1/admin/refCleanup/{TestNamespace}", new StringContent(string.Empty));
+            using StringContent content = new StringContent(string.Empty);
+            HttpResponseMessage cleanupResponse = await _httpClient!.PostAsync(new Uri($"api/v1/admin/refCleanup/{TestNamespace}", UriKind.Relative), content);
             cleanupResponse.EnsureSuccessStatusCode();
             RemovedRefRecordsResponse removedRefRecords = await cleanupResponse.Content.ReadAsAsync<RemovedRefRecordsResponse>();
 
-            KeyId[] expectedRemovedRefRecords = new[]
-            {
-                new KeyId("object2"),
-                new KeyId("object5"),
-                new KeyId("object6")
-            };
             Assert.AreEqual(3, removedRefRecords.CountOfRemovedRecords);
 
             _refMock!.Verify();
             _refMock.VerifyNoOtherCalls();
         }
-    }
 
+        protected virtual void Dispose(bool disposing)
+        {
+            if (disposing)
+            {
+                _httpClient?.Dispose();
+                _server?.Dispose();
+            }
+        }
+
+        public void Dispose()
+        {
+            Dispose(true);
+            System.GC.SuppressFinalize(this);
+        }
+    }
 
     [TestClass]
     public class MemoryGCReferencesTests : GCReferencesTests
@@ -128,10 +132,8 @@ namespace Horde.Storage.FunctionalTests.GC
         }
     }
 
-
-    public abstract class GCReferencesTests
+    public abstract class GCReferencesTests : IDisposable
     {
-        private TestServer? _server;
         private HttpClient? _httpClient;
 
         private readonly NamespaceId TestNamespace = new NamespaceId("test-namespace");
@@ -152,6 +154,7 @@ namespace Horde.Storage.FunctionalTests.GC
         private readonly IoHashKey object4Name = IoHashKey.FromName("object4");
         private readonly IoHashKey object5Name = IoHashKey.FromName("object5");
         private readonly IoHashKey object6Name = IoHashKey.FromName("object6");
+        private TestServer? _server;
 
         [TestInitialize]
         public async Task Setup()
@@ -170,7 +173,7 @@ namespace Horde.Storage.FunctionalTests.GC
                 .ReadFrom.Configuration(configuration)
                 .CreateLogger();
 
-            TestServer server = new TestServer(new WebHostBuilder()
+            _server = new TestServer(new WebHostBuilder()
                 .UseConfiguration(configuration)
                 .UseEnvironment("Testing")
                 .UseSerilog(logger)
@@ -178,10 +181,10 @@ namespace Horde.Storage.FunctionalTests.GC
                 {
                     collection.Configure<NamespaceSettings>(settings =>
                     {
-                        settings.Policies = new Dictionary<string, NamespaceSettings.PerNamespaceSettings>()
+                        settings.Policies = new Dictionary<string, NamespacePolicy>()
                         {
                             {
-                                TestNamespace.ToString(), new NamespaceSettings.PerNamespaceSettings()
+                                TestNamespace.ToString(), new NamespacePolicy()
                                 {
                                     IsLegacyNamespace = false
                                 }
@@ -191,11 +194,9 @@ namespace Horde.Storage.FunctionalTests.GC
                 })
                 .UseStartup<HordeStorageStartup>()
             );
-            _httpClient = server.CreateClient();
-            _server = server;
-            
+            _httpClient = _server.CreateClient();
 
-            IObjectService? objectService = server.Services.GetService<IObjectService>()!;
+            IObjectService? objectService = _server.Services.GetService<IObjectService>()!;
             Assert.IsNotNull(objectService);
             (BlobIdentifier ob0_hash, CbObject ob0_cb) = GetCBWithAttachment(object0id);
             await objectService.Put(TestNamespace, DefaultBucket, object0Name, ob0_hash, ob0_cb);
@@ -218,7 +219,7 @@ namespace Horde.Storage.FunctionalTests.GC
             (BlobIdentifier ob6_hash, CbObject ob6_cb) = GetCBWithAttachment(object6id);
             await objectService.Put(TestNamespace, DefaultBucket, object6Name, ob6_hash, ob6_cb);
 
-            IReferencesStore referenceStore = server.Services.GetService<IReferencesStore>()!;
+            IReferencesStore referenceStore = _server.Services.GetService<IReferencesStore>()!;
             DateTime oldTimestamp = DateTime.Now.AddDays(-30);
             DateTime newTimestamp = DateTime.Now;
             await referenceStore.UpdateLastAccessTime(TestNamespace, DefaultBucket, object0Name, oldTimestamp);
@@ -232,26 +233,18 @@ namespace Horde.Storage.FunctionalTests.GC
 
         protected abstract string GetImplementation();
 
-
         [TestMethod]
         public async Task RunRefCleanup()
         {
             // trigger the cleanup
-            var cleanupResponse = await _httpClient!.PostAsync($"api/v1/admin/refCleanup/{TestNamespace}", new StringContent(string.Empty));
+            using StringContent content = new StringContent(string.Empty);
+            HttpResponseMessage cleanupResponse = await _httpClient!.PostAsync(new Uri($"api/v1/admin/refCleanup/{TestNamespace}", UriKind.Relative), content);
             cleanupResponse.EnsureSuccessStatusCode();
             RemovedRefRecordsResponse removedRefRecords = await cleanupResponse.Content.ReadAsAsync<RemovedRefRecordsResponse>();
-
-            string[] expectedRemovedRefRecords = new[]
-            {
-                object0Name.ToString(),
-                object2Name.ToString(),
-                object3Name.ToString(),
-                object6Name.ToString()
-            };
             Assert.AreEqual(4, removedRefRecords.CountOfRemovedRecords);
         }
 
-        private (BlobIdentifier, CbObject) GetCBWithAttachment(BlobIdentifier blobIdentifier)
+        private static (BlobIdentifier, CbObject) GetCBWithAttachment(BlobIdentifier blobIdentifier)
         {
             CbWriter writer = new CbWriter();
             writer.BeginObject();
@@ -261,5 +254,20 @@ namespace Horde.Storage.FunctionalTests.GC
             byte[] b = writer.ToByteArray();
             return (BlobIdentifier.FromBlob(b), new CbObject(b));
         }
+
+        protected virtual void Dispose(bool disposing)
+        {
+            if (disposing)
+            {
+                _httpClient?.Dispose();
+                _server?.Dispose();
+            }
+        }
+
+        public void Dispose()
+        {
+            Dispose(true);
+            System.GC.SuppressFinalize(this);
+        }
     }
 }
diff --git a/Engine/Source/Programs/Horde/HordeStorage/Horde.Storage/Tests/Functional/Objects/RefReplicatorTests.cs b/Engine/Source/Programs/Horde/HordeStorage/Horde.Storage/Tests/Functional/Objects/RefReplicatorTests.cs
--- a/Engine/Source/Programs/Horde/HordeStorage/Horde.Storage/Tests/Functional/Objects/RefReplicatorTests.cs
+++ b/Engine/Source/Programs/Horde/HordeStorage/Horde.Storage/Tests/Functional/Objects/RefReplicatorTests.cs
@@ -7,7 +7,6 @@ using System.Linq;
 using System.Net;
 using System.Net.Http;
 using System.Text;
-using System.Text.Json;
 using System.Threading.Tasks;
 using Cassandra;
 using Horde.Storage.Controllers;
@@ -34,14 +33,11 @@ namespace Horde.Storage.FunctionalTests.Replication
     public class RefReplicatorTests
     {
         private static TestServer? _server;
-        private static HttpClient? _httpClient;
-        protected IBlobService _blobStore = null!;
-        protected IReferencesStore _referencesStore = null!;
-        private IReplicationLog _replicationLog = null!;
+        private IBlobService BlobStore { get; set; } = null!;
 
-        protected readonly NamespaceId TestNamespace = new NamespaceId("test-namespace");
-        protected readonly NamespaceId SnapshotNamespace = new NamespaceId("snapshot-namespace");
-        protected readonly BucketId TestBucket = new BucketId("test");
+        private NamespaceId TestNamespace { get; } = new NamespaceId("test-namespace");
+        private NamespaceId SnapshotNamespace { get; } = new NamespaceId("snapshot-namespace");
+        private BucketId TestBucket { get; } = new BucketId("test");
 
         [TestInitialize]
         public async Task Setup()
@@ -63,17 +59,15 @@ namespace Horde.Storage.FunctionalTests.Replication
                 .UseSerilog(logger)
                 .UseStartup<HordeStorageStartup>()
             );
-            _httpClient = server.CreateClient();
+            server.CreateClient();
             _server = server;
 
-            _blobStore = _server.Services.GetService<IBlobService>()!;
-            _referencesStore = _server.Services.GetService<IReferencesStore>()!;
-            _replicationLog = _server.Services.GetService<IReplicationLog>()!;
+            BlobStore = _server.Services.GetService<IBlobService>()!;
 
             await Task.CompletedTask;
         }
 
-        private IEnumerable<KeyValuePair<string, string>> GetSettings()
+        private static IEnumerable<KeyValuePair<string, string>> GetSettings()
         {
             return new[]
             {
@@ -81,9 +75,7 @@ namespace Horde.Storage.FunctionalTests.Replication
                 new KeyValuePair<string, string>("Horde_Storage:ReplicationLogWriterImplementation", HordeStorageSettings.ReplicationLogWriterImplementations.Scylla.ToString()),
             };
         }
-
-
-        private async Task TeardownDb(IServiceProvider provider)
+        private static async Task TeardownDb(IServiceProvider provider)
         {
             IScyllaSessionManager scyllaSessionManager = provider.GetService<IScyllaSessionManager>()!;
             ISession session = scyllaSessionManager.GetSessionForLocalKeyspace();
@@ -95,12 +87,13 @@ namespace Horde.Storage.FunctionalTests.Replication
             await session.ExecuteAsync(new SimpleStatement("DROP TABLE IF EXISTS replication_snapshot;"));
         }
 
-
         [TestCleanup]
         public async Task Teardown()
         {
-            if (_server != null) 
+            if (_server != null)
+            {
                 await TeardownDb(_server.Services);
+            }
         }
         
         [TestMethod]
@@ -146,7 +139,7 @@ namespace Horde.Storage.FunctionalTests.Replication
             }
 
             IHttpClientFactory httpClientFactory = handler.CreateClientFactory();
-            RefsReplicator replicator = ActivatorUtilities.CreateInstance<RefsReplicator>(_server!.Services, replicatorSettings, httpClientFactory);
+            using RefsReplicator replicator = ActivatorUtilities.CreateInstance<RefsReplicator>(_server!.Services, replicatorSettings, httpClientFactory);
             replicator.SetRefState(lastBucket, lastEvent);
 
             bool didRun = await replicator.TriggerNewReplications();
@@ -156,7 +149,7 @@ namespace Horde.Storage.FunctionalTests.Replication
             handler.Verify();
 
             // Verify that the objects are present
-            BlobIdentifier[] missingBlobs = await _blobStore.FilterOutKnownBlobs(TestNamespace, blobs.Keys.ToArray());
+            BlobIdentifier[] missingBlobs = await BlobStore.FilterOutKnownBlobs(TestNamespace, blobs.Keys.ToArray());
             Assert.IsFalse(missingBlobs.Any());
         }
 
@@ -220,7 +213,7 @@ namespace Horde.Storage.FunctionalTests.Replication
             }
 
             IHttpClientFactory httpClientFactory = handler.CreateClientFactory();
-            RefsReplicator replicator = ActivatorUtilities.CreateInstance<RefsReplicator>(_server!.Services, replicatorSettings, httpClientFactory);
+            using RefsReplicator replicator = ActivatorUtilities.CreateInstance<RefsReplicator>(_server!.Services, replicatorSettings, httpClientFactory);
             // no previous state, will download it from a snapshot
             replicator.SetRefState(null, null);
 
@@ -231,11 +224,10 @@ namespace Horde.Storage.FunctionalTests.Replication
             handler.Verify();
 
             // Verify that the objects are present
-            BlobIdentifier[] missingBlobs = await _blobStore.FilterOutKnownBlobs(TestNamespace, blobs.Keys.ToArray());
+            BlobIdentifier[] missingBlobs = await BlobStore.FilterOutKnownBlobs(TestNamespace, blobs.Keys.ToArray());
             Assert.IsFalse(missingBlobs.Any());
         }
 
-
         [TestMethod]
         public async Task ReplicationStateBoth()
         {
@@ -252,8 +244,8 @@ namespace Horde.Storage.FunctionalTests.Replication
             List<ReplicationLogEvent> incrementalEvents = new();
             Dictionary<BlobIdentifier, byte[]> blobs = new();
 
-            const int countOfTestEvents = 100;
-            for (int i = 0; i < countOfTestEvents; i++)
+            const int CountOfTestEvents = 100;
+            for (int i = 0; i < CountOfTestEvents; i++)
             {
                 byte[] blobContents = Encoding.UTF8.GetBytes($"random content in snapshot {i}");
                 BlobIdentifier blob = BlobIdentifier.FromBlob(blobContents);
@@ -261,7 +253,7 @@ namespace Horde.Storage.FunctionalTests.Replication
                 snapshotEvents.Add(new ReplicationLogEvent(TestNamespace, TestBucket,IoHashKey.FromName($"event-{i}"), blob, Guid.NewGuid(), "refs-000", DateTime.Now, ReplicationLogEvent.OpType.Added));
             }
 
-            for (int i = 0; i < countOfTestEvents; i++)
+            for (int i = 0; i < CountOfTestEvents; i++)
             {
                 byte[] blobContents = Encoding.UTF8.GetBytes($"random content {i}");
                 BlobIdentifier blob = BlobIdentifier.FromBlob(blobContents);
@@ -308,7 +300,7 @@ namespace Horde.Storage.FunctionalTests.Replication
             }
 
             IHttpClientFactory httpClientFactory = handler.CreateClientFactory();
-            RefsReplicator replicator = ActivatorUtilities.CreateInstance<RefsReplicator>(_server!.Services, replicatorSettings, httpClientFactory);
+            using RefsReplicator replicator = ActivatorUtilities.CreateInstance<RefsReplicator>(_server!.Services, replicatorSettings, httpClientFactory);
             // no previous state, will download it from a snapshot
             replicator.SetRefState(null, null);
 
@@ -319,11 +311,9 @@ namespace Horde.Storage.FunctionalTests.Replication
             handler.Verify();
 
             // Verify that the objects are present
-            BlobIdentifier[] missingBlobs = await _blobStore.FilterOutKnownBlobs(TestNamespace, blobs.Keys.ToArray());
+            BlobIdentifier[] missingBlobs = await BlobStore.FilterOutKnownBlobs(TestNamespace, blobs.Keys.ToArray());
             Assert.IsFalse(missingBlobs.Any());
         }
-
-
         
         [TestMethod]
         public async Task ReplicationStateSnapshotFallback()
@@ -340,8 +330,8 @@ namespace Horde.Storage.FunctionalTests.Replication
             List<ReplicationLogEvent> snapshotEvents = new();
             Dictionary<BlobIdentifier, byte[]> blobs = new();
 
-            const int countOfTestEvents = 100;
-            for (int i = 0; i < countOfTestEvents; i++)
+            const int CountOfTestEvents = 100;
+            for (int i = 0; i < CountOfTestEvents; i++)
             {
                 byte[] blobContents = Encoding.UTF8.GetBytes($"random content in snapshot {i}");
                 BlobIdentifier blob = BlobIdentifier.FromBlob(blobContents);
@@ -380,7 +370,6 @@ namespace Horde.Storage.FunctionalTests.Replication
                 Extensions = { { "SnapshotId", snapshotBlob } }
             }), "application/json");
 
-
             // after processing the snapshot we do not replicate anything more
             handler.SetupRequest($"http://localhost/api/v1/replication-log/incremental/{TestNamespace}?lastBucket={snapshotEvents.Last().TimeBucket}&lastEvent={snapshotEvents.Last().EventId}").ReturnsResponse(JsonConvert.SerializeObject(new ReplicationLogEvents(new List<ReplicationLogEvent>())), "application/json");
 
@@ -397,7 +386,7 @@ namespace Horde.Storage.FunctionalTests.Replication
             }
 
             IHttpClientFactory httpClientFactory = handler.CreateClientFactory();
-            RefsReplicator replicator = ActivatorUtilities.CreateInstance<RefsReplicator>(_server!.Services, replicatorSettings, httpClientFactory);
+            using RefsReplicator replicator = ActivatorUtilities.CreateInstance<RefsReplicator>(_server!.Services, replicatorSettings, httpClientFactory);
             // specify a bucket that does not exist, we should fallback to replicating a snapshot
             replicator.SetRefState(missingBucket, missingId);
 
@@ -408,7 +397,7 @@ namespace Horde.Storage.FunctionalTests.Replication
             handler.Verify();
 
             // Verify that the objects are present
-            BlobIdentifier[] missingBlobs = await _blobStore.FilterOutKnownBlobs(TestNamespace, blobs.Keys.ToArray());
+            BlobIdentifier[] missingBlobs = await BlobStore.FilterOutKnownBlobs(TestNamespace, blobs.Keys.ToArray());
             Assert.IsFalse(missingBlobs.Any());
         }
     }
