Horde Storage - Updated editorconfig to more closely match what we run for Horde and updated code to fix all static analysis errors

diff --git a/Engine/Source/Programs/Horde/HordeStorage/Horde.Storage/Tests/Functional/Objects/ReferencesTests.cs b/Engine/Source/Programs/Horde/HordeStorage/Horde.Storage/Tests/Functional/Objects/ReferencesTests.cs
--- a/Engine/Source/Programs/Horde/HordeStorage/Horde.Storage/Tests/Functional/Objects/ReferencesTests.cs
+++ b/Engine/Source/Programs/Horde/HordeStorage/Horde.Storage/Tests/Functional/Objects/ReferencesTests.cs
@@ -29,7 +29,6 @@ using Serilog;
 using Logger = Serilog.Core.Logger;
 using EpicGames.Horde.Storage;
 using EpicGames.Serialization;
-using ContentHash = Jupiter.Implementation.ContentHash;
 
 namespace Horde.Storage.FunctionalTests.References
 {
@@ -61,7 +60,7 @@ namespace Horde.Storage.FunctionalTests.References
             //verify we are using the replication log writer
             Assert.IsTrue(replicationLog.GetType() == typeof(ScyllaReplicationLog));
 
-            await SeedTestData(referencesStore);
+            await SeedTestData();
         }
 
         protected override async Task TeardownDb(IServiceProvider provider)
@@ -115,7 +114,7 @@ namespace Horde.Storage.FunctionalTests.References
             //verify we are using the replication log writer
             Assert.IsTrue(replicationLog.GetType() == typeof(MemoryReplicationLog));
 
-            await SeedTestData(referencesStore);
+            await SeedTestData();
         }
 
         protected override async Task TeardownDb(IServiceProvider provider)
@@ -151,7 +150,7 @@ namespace Horde.Storage.FunctionalTests.References
             //verify we are using the replication log writer
             Assert.IsTrue(replicationLog.GetType() == typeof(MemoryReplicationLog));
 
-            await SeedTestData(referencesStore);
+            await SeedTestData();
         }
 
         protected override Task TeardownDb(IServiceProvider provider)
@@ -164,10 +163,10 @@ namespace Horde.Storage.FunctionalTests.References
     {
         private static TestServer? _server;
         private static HttpClient? _httpClient;
-        protected IBlobService _blobService = null!;
-        protected IReferencesStore ReferencesStore = null!;
 
-        protected readonly NamespaceId TestNamespace = new NamespaceId("test-namespace");
+        protected IBlobService Service { get; set; } = null!;
+        protected IReferencesStore ReferencesStore { get; set; } = null!;
+        protected NamespaceId TestNamespace { get; } = new NamespaceId("test-namespace");
 
         [TestInitialize]
         public async Task Setup()
@@ -193,22 +192,23 @@ namespace Horde.Storage.FunctionalTests.References
             _httpClient = server.CreateClient();
             _server = server;
 
-            _blobService = _server.Services.GetService<IBlobService>()!;
+            Service = _server.Services.GetService<IBlobService>()!;
             ReferencesStore = _server.Services.GetService<IReferencesStore>()!;
             await SeedDb(server.Services);
         }
 
-        protected async Task SeedTestData(IReferencesStore referencesStore)
+        protected virtual async Task SeedTestData()
         {
             await Task.CompletedTask;
         }
 
-
         [TestCleanup]
         public async Task Teardown()
         {
-            if (_server != null) 
+            if (_server != null)
+            {
                 await TeardownDb(_server.Services);
+            }
         }
 
         protected abstract IEnumerable<KeyValuePair<string, string>> GetSettings();
@@ -223,15 +223,15 @@ namespace Horde.Storage.FunctionalTests.References
             byte[] data = Encoding.ASCII.GetBytes(objectContents);
             BlobIdentifier objectHash = BlobIdentifier.FromBlob(data);
             IoHashKey key = IoHashKey.FromName("newBlobObject");
-            HttpContent requestContent = new ByteArrayContent(data);
+            using HttpContent requestContent = new ByteArrayContent(data);
             requestContent.Headers.ContentType = new MediaTypeHeaderValue(MediaTypeNames.Application.Octet);
             requestContent.Headers.Add(CommonHeaders.HashHeaderName, objectHash.ToString());
 
-            HttpResponseMessage result = await _httpClient!.PutAsync(requestUri: $"api/v1/refs/{TestNamespace}/bucket/{key}", requestContent);
+            HttpResponseMessage result = await _httpClient!.PutAsync(new Uri($"api/v1/refs/{TestNamespace}/bucket/{key}", UriKind.Relative), requestContent);
             result.EnsureSuccessStatusCode();
 
             {
-                HttpResponseMessage getResponse = await _httpClient.GetAsync($"api/v1/refs/{TestNamespace}/bucket/{key}.raw");
+                HttpResponseMessage getResponse = await _httpClient.GetAsync(new Uri($"api/v1/refs/{TestNamespace}/bucket/{key}.raw", UriKind.Relative));
                 getResponse.EnsureSuccessStatusCode();
                 await using MemoryStream ms = new MemoryStream();
                 await getResponse.Content.CopyToAsync(ms);
@@ -247,7 +247,7 @@ namespace Horde.Storage.FunctionalTests.References
             {
                 BlobIdentifier attachment;
                 {
-                    HttpResponseMessage getResponse = await _httpClient.GetAsync($"api/v1/refs/{TestNamespace}/bucket/{key}.uecb");
+                    HttpResponseMessage getResponse = await _httpClient.GetAsync(new Uri($"api/v1/refs/{TestNamespace}/bucket/{key}.uecb", UriKind.Relative));
                     getResponse.EnsureSuccessStatusCode();
                     await using MemoryStream ms = new MemoryStream();
                     await getResponse.Content.CopyToAsync(ms);
@@ -264,7 +264,7 @@ namespace Horde.Storage.FunctionalTests.References
                 }
 
                 {
-                    HttpResponseMessage getAttachment = await _httpClient.GetAsync($"api/v1/blobs/{TestNamespace}/{attachment}");
+                    HttpResponseMessage getAttachment = await _httpClient.GetAsync(new Uri($"api/v1/blobs/{TestNamespace}/{attachment}", UriKind.Relative));
                     getAttachment.EnsureSuccessStatusCode();
                     await using MemoryStream ms = new MemoryStream();
                     await getAttachment.Content.CopyToAsync(ms);
@@ -277,7 +277,7 @@ namespace Horde.Storage.FunctionalTests.References
             }
 
             {
-                HttpResponseMessage getResponse = await _httpClient.GetAsync($"api/v1/refs/{TestNamespace}/bucket/{key}.json");
+                HttpResponseMessage getResponse = await _httpClient.GetAsync(new Uri($"api/v1/refs/{TestNamespace}/bucket/{key}.json", UriKind.Relative));
                 getResponse.EnsureSuccessStatusCode();
                 await using MemoryStream ms = new MemoryStream();
                 await getResponse.Content.CopyToAsync(ms);
@@ -301,7 +301,7 @@ namespace Horde.Storage.FunctionalTests.References
 
             {
                 // request the object as a json response using accept instead of the format filter
-                var request = new HttpRequestMessage(HttpMethod.Get, $"api/v1/refs/{TestNamespace}/bucket/{key}");
+                using HttpRequestMessage request = new(HttpMethod.Get, new Uri($"api/v1/refs/{TestNamespace}/bucket/{key}", UriKind.Relative));
                 request.Headers.Accept.Add(new MediaTypeWithQualityHeaderValue(MediaTypeNames.Application.Json));
                 HttpResponseMessage getResponse = await _httpClient.SendAsync(request);
                 getResponse.EnsureSuccessStatusCode();
@@ -329,7 +329,7 @@ namespace Horde.Storage.FunctionalTests.References
 
             {
                 // request the object as a jupiter inlined payload
-                var request = new HttpRequestMessage(HttpMethod.Get, $"api/v1/refs/{TestNamespace}/bucket/{key}");
+                using HttpRequestMessage request = new (HttpMethod.Get, new Uri($"api/v1/refs/{TestNamespace}/bucket/{key}", UriKind.Relative));
                 request.Headers.Accept.Add(new MediaTypeWithQualityHeaderValue(CustomMediaTypeNames.JupiterInlinedPayload));
                 HttpResponseMessage getResponse = await _httpClient.SendAsync(request);
                 getResponse.EnsureSuccessStatusCode();
@@ -358,11 +358,11 @@ namespace Horde.Storage.FunctionalTests.References
             BlobIdentifier objectHash = BlobIdentifier.FromBlob(objectData);
             IoHashKey key = IoHashKey.FromName("newReferenceObject");
 
-            HttpContent requestContent = new ByteArrayContent(objectData);
+            using HttpContent requestContent = new ByteArrayContent(objectData);
             requestContent.Headers.ContentType = new MediaTypeHeaderValue(CustomMediaTypeNames.UnrealCompactBinary);
             requestContent.Headers.Add(CommonHeaders.HashHeaderName, objectHash.ToString());
 
-            HttpResponseMessage result = await _httpClient!.PutAsync(requestUri: $"api/v1/refs/{TestNamespace}/bucket/{key}.uecb", requestContent);
+            using HttpResponseMessage result = await _httpClient!.PutAsync(new Uri($"api/v1/refs/{TestNamespace}/bucket/{key}.uecb", UriKind.Relative), requestContent);
             result.EnsureSuccessStatusCode();
 
             {
@@ -390,7 +390,7 @@ namespace Horde.Storage.FunctionalTests.References
             }
 
             {
-                HttpResponseMessage getResponse = await _httpClient.GetAsync($"api/v1/refs/{TestNamespace}/bucket/{key}.raw");
+                HttpResponseMessage getResponse = await _httpClient.GetAsync(new Uri($"api/v1/refs/{TestNamespace}/bucket/{key}.raw", UriKind.Relative));
                 getResponse.EnsureSuccessStatusCode();
                 await using MemoryStream ms = new MemoryStream();
                 await getResponse.Content.CopyToAsync(ms);
@@ -402,7 +402,7 @@ namespace Horde.Storage.FunctionalTests.References
             }
 
             {
-                HttpResponseMessage getResponse = await _httpClient.GetAsync($"api/v1/refs/{TestNamespace}/bucket/{key}.uecb");
+                HttpResponseMessage getResponse = await _httpClient.GetAsync(new Uri($"api/v1/refs/{TestNamespace}/bucket/{key}.uecb", UriKind.Relative));
                 getResponse.EnsureSuccessStatusCode();
                 await using MemoryStream ms = new MemoryStream();
                 await getResponse.Content.CopyToAsync(ms);
@@ -417,9 +417,8 @@ namespace Horde.Storage.FunctionalTests.References
                 Assert.AreEqual("thisIsAField", stringField.AsString());
             }
 
-
             {
-                HttpResponseMessage getResponse = await _httpClient.GetAsync($"api/v1/refs/{TestNamespace}/bucket/{key}.json");
+                HttpResponseMessage getResponse = await _httpClient.GetAsync(new Uri($"api/v1/refs/{TestNamespace}/bucket/{key}.json", UriKind.Relative));
                 getResponse.EnsureSuccessStatusCode();
                 await using MemoryStream ms = new MemoryStream();
                 await getResponse.Content.CopyToAsync(ms);
@@ -457,11 +456,11 @@ namespace Horde.Storage.FunctionalTests.References
             BlobIdentifier objectHash = BlobIdentifier.FromBlob(objectData);
             IoHashKey key = IoHashKey.FromName("newReferenceObject");
 
-            HttpContent requestContent = new ByteArrayContent(objectData);
+            using HttpContent requestContent = new ByteArrayContent(objectData);
             requestContent.Headers.ContentType = new MediaTypeHeaderValue(CustomMediaTypeNames.UnrealCompactBinary);
             requestContent.Headers.Add(CommonHeaders.HashHeaderName, objectHash.ToString());
 
-            HttpResponseMessage result = await _httpClient!.PutAsync(requestUri: $"api/v1/refs/{TestNamespace}/bucket/{key}.uecb", requestContent);
+            HttpResponseMessage result = await _httpClient!.PutAsync(new Uri($"api/v1/refs/{TestNamespace}/bucket/{key}.uecb", UriKind.Relative), requestContent);
             result.EnsureSuccessStatusCode();
 
             {
@@ -489,7 +488,7 @@ namespace Horde.Storage.FunctionalTests.References
             }
             
             {
-                HttpResponseMessage getResponse = await _httpClient.GetAsync($"api/v1/refs/{TestNamespace}/bucket/{key}.json?fields=name");
+                HttpResponseMessage getResponse = await _httpClient.GetAsync(new Uri($"api/v1/refs/{TestNamespace}/bucket/{key}.json?fields=name", UriKind.Relative));
                 getResponse.EnsureSuccessStatusCode();
                 await using MemoryStream ms = new MemoryStream();
                 await getResponse.Content.CopyToAsync(ms);
@@ -535,11 +534,11 @@ namespace Horde.Storage.FunctionalTests.References
             IoHashKey key = IoHashKey.FromName("newHierarchyObject");
             // this first upload should fail with the child object missing
             {
-                HttpContent requestContent = new ByteArrayContent(parentObjectData);
+                using HttpContent requestContent = new ByteArrayContent(parentObjectData);
                 requestContent.Headers.ContentType = new MediaTypeHeaderValue(CustomMediaTypeNames.UnrealCompactBinary);
                 requestContent.Headers.Add(CommonHeaders.HashHeaderName, parentObjectHash.ToString());
 
-                HttpResponseMessage result = await _httpClient!.PutAsync(requestUri: $"api/v1/refs/{TestNamespace}/bucket/{key}.uecb", requestContent);
+                HttpResponseMessage result = await _httpClient!.PutAsync(new Uri($"api/v1/refs/{TestNamespace}/bucket/{key}.uecb", UriKind.Relative), requestContent);
                 result.EnsureSuccessStatusCode();
 
                 Assert.AreEqual(CustomMediaTypeNames.UnrealCompactBinary, result!.Content.Headers.ContentType!.MediaType);
@@ -557,7 +556,7 @@ namespace Horde.Storage.FunctionalTests.References
 
             // upload the child object
             {
-                HttpRequestMessage request = new HttpRequestMessage(HttpMethod.Put, $"api/v1/objects/{TestNamespace}/{childObjectHash}");
+                using HttpRequestMessage request = new HttpRequestMessage(HttpMethod.Put, new Uri($"api/v1/objects/{TestNamespace}/{childObjectHash}", UriKind.Relative));
                 request.Headers.Accept.Add(new MediaTypeWithQualityHeaderValue(CustomMediaTypeNames.UnrealCompactBinary));
 
                 HttpContent requestContent = new ByteArrayContent(childObjectData);
@@ -581,11 +580,11 @@ namespace Horde.Storage.FunctionalTests.References
 
             // since we have now uploaded the child object putting the object again should result in no missing references
             {
-                HttpContent requestContent = new ByteArrayContent(parentObjectData);
+                using HttpContent requestContent = new ByteArrayContent(parentObjectData);
                 requestContent.Headers.ContentType = new MediaTypeHeaderValue(CustomMediaTypeNames.UnrealCompactBinary);
                 requestContent.Headers.Add(CommonHeaders.HashHeaderName, parentObjectHash.ToString());
 
-                HttpResponseMessage result = await _httpClient!.PutAsync(requestUri: $"api/v1/refs/{TestNamespace}/bucket/{key}.uecb", requestContent);
+                HttpResponseMessage result = await _httpClient!.PutAsync(new Uri($"api/v1/refs/{TestNamespace}/bucket/{key}.uecb", UriKind.Relative), requestContent);
                 result.EnsureSuccessStatusCode();
 
                 Assert.AreEqual(result!.Content.Headers.ContentType!.MediaType, CustomMediaTypeNames.UnrealCompactBinary);
@@ -611,7 +610,7 @@ namespace Horde.Storage.FunctionalTests.References
             }
 
             {
-                HttpResponseMessage getResponse = await _httpClient.GetAsync($"api/v1/refs/{TestNamespace}/bucket/{key}.raw");
+                HttpResponseMessage getResponse = await _httpClient.GetAsync(new Uri($"api/v1/refs/{TestNamespace}/bucket/{key}.raw", UriKind.Relative));
                 getResponse.EnsureSuccessStatusCode();
                 await using MemoryStream ms = new MemoryStream();
                 await getResponse.Content.CopyToAsync(ms);
@@ -623,7 +622,7 @@ namespace Horde.Storage.FunctionalTests.References
             }
 
             {
-                HttpResponseMessage getResponse = await _httpClient.GetAsync($"api/v1/refs/{TestNamespace}/bucket/{key}.uecb");
+                HttpResponseMessage getResponse = await _httpClient.GetAsync(new Uri($"api/v1/refs/{TestNamespace}/bucket/{key}.uecb", UriKind.Relative));
                 getResponse.EnsureSuccessStatusCode();
                 await using MemoryStream ms = new MemoryStream();
                 await getResponse.Content.CopyToAsync(ms);
@@ -638,9 +637,8 @@ namespace Horde.Storage.FunctionalTests.References
                 Assert.AreEqual(childObjectHash, BlobIdentifier.FromIoHash(childObjectField.AsHash()));
             }
 
-
             {
-                HttpResponseMessage getResponse = await _httpClient.GetAsync($"api/v1/refs/{TestNamespace}/bucket/{key}.json");
+                HttpResponseMessage getResponse = await _httpClient.GetAsync(new Uri($"api/v1/refs/{TestNamespace}/bucket/{key}.json", UriKind.Relative));
                 getResponse.EnsureSuccessStatusCode();
                 await using MemoryStream ms = new MemoryStream();
                 await getResponse.Content.CopyToAsync(ms);
@@ -665,7 +663,7 @@ namespace Horde.Storage.FunctionalTests.References
             }
 
             {
-                HttpResponseMessage getResponse = await _httpClient.GetAsync($"api/v1/objects/{TestNamespace}/{childObjectHash}");
+                HttpResponseMessage getResponse = await _httpClient.GetAsync(new Uri($"api/v1/objects/{TestNamespace}/{childObjectHash}", UriKind.Relative));
                 getResponse.EnsureSuccessStatusCode();
                 await using MemoryStream ms = new MemoryStream();
                 await getResponse.Content.CopyToAsync(ms);
@@ -675,7 +673,6 @@ namespace Horde.Storage.FunctionalTests.References
                 CollectionAssert.AreEqual(childObjectData, roundTrippedBuffer);
                 Assert.AreEqual(childObjectHash, BlobIdentifier.FromBlob(roundTrippedBuffer));
             }
-
         }
 
         [TestMethod]
@@ -685,22 +682,24 @@ namespace Horde.Storage.FunctionalTests.References
             byte[] data = Encoding.ASCII.GetBytes(objectContents);
             BlobIdentifier objectHash = BlobIdentifier.FromBlob(data);
             IoHashKey key = IoHashKey.FromName("newObject");
-            HttpContent requestContent = new ByteArrayContent(data);
+            using HttpContent requestContent = new ByteArrayContent(data);
             requestContent.Headers.ContentType = new MediaTypeHeaderValue(MediaTypeNames.Application.Octet);
             requestContent.Headers.Add(CommonHeaders.HashHeaderName, objectHash.ToString());
 
             {
-                HttpResponseMessage result = await _httpClient!.PutAsync(requestUri: $"api/v1/refs/{TestNamespace}/bucket/{key}", requestContent);
+                HttpResponseMessage result = await _httpClient!.PutAsync(new Uri($"api/v1/refs/{TestNamespace}/bucket/{key}", UriKind.Relative), requestContent);
                 result.EnsureSuccessStatusCode();
             }
 
             {
-                HttpResponseMessage result = await _httpClient!.SendAsync(new HttpRequestMessage(HttpMethod.Head, $"api/v1/refs/{TestNamespace}/bucket/{key}"));
+                using HttpRequestMessage message = new(HttpMethod.Head, new Uri($"api/v1/refs/{TestNamespace}/bucket/{key}", UriKind.Relative));
+                HttpResponseMessage result = await _httpClient!.SendAsync(message);
                 result.EnsureSuccessStatusCode();
             }
 
             {
-                HttpResponseMessage result = await _httpClient!.SendAsync(new HttpRequestMessage(HttpMethod.Head, $"api/v1/refs/{TestNamespace}/bucket/{IoHashKey.FromName("missingObject")}"));
+                using HttpRequestMessage message = new(HttpMethod.Head, new Uri($"api/v1/refs/{TestNamespace}/bucket/{IoHashKey.FromName("missingObject")}", UriKind.Relative));
+                HttpResponseMessage result = await _httpClient!.SendAsync(message);
                 Assert.AreEqual(HttpStatusCode.NotFound, result.StatusCode);
             }
         }
@@ -714,12 +713,12 @@ namespace Horde.Storage.FunctionalTests.References
             const string objectContents = "This is treated as a opaque blob";
             byte[] data = Encoding.ASCII.GetBytes(objectContents);
             BlobIdentifier objectHash = BlobIdentifier.FromBlob(data);
-            HttpContent requestContent = new ByteArrayContent(data);
+            using HttpContent requestContent = new ByteArrayContent(data);
             requestContent.Headers.ContentType = new MediaTypeHeaderValue(MediaTypeNames.Application.Octet);
             requestContent.Headers.Add(CommonHeaders.HashHeaderName, objectHash.ToString());
 
             {
-                HttpResponseMessage result = await _httpClient!.PutAsync(requestUri: $"api/v1/refs/{TestNamespace}/{bucket}/{existingObject}", requestContent);
+                HttpResponseMessage result = await _httpClient!.PutAsync(new Uri($"api/v1/refs/{TestNamespace}/{bucket}/{existingObject}", UriKind.Relative), requestContent);
                 result.EnsureSuccessStatusCode();
             }
 
@@ -728,7 +727,7 @@ namespace Horde.Storage.FunctionalTests.References
             string queryString = $"?names={bucket}.{existingObject}&names={bucket}.{missingObject}";
 
             {
-                HttpResponseMessage result = await _httpClient!.GetAsync(requestUri: $"api/v1/refs/{TestNamespace}/exists" + queryString);
+                HttpResponseMessage result = await _httpClient!.GetAsync(new Uri($"api/v1/refs/{TestNamespace}/exists" + queryString, UriKind.Relative));
                 result.EnsureSuccessStatusCode();
                 ExistCheckMultipleRefsResponse response = await result.Content.ReadAsAsync<ExistCheckMultipleRefsResponse>();
 
@@ -738,20 +737,18 @@ namespace Horde.Storage.FunctionalTests.References
             }
         }
 
-
-
         [TestMethod]
         public async Task PutGetObjectHierarchy()
         {
             string blobContents = "This is a string that is referenced as a blob";
             byte[] blobData = Encoding.ASCII.GetBytes(blobContents);
             BlobIdentifier blobHash = BlobIdentifier.FromBlob(blobData);
-            await _blobService.PutObject(TestNamespace, blobData, blobHash);
+            await Service.PutObject(TestNamespace, blobData, blobHash);
 
             string blobContentsChild = "This string is also referenced as a blob but from a child object";
             byte[] dataChild = Encoding.ASCII.GetBytes(blobContentsChild);
             BlobIdentifier blobHashChild = BlobIdentifier.FromBlob(dataChild);
-            await _blobService.PutObject(TestNamespace, dataChild, blobHashChild);
+            await Service.PutObject(TestNamespace, dataChild, blobHashChild);
 
             CbWriter writerChild = new CbWriter();
             writerChild.BeginObject();
@@ -760,7 +757,7 @@ namespace Horde.Storage.FunctionalTests.References
 
             byte[] childDataObject = writerChild.ToByteArray();
             BlobIdentifier childDataObjectHash = BlobIdentifier.FromBlob(childDataObject);
-            await _blobService.PutObject(TestNamespace, childDataObject, childDataObjectHash);
+            await Service.PutObject(TestNamespace, childDataObject, childDataObjectHash);
 
             CbWriter writerParent = new CbWriter();
             writerParent.BeginObject();
@@ -774,11 +771,11 @@ namespace Horde.Storage.FunctionalTests.References
 
             IoHashKey key = IoHashKey.FromName("newHierarchyObject");
 
-            HttpContent requestContent = new ByteArrayContent(objectData);
+            using HttpContent requestContent = new ByteArrayContent(objectData);
             requestContent.Headers.ContentType = new MediaTypeHeaderValue(CustomMediaTypeNames.UnrealCompactBinary);
             requestContent.Headers.Add(CommonHeaders.HashHeaderName, objectHash.ToString());
 
-            HttpResponseMessage result = await _httpClient!.PutAsync(requestUri: $"api/v1/refs/{TestNamespace}/bucket/{key}.uecb", requestContent);
+            HttpResponseMessage result = await _httpClient!.PutAsync(new Uri($"api/v1/refs/{TestNamespace}/bucket/{key}.uecb", UriKind.Relative), requestContent);
             result.EnsureSuccessStatusCode();
 
             // check the response
@@ -808,7 +805,7 @@ namespace Horde.Storage.FunctionalTests.References
 
             // verify attachments
             {
-                HttpResponseMessage getResponse = await _httpClient.GetAsync($"api/v1/refs/{TestNamespace}/bucket/{key}.raw");
+                HttpResponseMessage getResponse = await _httpClient.GetAsync(new Uri($"api/v1/refs/{TestNamespace}/bucket/{key}.raw", UriKind.Relative));
                 getResponse.EnsureSuccessStatusCode();
                 await using MemoryStream ms = new MemoryStream();
                 await getResponse.Content.CopyToAsync(ms);
@@ -823,7 +820,7 @@ namespace Horde.Storage.FunctionalTests.References
                 BlobIdentifier blobAttachment;
                 BlobIdentifier objectAttachment;
                 {
-                    HttpResponseMessage getResponse = await _httpClient.GetAsync($"api/v1/refs/{TestNamespace}/bucket/{key}.uecb");
+                    HttpResponseMessage getResponse = await _httpClient.GetAsync(new Uri($"api/v1/refs/{TestNamespace}/bucket/{key}.uecb", UriKind.Relative));
                     getResponse.EnsureSuccessStatusCode();
                     await using MemoryStream ms = new MemoryStream();
                     await getResponse.Content.CopyToAsync(ms);
@@ -842,7 +839,7 @@ namespace Horde.Storage.FunctionalTests.References
                 }
 
                 {
-                    HttpResponseMessage getAttachment = await _httpClient.GetAsync($"api/v1/blobs/{TestNamespace}/{blobAttachment}");
+                    HttpResponseMessage getAttachment = await _httpClient.GetAsync(new Uri($"api/v1/blobs/{TestNamespace}/{blobAttachment}", UriKind.Relative));
                     getAttachment.EnsureSuccessStatusCode();
                     await using MemoryStream ms = new MemoryStream();
                     await getAttachment.Content.CopyToAsync(ms);
@@ -855,7 +852,7 @@ namespace Horde.Storage.FunctionalTests.References
 
                 BlobIdentifier attachedBlobIdentifier;
                 {
-                    HttpResponseMessage getAttachment = await _httpClient.GetAsync($"api/v1/blobs/{TestNamespace}/{objectAttachment}");
+                    HttpResponseMessage getAttachment = await _httpClient.GetAsync(new Uri($"api/v1/blobs/{TestNamespace}/{objectAttachment}", UriKind.Relative));
                     getAttachment.EnsureSuccessStatusCode();
                     await using MemoryStream ms = new MemoryStream();
                     await getAttachment.Content.CopyToAsync(ms);
@@ -871,7 +868,7 @@ namespace Horde.Storage.FunctionalTests.References
                 }
 
                 {
-                    HttpResponseMessage getAttachment = await _httpClient.GetAsync($"api/v1/blobs/{TestNamespace}/{attachedBlobIdentifier}");
+                    HttpResponseMessage getAttachment = await _httpClient.GetAsync(new Uri($"api/v1/blobs/{TestNamespace}/{attachedBlobIdentifier}", UriKind.Relative));
                     getAttachment.EnsureSuccessStatusCode();
                     await using MemoryStream ms = new MemoryStream();
                     await getAttachment.Content.CopyToAsync(ms);
@@ -885,7 +882,7 @@ namespace Horde.Storage.FunctionalTests.References
 
             // check json representation
             {
-                HttpResponseMessage getResponse = await _httpClient.GetAsync($"api/v1/refs/{TestNamespace}/bucket/{key}.json");
+                HttpResponseMessage getResponse = await _httpClient.GetAsync(new Uri($"api/v1/refs/{TestNamespace}/bucket/{key}.json", UriKind.Relative));
                 getResponse.EnsureSuccessStatusCode();
                 await using MemoryStream ms = new MemoryStream();
                 await getResponse.Content.CopyToAsync(ms);
@@ -905,8 +902,6 @@ namespace Horde.Storage.FunctionalTests.References
 
             }
         }
-
-
         
         [TestMethod]
         public async Task PutPartialHierarchy()
@@ -927,7 +922,7 @@ namespace Horde.Storage.FunctionalTests.References
 
             byte[] childDataObject = writerChild.ToByteArray();
             BlobIdentifier childDataObjectHash = BlobIdentifier.FromBlob(childDataObject);
-            await _blobService.PutObject(TestNamespace, childDataObject, childDataObjectHash);
+            await Service.PutObject(TestNamespace, childDataObject, childDataObjectHash);
 
             CbWriter writerParent = new CbWriter();
             writerParent.BeginObject();
@@ -942,11 +937,11 @@ namespace Horde.Storage.FunctionalTests.References
             IoHashKey key = IoHashKey.FromName("newHierarchyObject");
 
             {
-                HttpContent requestContent = new ByteArrayContent(objectData);
+                using HttpContent requestContent = new ByteArrayContent(objectData);
                 requestContent.Headers.ContentType = new MediaTypeHeaderValue(CustomMediaTypeNames.UnrealCompactBinary);
                 requestContent.Headers.Add(CommonHeaders.HashHeaderName, objectHash.ToString());
 
-                HttpResponseMessage result = await _httpClient!.PutAsync(requestUri: $"api/v1/refs/{TestNamespace}/bucket/{key}.uecb", requestContent);
+                HttpResponseMessage result = await _httpClient!.PutAsync(new Uri($"api/v1/refs/{TestNamespace}/bucket/{key}.uecb", UriKind.Relative), requestContent);
                 result.EnsureSuccessStatusCode();
 
                 {
@@ -965,11 +960,11 @@ namespace Horde.Storage.FunctionalTests.References
             }
 
             {
-                HttpContent requestContent = new ByteArrayContent(objectData);
+                using HttpContent requestContent = new ByteArrayContent(objectData);
                 requestContent.Headers.ContentType = new MediaTypeHeaderValue(CustomMediaTypeNames.UnrealCompactBinary);
                 requestContent.Headers.Add(CommonHeaders.HashHeaderName, objectHash.ToString());
 
-                HttpResponseMessage result = await _httpClient!.PutAsync(requestUri: $"api/v1/refs/{TestNamespace}/bucket/{key}.json", requestContent);
+                HttpResponseMessage result = await _httpClient!.PutAsync( new Uri($"api/v1/refs/{TestNamespace}/bucket/{key}.json", UriKind.Relative), requestContent);
                 result.EnsureSuccessStatusCode();
 
                 {
@@ -1016,11 +1011,11 @@ namespace Horde.Storage.FunctionalTests.References
             IoHashKey key = IoHashKey.FromName("putContentIdMissingBlob");
 
             {
-                HttpContent requestContent = new ByteArrayContent(objectData);
+                using HttpContent requestContent = new ByteArrayContent(objectData);
                 requestContent.Headers.ContentType = new MediaTypeHeaderValue(CustomMediaTypeNames.UnrealCompactBinary);
                 requestContent.Headers.Add(CommonHeaders.HashHeaderName, objectHash.ToString());
 
-                HttpResponseMessage result = await _httpClient!.PutAsync(requestUri: $"api/v1/refs/{TestNamespace}/bucket/{key}.uecb", requestContent);
+                HttpResponseMessage result = await _httpClient!.PutAsync(new Uri($"api/v1/refs/{TestNamespace}/bucket/{key}.uecb", UriKind.Relative), requestContent);
                 result.EnsureSuccessStatusCode();
 
                 {
@@ -1040,11 +1035,11 @@ namespace Horde.Storage.FunctionalTests.References
             }
 
             {
-                HttpContent requestContent = new ByteArrayContent(objectData);
+                using HttpContent requestContent = new ByteArrayContent(objectData);
                 requestContent.Headers.ContentType = new MediaTypeHeaderValue(CustomMediaTypeNames.UnrealCompactBinary);
                 requestContent.Headers.Add(CommonHeaders.HashHeaderName, objectHash.ToString());
 
-                HttpResponseMessage result = await _httpClient!.PutAsync(requestUri: $"api/v1/refs/{TestNamespace}/bucket/{key}.json", requestContent);
+                HttpResponseMessage result = await _httpClient!.PutAsync(new Uri($"api/v1/refs/{TestNamespace}/bucket/{key}.json", UriKind.Relative), requestContent);
                 result.EnsureSuccessStatusCode();
 
                 {
@@ -1089,11 +1084,11 @@ namespace Horde.Storage.FunctionalTests.References
             IoHashKey key = IoHashKey.FromName("putContentIdMissingBlobComplex");
 
             {
-                HttpContent requestContent = new ByteArrayContent(objectData);
+                using HttpContent requestContent = new ByteArrayContent(objectData);
                 requestContent.Headers.ContentType = new MediaTypeHeaderValue(CustomMediaTypeNames.UnrealCompactBinary);
                 requestContent.Headers.Add(CommonHeaders.HashHeaderName, objectHash.ToString());
 
-                HttpResponseMessage result = await _httpClient!.PutAsync(requestUri: $"api/v1/refs/{TestNamespace}/bucket/{key}.uecb", requestContent);
+                HttpResponseMessage result = await _httpClient!.PutAsync(new Uri($"api/v1/refs/{TestNamespace}/bucket/{key}.uecb", UriKind.Relative), requestContent);
                 result.EnsureSuccessStatusCode();
 
                 {
@@ -1134,7 +1129,7 @@ namespace Horde.Storage.FunctionalTests.References
 
             byte[] childDataObject = writerChild.ToByteArray();
             BlobIdentifier childDataObjectHash = BlobIdentifier.FromBlob(childDataObject);
-            await _blobService.PutObject(TestNamespace, childDataObject, childDataObjectHash);
+            await Service.PutObject(TestNamespace, childDataObject, childDataObjectHash);
 
             CbWriter writerParent = new CbWriter();
             writerParent.BeginObject();
@@ -1147,11 +1142,11 @@ namespace Horde.Storage.FunctionalTests.References
             BlobIdentifier objectHash = BlobIdentifier.FromBlob(objectData);
 
             {
-                HttpContent requestContent = new ByteArrayContent(objectData);
+                using HttpContent requestContent = new ByteArrayContent(objectData);
                 requestContent.Headers.ContentType = new MediaTypeHeaderValue(CustomMediaTypeNames.UnrealCompactBinary);
                 requestContent.Headers.Add(CommonHeaders.HashHeaderName, objectHash.ToString());
 
-                HttpResponseMessage result = await _httpClient!.PutAsync(requestUri: $"api/v1/refs/{TestNamespace}/bucket/{key}.uecb", requestContent);
+                HttpResponseMessage result = await _httpClient!.PutAsync(new Uri($"api/v1/refs/{TestNamespace}/bucket/{key}.uecb", UriKind.Relative), requestContent);
                 result.EnsureSuccessStatusCode();
 
                 {
@@ -1180,15 +1175,15 @@ namespace Horde.Storage.FunctionalTests.References
 
             // upload missing pieces
             {
-                await _blobService.PutObject(TestNamespace, blobData, blobHash);
-                await _blobService.PutObject(TestNamespace, dataChild, blobHashChild);
+                await Service.PutObject(TestNamespace, blobData, blobHash);
+                await Service.PutObject(TestNamespace, dataChild, blobHashChild);
             }
 
             // finalize the object as no pieces is now missing
             {
-                HttpContent requestContent = new ByteArrayContent(Array.Empty<byte>());
+                using HttpContent requestContent = new ByteArrayContent(Array.Empty<byte>());
 
-                HttpResponseMessage result = await _httpClient!.PostAsync(requestUri: $"api/v1/refs/{TestNamespace}/bucket/{key}/finalize/{objectHash}.uecb", requestContent);
+                HttpResponseMessage result = await _httpClient!.PostAsync(new Uri($"api/v1/refs/{TestNamespace}/bucket/{key}/finalize/{objectHash}.uecb", UriKind.Relative), requestContent);
                 result.EnsureSuccessStatusCode();
 
                 {
@@ -1222,11 +1217,11 @@ namespace Horde.Storage.FunctionalTests.References
             BlobIdentifier blobHash = BlobIdentifier.FromBlob(blobData);
             IoHashKey key = IoHashKey.FromName("newReferenceObject");
 
-            HttpContent requestContent = new ByteArrayContent(blobData);
+            using HttpContent requestContent = new ByteArrayContent(blobData);
             requestContent.Headers.ContentType = new MediaTypeHeaderValue(MediaTypeNames.Application.Octet);
             requestContent.Headers.Add(CommonHeaders.HashHeaderName, blobHash.ToString());
 
-            HttpResponseMessage result = await _httpClient!.PutAsync(requestUri: $"api/v1/refs/{TestNamespace}/bucket/{key}.uecb", requestContent);
+            HttpResponseMessage result = await _httpClient!.PutAsync(new Uri($"api/v1/refs/{TestNamespace}/bucket/{key}.uecb", UriKind.Relative), requestContent);
             result.EnsureSuccessStatusCode();
 
             {
@@ -1243,7 +1238,7 @@ namespace Horde.Storage.FunctionalTests.References
 
             {
                 // verify we can fetch the blob properly
-                HttpResponseMessage getResponse = await _httpClient.GetAsync($"api/v1/refs/{TestNamespace}/bucket/{key}.raw");
+                HttpResponseMessage getResponse = await _httpClient.GetAsync(new Uri($"api/v1/refs/{TestNamespace}/bucket/{key}.raw", UriKind.Relative));
                 getResponse.EnsureSuccessStatusCode();
                 await using MemoryStream ms = new MemoryStream();
                 await getResponse.Content.CopyToAsync(ms);
@@ -1256,24 +1251,24 @@ namespace Horde.Storage.FunctionalTests.References
 
             {
                 // delete the blob 
-                await _blobService.DeleteObject(TestNamespace, blobHash);
+                await Service.DeleteObject(TestNamespace, blobHash);
             }
 
             {
                 // the compact binary object still exists, so this returns a success (trying to resolve the attachment will fail)
-                HttpResponseMessage getResponse = await _httpClient.GetAsync($"api/v1/refs/{TestNamespace}/bucket/{key}.uecb");
+                HttpResponseMessage getResponse = await _httpClient.GetAsync(new Uri($"api/v1/refs/{TestNamespace}/bucket/{key}.uecb", UriKind.Relative));
                 getResponse.EnsureSuccessStatusCode();
             }
 
             {
                 // the compact binary object still exists, so this returns a success (trying to resolve the attachment will fail)
-                HttpResponseMessage getResponse = await _httpClient.GetAsync($"api/v1/refs/{TestNamespace}/bucket/{key}.json");
+                HttpResponseMessage getResponse = await _httpClient.GetAsync(new Uri($"api/v1/refs/{TestNamespace}/bucket/{key}.json", UriKind.Relative));
                 getResponse.EnsureSuccessStatusCode();
             }
 
             {
                 // we should now see a 404 as the blob is missing
-                HttpResponseMessage getResponse = await _httpClient.GetAsync($"api/v1/refs/{TestNamespace}/bucket/{key}.raw");
+                HttpResponseMessage getResponse = await _httpClient.GetAsync(new Uri($"api/v1/refs/{TestNamespace}/bucket/{key}.raw", UriKind.Relative));
                 Assert.AreEqual(HttpStatusCode.NotFound, getResponse.StatusCode);
                 Assert.AreEqual("application/problem+json", getResponse.Content.Headers.ContentType!.MediaType);
                 string s = await getResponse.Content.ReadAsStringAsync();
@@ -1283,7 +1278,6 @@ namespace Horde.Storage.FunctionalTests.References
             }
         }
         
-        
         [TestMethod]
         public async Task DeleteObject()
         {
@@ -1291,38 +1285,38 @@ namespace Horde.Storage.FunctionalTests.References
             byte[] data = Encoding.ASCII.GetBytes(objectContents);
             BlobIdentifier objectHash = BlobIdentifier.FromBlob(data);
 
-            HttpContent requestContent = new ByteArrayContent(data);
+            using HttpContent requestContent = new ByteArrayContent(data);
             requestContent.Headers.ContentType = new MediaTypeHeaderValue(MediaTypeNames.Application.Octet);
             requestContent.Headers.Add(CommonHeaders.HashHeaderName, objectHash.ToString());
 
             IoHashKey key = IoHashKey.FromName("deletableObject");
             // submit the object
             {
-                HttpResponseMessage result = await _httpClient!.PutAsync(requestUri: $"api/v1/refs/{TestNamespace}/bucket/{key}", requestContent);
+                HttpResponseMessage result = await _httpClient!.PutAsync(new Uri($"api/v1/refs/{TestNamespace}/bucket/{key}", UriKind.Relative), requestContent);
                 result.EnsureSuccessStatusCode();
             }
 
             // verify it is present
             {
-                HttpResponseMessage result = await _httpClient!.GetAsync(requestUri: $"api/v1/refs/{TestNamespace}/bucket/{key}.raw");
+                HttpResponseMessage result = await _httpClient!.GetAsync(new Uri($"api/v1/refs/{TestNamespace}/bucket/{key}.raw", UriKind.Relative));
                 result.EnsureSuccessStatusCode();
             }
 
             // delete the object
             {
-                HttpResponseMessage result = await _httpClient!.DeleteAsync(requestUri: $"api/v1/refs/{TestNamespace}/bucket/{key}");
+                HttpResponseMessage result = await _httpClient!.DeleteAsync(new Uri($"api/v1/refs/{TestNamespace}/bucket/{key}", UriKind.Relative));
                 result.EnsureSuccessStatusCode();
             }
 
             // ensure the object is not present anymore
             {
-                HttpResponseMessage result = await _httpClient!.GetAsync(requestUri: $"api/v1/refs/{TestNamespace}/bucket/{key}.raw");
+                HttpResponseMessage result = await _httpClient!.GetAsync(new Uri($"api/v1/refs/{TestNamespace}/bucket/{key}.raw", UriKind.Relative));
                 Assert.AreEqual(HttpStatusCode.NotFound, result.StatusCode);
             }
 
             // delete the object again which doesn't exist anymore
             {
-                HttpResponseMessage result = await _httpClient!.DeleteAsync(requestUri: $"api/v1/refs/{TestNamespace}/bucket/{key}");
+                HttpResponseMessage result = await _httpClient!.DeleteAsync(new Uri($"api/v1/refs/{TestNamespace}/bucket/{key}", UriKind.Relative));
                 Assert.AreEqual(HttpStatusCode.NotFound, result.StatusCode);
             }
         }
@@ -1337,7 +1331,7 @@ namespace Horde.Storage.FunctionalTests.References
             byte[] data = Encoding.ASCII.GetBytes(objectContents);
             BlobIdentifier objectHash = BlobIdentifier.FromBlob(data);
 
-            HttpContent requestContent = new ByteArrayContent(data);
+            using HttpContent requestContent = new ByteArrayContent(data);
             requestContent.Headers.ContentType = new MediaTypeHeaderValue(MediaTypeNames.Application.Octet);
             requestContent.Headers.Add(CommonHeaders.HashHeaderName, objectHash.ToString());
 
@@ -1345,42 +1339,42 @@ namespace Horde.Storage.FunctionalTests.References
 
             // submit the object into multiple buckets
             {
-                HttpResponseMessage result = await _httpClient!.PutAsync(requestUri: $"api/v1/refs/{TestNamespace}/{BucketToDelete}/{key}", requestContent);
+                HttpResponseMessage result = await _httpClient!.PutAsync(new Uri($"api/v1/refs/{TestNamespace}/{BucketToDelete}/{key}", UriKind.Relative), requestContent);
                 result.EnsureSuccessStatusCode();
             }
 
             {
-                HttpResponseMessage result = await _httpClient!.PutAsync(requestUri: $"api/v1/refs/{TestNamespace}/bucket/{key}", requestContent);
+                HttpResponseMessage result = await _httpClient!.PutAsync(new Uri($"api/v1/refs/{TestNamespace}/bucket/{key}", UriKind.Relative), requestContent);
                 result.EnsureSuccessStatusCode();
             }
 
             // verify it is present
             {
-                HttpResponseMessage result = await _httpClient!.GetAsync(requestUri: $"api/v1/refs/{TestNamespace}/bucket/{key}.raw");
+                HttpResponseMessage result = await _httpClient!.GetAsync(new Uri($"api/v1/refs/{TestNamespace}/bucket/{key}.raw", UriKind.Relative));
                 result.EnsureSuccessStatusCode();
             }
 
             // verify it is present
             {
-                HttpResponseMessage result = await _httpClient!.GetAsync(requestUri: $"api/v1/refs/{TestNamespace}/{BucketToDelete}/{key}.raw");
+                HttpResponseMessage result = await _httpClient!.GetAsync(new Uri($"api/v1/refs/{TestNamespace}/{BucketToDelete}/{key}.raw", UriKind.Relative));
                 result.EnsureSuccessStatusCode();
             }
 
             // delete the bucket
             {
-                HttpResponseMessage result = await _httpClient!.DeleteAsync(requestUri: $"api/v1/refs/{TestNamespace}/{BucketToDelete}");
+                HttpResponseMessage result = await _httpClient!.DeleteAsync(new Uri($"api/v1/refs/{TestNamespace}/{BucketToDelete}", UriKind.Relative));
                 result.EnsureSuccessStatusCode();
             }
 
             // ensure the object is present in not deleted bucket
             {
-                HttpResponseMessage result = await _httpClient!.GetAsync(requestUri: $"api/v1/refs/{TestNamespace}/bucket/{key}.raw");
+                HttpResponseMessage result = await _httpClient!.GetAsync(new Uri($"api/v1/refs/{TestNamespace}/bucket/{key}.raw", UriKind.Relative));
                 result.EnsureSuccessStatusCode();
             }
 
             // ensure the object is not present anymore
             {
-                HttpResponseMessage result = await _httpClient!.GetAsync(requestUri: $"api/v1/refs/{TestNamespace}/{BucketToDelete}/{key}.raw");
+                HttpResponseMessage result = await _httpClient!.GetAsync(new Uri($"api/v1/refs/{TestNamespace}/{BucketToDelete}/{key}.raw", UriKind.Relative));
                 Assert.AreEqual(HttpStatusCode.NotFound, result.StatusCode);
             }
         }
@@ -1395,55 +1389,55 @@ namespace Horde.Storage.FunctionalTests.References
             byte[] data = Encoding.ASCII.GetBytes(objectContents);
             BlobIdentifier objectHash = BlobIdentifier.FromBlob(data);
 
-            HttpContent requestContent = new ByteArrayContent(data);
+            using HttpContent requestContent = new ByteArrayContent(data);
             requestContent.Headers.ContentType = new MediaTypeHeaderValue(MediaTypeNames.Application.Octet);
             requestContent.Headers.Add(CommonHeaders.HashHeaderName, objectHash.ToString());
 
             IoHashKey key = IoHashKey.FromName("deletableObject");
             // submit the object into multiple namespaces
             {
-                HttpResponseMessage result = await _httpClient!.PutAsync(requestUri: $"api/v1/refs/{TestNamespace}/bucket/{key}", requestContent);
+                HttpResponseMessage result = await _httpClient!.PutAsync(new Uri($"api/v1/refs/{TestNamespace}/bucket/{key}", UriKind.Relative), requestContent);
                 result.EnsureSuccessStatusCode();
             }
 
             {
-                HttpResponseMessage result = await _httpClient!.PutAsync(requestUri: $"api/v1/refs/{NamespaceToBeDeleted}/bucket/{key}", requestContent);
+                HttpResponseMessage result = await _httpClient!.PutAsync(new Uri($"api/v1/refs/{NamespaceToBeDeleted}/bucket/{key}", UriKind.Relative), requestContent);
                 result.EnsureSuccessStatusCode();
             }
 
             // verify it is present
             {
-                HttpResponseMessage result = await _httpClient!.GetAsync(requestUri: $"api/v1/refs/{TestNamespace}/bucket/{key}.raw");
+                HttpResponseMessage result = await _httpClient!.GetAsync(new Uri($"api/v1/refs/{TestNamespace}/bucket/{key}.raw", UriKind.Relative));
                 result.EnsureSuccessStatusCode();
             }
 
             // verify it is present
             {
-                HttpResponseMessage result = await _httpClient!.GetAsync(requestUri: $"api/v1/refs/{NamespaceToBeDeleted}/bucket/{key}.raw");
+                HttpResponseMessage result = await _httpClient!.GetAsync(new Uri($"api/v1/refs/{NamespaceToBeDeleted}/bucket/{key}.raw", UriKind.Relative));
                 result.EnsureSuccessStatusCode();
             }
 
             // delete the namespace
             {
-                HttpResponseMessage result = await _httpClient!.DeleteAsync(requestUri: $"api/v1/refs/{NamespaceToBeDeleted}");
+                HttpResponseMessage result = await _httpClient!.DeleteAsync(new Uri($"api/v1/refs/{NamespaceToBeDeleted}", UriKind.Relative));
                 result.EnsureSuccessStatusCode();
             }
 
             // ensure the object is present in not deleted namespace
             {
-                HttpResponseMessage result = await _httpClient!.GetAsync(requestUri: $"api/v1/refs/{TestNamespace}/bucket/{key}.raw");
+                HttpResponseMessage result = await _httpClient!.GetAsync(new Uri($"api/v1/refs/{TestNamespace}/bucket/{key}.raw", UriKind.Relative));
                 result.EnsureSuccessStatusCode();
             }
 
             // ensure the object is not present anymore
             {
-                HttpResponseMessage result = await _httpClient!.GetAsync(requestUri: $"api/v1/refs/{NamespaceToBeDeleted}/bucket/{key}.raw");
+                HttpResponseMessage result = await _httpClient!.GetAsync(new Uri($"api/v1/refs/{NamespaceToBeDeleted}/bucket/{key}.raw", UriKind.Relative));
                 Assert.AreEqual(HttpStatusCode.NotFound, result.StatusCode);
             }
 
             // make sure the namespace is not considered a valid namespace anymore
             {
-                HttpResponseMessage result = await _httpClient!.GetAsync(requestUri: $"api/v1/refs");
+                HttpResponseMessage result = await _httpClient!.GetAsync(new Uri($"api/v1/refs", UriKind.Relative));
                 result.EnsureSuccessStatusCode();
                 Assert.AreEqual(result!.Content.Headers.ContentType!.MediaType, MediaTypeNames.Application.Json);
                 GetNamespacesResponse response = await result.Content.ReadAsAsync<GetNamespacesResponse>();
@@ -1451,7 +1445,6 @@ namespace Horde.Storage.FunctionalTests.References
             }
         }
 
-
         
         [TestMethod]
         public async Task ListNamespaces()
@@ -1461,25 +1454,24 @@ namespace Horde.Storage.FunctionalTests.References
             BlobIdentifier objectHash = BlobIdentifier.FromBlob(data);
             IoHashKey key = IoHashKey.FromName("notUsedObject");
 
-            HttpContent requestContent = new ByteArrayContent(data);
+            using HttpContent requestContent = new ByteArrayContent(data);
             requestContent.Headers.ContentType = new MediaTypeHeaderValue(MediaTypeNames.Application.Octet);
             requestContent.Headers.Add(CommonHeaders.HashHeaderName, objectHash.ToString());
 
             // submit a object to make sure a namespace is created
             {
-                HttpResponseMessage result = await _httpClient!.PutAsync(requestUri: $"api/v1/refs/{TestNamespace}/bucket/{key}", requestContent);
+                HttpResponseMessage result = await _httpClient!.PutAsync(new Uri($"api/v1/refs/{TestNamespace}/bucket/{key}", UriKind.Relative), requestContent);
                 result.EnsureSuccessStatusCode();
             }
 
             {
-                HttpResponseMessage result = await _httpClient!.GetAsync(requestUri: $"api/v1/refs");
+                HttpResponseMessage result = await _httpClient!.GetAsync(new Uri($"api/v1/refs", UriKind.Relative));
                 result.EnsureSuccessStatusCode();
                 Assert.AreEqual(result!.Content.Headers.ContentType!.MediaType, MediaTypeNames.Application.Json);
                 GetNamespacesResponse response = await result.Content.ReadAsAsync<GetNamespacesResponse>();
                 Assert.IsTrue(response.Namespaces.Contains(TestNamespace));
             }
         }
-
         
         [TestMethod]
         public async Task GetOldRecords()
@@ -1488,14 +1480,14 @@ namespace Horde.Storage.FunctionalTests.References
             byte[] data = Encoding.ASCII.GetBytes(objectContents);
             BlobIdentifier objectHash = BlobIdentifier.FromBlob(data);
 
-            HttpContent requestContent = new ByteArrayContent(data);
+            using HttpContent requestContent = new ByteArrayContent(data);
             requestContent.Headers.ContentType = new MediaTypeHeaderValue(MediaTypeNames.Application.Octet);
             requestContent.Headers.Add(CommonHeaders.HashHeaderName, objectHash.ToString());
 
             IoHashKey key = IoHashKey.FromName("oldRecord");
             // submit some contents
             {
-                HttpResponseMessage result = await _httpClient!.PutAsync(requestUri: $"api/v1/refs/{TestNamespace}/bucket/{key}", requestContent);
+                HttpResponseMessage result = await _httpClient!.PutAsync(new Uri($"api/v1/refs/{TestNamespace}/bucket/{key}", UriKind.Relative), requestContent);
                 result.EnsureSuccessStatusCode();
             }
 
@@ -1506,17 +1498,16 @@ namespace Horde.Storage.FunctionalTests.References
             Assert.AreEqual(key, oldRecord.Value.Item2);
             Assert.AreEqual("bucket", oldRecord.Value.Item1.ToString());
         }
-
         
         [TestMethod]
         public async Task BatchJsonRequest()
         {
             // verifies that json request against the batch endpoint will fail
             {
-                HttpContent requestContent = new StringContent("{}");
+                using HttpContent requestContent = new StringContent("{}");
                 requestContent.Headers.ContentType = new MediaTypeHeaderValue(MediaTypeNames.Application.Json);
 
-                HttpResponseMessage result = await _httpClient!.PostAsync(requestUri: $"api/v1/refs/{TestNamespace}", requestContent);
+                HttpResponseMessage result = await _httpClient!.PostAsync(new Uri($"api/v1/refs/{TestNamespace}", UriKind.Relative), requestContent);
 
                 Assert.AreEqual(HttpStatusCode.UnsupportedMediaType, result.StatusCode);
             }
@@ -1568,10 +1559,10 @@ namespace Horde.Storage.FunctionalTests.References
             byte[] batchRequestData = batchRequestWriter.ToByteArray();
 
             {
-                HttpContent requestContent = new ByteArrayContent(batchRequestData);
+                using HttpContent requestContent = new ByteArrayContent(batchRequestData);
                 requestContent.Headers.ContentType = new MediaTypeHeaderValue(CustomMediaTypeNames.UnrealCompactBinary);
 
-                HttpResponseMessage result = await _httpClient!.PostAsync(requestUri: $"api/v1/refs/{TestNamespace}", requestContent);
+                HttpResponseMessage result = await _httpClient!.PostAsync(new Uri($"api/v1/refs/{TestNamespace}", UriKind.Relative), requestContent);
                 result.EnsureSuccessStatusCode();
 
                 await using MemoryStream ms = new MemoryStream();
@@ -1595,7 +1586,6 @@ namespace Horde.Storage.FunctionalTests.References
             }
         }
 
-
         [TestMethod]
         public async Task BatchGetOperations()
         {
@@ -1608,11 +1598,11 @@ namespace Horde.Storage.FunctionalTests.References
                 byte[] cbObjectBytes = newBlobObject.GetView().ToArray();
                 BlobIdentifier blobHash = BlobIdentifier.FromBlob(cbObjectBytes);
 
-                HttpContent requestContent = new ByteArrayContent(cbObjectBytes);
+                using HttpContent requestContent = new ByteArrayContent(cbObjectBytes);
                 requestContent.Headers.ContentType = new MediaTypeHeaderValue(CustomMediaTypeNames.UnrealCompactBinary);
                 requestContent.Headers.Add(CommonHeaders.HashHeaderName, blobHash.ToString());
 
-                HttpResponseMessage result = await _httpClient!.PutAsync(requestUri: $"api/v1/refs/{TestNamespace}/{bucket}/{newBlobObjectKey}.uecb", requestContent);
+                HttpResponseMessage result = await _httpClient!.PutAsync(new Uri($"api/v1/refs/{TestNamespace}/{bucket}/{newBlobObjectKey}.uecb", UriKind.Relative), requestContent);
                 result.EnsureSuccessStatusCode();
                 byte[] content = await result.Content.ReadAsByteArrayAsync();
                 PutObjectResponse? response = CbSerializer.Deserialize<PutObjectResponse?>(content);
@@ -1627,11 +1617,11 @@ namespace Horde.Storage.FunctionalTests.References
             {
                 BlobIdentifier blobHash = BlobIdentifier.FromBlob(blobContents);
 
-                HttpContent requestContent = new ByteArrayContent(blobContents);
+                using HttpContent requestContent = new ByteArrayContent(blobContents);
                 requestContent.Headers.ContentType = new MediaTypeHeaderValue(MediaTypeNames.Application.Octet);
                 requestContent.Headers.Add(CommonHeaders.HashHeaderName, blobHash.ToString());
 
-                HttpResponseMessage result = await _httpClient!.PutAsync(requestUri: $"api/v1/blobs/{TestNamespace}/{blobHash}", requestContent);
+                HttpResponseMessage result = await _httpClient!.PutAsync(new Uri($"api/v1/blobs/{TestNamespace}/{blobHash}", UriKind.Relative), requestContent);
                 result.EnsureSuccessStatusCode();
             }
 
@@ -1642,11 +1632,11 @@ namespace Horde.Storage.FunctionalTests.References
             {
                 BlobIdentifier blobHash = BlobIdentifier.FromBlob(blobContents);
 
-                HttpContent requestContent = new ByteArrayContent(blobContents);
+                using HttpContent requestContent = new ByteArrayContent(blobContents);
                 requestContent.Headers.ContentType = new MediaTypeHeaderValue(MediaTypeNames.Application.Octet);
                 requestContent.Headers.Add(CommonHeaders.HashHeaderName, blobHash.ToString());
 
-                HttpResponseMessage result = await _httpClient!.PutAsync(requestUri: $"api/v1/blobs/{TestNamespace}/{blobHash}", requestContent);
+                HttpResponseMessage result = await _httpClient!.PutAsync(new Uri($"api/v1/blobs/{TestNamespace}/{blobHash}", UriKind.Relative), requestContent);
                 result.EnsureSuccessStatusCode();
             }
 
@@ -1654,11 +1644,11 @@ namespace Horde.Storage.FunctionalTests.References
                 byte[] cbObjectBytes = newReferenceObject.GetView().ToArray();
                 BlobIdentifier blobHash = BlobIdentifier.FromBlob(cbObjectBytes);
 
-                HttpContent requestContent = new ByteArrayContent(cbObjectBytes);
+                using HttpContent requestContent = new ByteArrayContent(cbObjectBytes);
                 requestContent.Headers.ContentType = new MediaTypeHeaderValue(CustomMediaTypeNames.UnrealCompactBinary);
                 requestContent.Headers.Add(CommonHeaders.HashHeaderName, blobHash.ToString());
 
-                HttpResponseMessage result = await _httpClient!.PutAsync(requestUri: $"api/v1/refs/{TestNamespace}/{bucket}/{newReferenceObjectKey}.uecb", requestContent);
+                HttpResponseMessage result = await _httpClient!.PutAsync(new Uri($"api/v1/refs/{TestNamespace}/{bucket}/{newReferenceObjectKey}.uecb", UriKind.Relative), requestContent);
                 result.EnsureSuccessStatusCode();
 
                 result.EnsureSuccessStatusCode();
@@ -1711,10 +1701,10 @@ namespace Horde.Storage.FunctionalTests.References
             byte[] batchRequestData = batchRequestWriter.ToByteArray();
 
             {
-                HttpContent requestContent = new ByteArrayContent(batchRequestData);
+                using HttpContent requestContent = new ByteArrayContent(batchRequestData);
                 requestContent.Headers.ContentType = new MediaTypeHeaderValue(CustomMediaTypeNames.UnrealCompactBinary);
 
-                HttpResponseMessage result = await _httpClient!.PostAsync(requestUri: $"api/v1/refs/{TestNamespace}", requestContent);
+                HttpResponseMessage result = await _httpClient!.PostAsync(new Uri($"api/v1/refs/{TestNamespace}", UriKind.Relative), requestContent);
                 result.EnsureSuccessStatusCode();
 
                 await using MemoryStream ms = new MemoryStream();
@@ -1753,11 +1743,11 @@ namespace Horde.Storage.FunctionalTests.References
                 byte[] cbObjectBytes = newBlobObject.GetView().ToArray();
                 BlobIdentifier blobHash = BlobIdentifier.FromBlob(cbObjectBytes);
 
-                HttpContent requestContent = new ByteArrayContent(cbObjectBytes);
+                using HttpContent requestContent = new ByteArrayContent(cbObjectBytes);
                 requestContent.Headers.ContentType = new MediaTypeHeaderValue(CustomMediaTypeNames.UnrealCompactBinary);
                 requestContent.Headers.Add(CommonHeaders.HashHeaderName, blobHash.ToString());
 
-                HttpResponseMessage result = await _httpClient!.PutAsync(requestUri: $"api/v1/refs/{TestNamespace}/{bucket}/{newBlobObjectKey}.uecb", requestContent);
+                HttpResponseMessage result = await _httpClient!.PutAsync(new Uri($"api/v1/refs/{TestNamespace}/{bucket}/{newBlobObjectKey}.uecb", UriKind.Relative), requestContent);
                 result.EnsureSuccessStatusCode();
                 byte[] content = await result.Content.ReadAsByteArrayAsync();
                 PutObjectResponse? response = CbSerializer.Deserialize<PutObjectResponse?>(content);
@@ -1772,11 +1762,11 @@ namespace Horde.Storage.FunctionalTests.References
             {
                 BlobIdentifier blobHash = BlobIdentifier.FromBlob(blobContents);
 
-                HttpContent requestContent = new ByteArrayContent(blobContents);
+                using HttpContent requestContent = new ByteArrayContent(blobContents);
                 requestContent.Headers.ContentType = new MediaTypeHeaderValue(MediaTypeNames.Application.Octet);
                 requestContent.Headers.Add(CommonHeaders.HashHeaderName, blobHash.ToString());
 
-                HttpResponseMessage result = await _httpClient!.PutAsync(requestUri: $"api/v1/blobs/{TestNamespace}/{blobHash}", requestContent);
+                HttpResponseMessage result = await _httpClient!.PutAsync(new Uri($"api/v1/blobs/{TestNamespace}/{blobHash}", UriKind.Relative), requestContent);
                 result.EnsureSuccessStatusCode();
             }
 
@@ -1784,11 +1774,11 @@ namespace Horde.Storage.FunctionalTests.References
                 byte[] cbObjectBytes = newReferenceObject.GetView().ToArray();
                 BlobIdentifier blobHash = BlobIdentifier.FromBlob(cbObjectBytes);
 
-                HttpContent requestContent = new ByteArrayContent(cbObjectBytes);
+                using HttpContent requestContent = new ByteArrayContent(cbObjectBytes);
                 requestContent.Headers.ContentType = new MediaTypeHeaderValue(CustomMediaTypeNames.UnrealCompactBinary);
                 requestContent.Headers.Add(CommonHeaders.HashHeaderName, blobHash.ToString());
 
-                HttpResponseMessage result = await _httpClient!.PutAsync(requestUri: $"api/v1/refs/{TestNamespace}/{bucket}/{newReferenceObjectKey}.uecb", requestContent);
+                HttpResponseMessage result = await _httpClient!.PutAsync(new Uri($"api/v1/refs/{TestNamespace}/{bucket}/{newReferenceObjectKey}.uecb", UriKind.Relative), requestContent);
                 result.EnsureSuccessStatusCode();
 
                 result.EnsureSuccessStatusCode();
@@ -1830,10 +1820,10 @@ namespace Horde.Storage.FunctionalTests.References
             byte[] batchRequestData = batchRequestWriter.ToByteArray();
             
             {
-                HttpContent requestContent = new ByteArrayContent(batchRequestData);
+                using HttpContent requestContent = new ByteArrayContent(batchRequestData);
                 requestContent.Headers.ContentType = new MediaTypeHeaderValue(CustomMediaTypeNames.UnrealCompactBinary);
 
-                HttpResponseMessage result = await _httpClient!.PostAsync(requestUri: $"api/v1/refs/{TestNamespace}", requestContent);
+                HttpResponseMessage result = await _httpClient!.PostAsync(new Uri($"api/v1/refs/{TestNamespace}", UriKind.Relative), requestContent);
                 result.EnsureSuccessStatusCode();
 
                 await using MemoryStream ms = new MemoryStream();
@@ -1905,10 +1895,10 @@ namespace Horde.Storage.FunctionalTests.References
             byte[] batchRequestData = batchRequestWriter.ToByteArray();
 
             {
-                HttpContent requestContent = new ByteArrayContent(batchRequestData);
+                using HttpContent requestContent = new ByteArrayContent(batchRequestData);
                 requestContent.Headers.ContentType = new MediaTypeHeaderValue(CustomMediaTypeNames.UnrealCompactBinary);
 
-                HttpResponseMessage result = await _httpClient!.PostAsync(requestUri: $"api/v1/refs/{TestNamespace}", requestContent);
+                HttpResponseMessage result = await _httpClient!.PostAsync(new Uri($"api/v1/refs/{TestNamespace}", UriKind.Relative), requestContent);
                 result.EnsureSuccessStatusCode();
 
                 await using MemoryStream ms = new MemoryStream();
@@ -1949,11 +1939,11 @@ namespace Horde.Storage.FunctionalTests.References
                 byte[] cbObjectBytes = getObject.GetView().ToArray();
                 BlobIdentifier blobHash = BlobIdentifier.FromBlob(cbObjectBytes);
 
-                HttpContent requestContent = new ByteArrayContent(cbObjectBytes);
+                using HttpContent requestContent = new ByteArrayContent(cbObjectBytes);
                 requestContent.Headers.ContentType = new MediaTypeHeaderValue(CustomMediaTypeNames.UnrealCompactBinary);
                 requestContent.Headers.Add(CommonHeaders.HashHeaderName, blobHash.ToString());
 
-                HttpResponseMessage result = await _httpClient!.PutAsync(requestUri: $"api/v1/refs/{TestNamespace}/{bucket}/{getObjectKey}.uecb", requestContent);
+                HttpResponseMessage result = await _httpClient!.PutAsync(new Uri($"api/v1/refs/{TestNamespace}/{bucket}/{getObjectKey}.uecb", UriKind.Relative), requestContent);
                 result.EnsureSuccessStatusCode();
                 byte[] content = await result.Content.ReadAsByteArrayAsync();
                 PutObjectResponse? response = CbSerializer.Deserialize<PutObjectResponse?>(content);
@@ -2015,10 +2005,10 @@ namespace Horde.Storage.FunctionalTests.References
             byte[] batchRequestData = batchRequestWriter.ToByteArray();
             
             {
-                HttpContent requestContent = new ByteArrayContent(batchRequestData);
+                using HttpContent requestContent = new ByteArrayContent(batchRequestData);
                 requestContent.Headers.ContentType = new MediaTypeHeaderValue(CustomMediaTypeNames.UnrealCompactBinary);
 
-                HttpResponseMessage result = await _httpClient!.PostAsync(requestUri: $"api/v1/refs/{TestNamespace}", requestContent);
+                HttpResponseMessage result = await _httpClient!.PostAsync(new Uri($"api/v1/refs/{TestNamespace}", UriKind.Relative), requestContent);
                 result.EnsureSuccessStatusCode();
 
                 await using MemoryStream ms = new MemoryStream();
diff --git a/Engine/Source/Programs/Horde/HordeStorage/Horde.Storage/Tests/Functional/Objects/ReplicationServerTests.cs b/Engine/Source/Programs/Horde/HordeStorage/Horde.Storage/Tests/Functional/Objects/ReplicationServerTests.cs
--- a/Engine/Source/Programs/Horde/HordeStorage/Horde.Storage/Tests/Functional/Objects/ReplicationServerTests.cs
+++ b/Engine/Source/Programs/Horde/HordeStorage/Horde.Storage/Tests/Functional/Objects/ReplicationServerTests.cs
@@ -24,8 +24,8 @@ namespace Horde.Storage.FunctionalTests.References
     {
         private static IWebHost? _server;
 
-        protected readonly NamespaceId TestNamespace = new NamespaceId("test-namespace");
-        protected readonly BucketId TestBucket = new BucketId("default");
+        private readonly NamespaceId TestNamespace = new NamespaceId("test-namespace");
+        private readonly BucketId TestBucket = new BucketId("default");
 
         [TestInitialize]
         public async Task Setup()
@@ -69,21 +69,21 @@ namespace Horde.Storage.FunctionalTests.References
             byte[] objectData = writer.ToByteArray();
             BlobIdentifier objectHash = BlobIdentifier.FromBlob(objectData);
 
-            HttpContent requestContent = new ByteArrayContent(objectData);
+            using HttpContent requestContent = new ByteArrayContent(objectData);
             requestContent.Headers.ContentType = new MediaTypeHeaderValue(CustomMediaTypeNames.UnrealCompactBinary);
             requestContent.Headers.Add(CommonHeaders.HashHeaderName, objectHash.ToString());
 
             {
                 using HttpClient httpClient = new HttpClient();
                 httpClient.BaseAddress = new Uri("http://localhost:8080");
-                HttpResponseMessage result = await httpClient!.PutAsync(requestUri: $"api/v1/refs/{TestNamespace}/{TestBucket}/{IoHashKey.FromName("newReferenceObject")}.uecb", requestContent);
+                HttpResponseMessage result = await httpClient!.PutAsync(new Uri($"api/v1/refs/{TestNamespace}/{TestBucket}/{IoHashKey.FromName("newReferenceObject")}.uecb", UriKind.Relative), requestContent);
                 result.EnsureSuccessStatusCode();
             }
             // access the replication-log using a public endpoint and a internal endpoint, only the internal endpoint should return results
             {
                 using HttpClient httpClient = new HttpClient();
                 httpClient.BaseAddress = new Uri("http://localhost:8080");
-                HttpResponseMessage result = await httpClient!.GetAsync(requestUri: $"api/v1/replication-log/incremental/{TestNamespace}");
+                HttpResponseMessage result = await httpClient!.GetAsync(new Uri($"api/v1/replication-log/incremental/{TestNamespace}", UriKind.Relative));
                 result.EnsureSuccessStatusCode();
             }
 
@@ -91,7 +91,7 @@ namespace Horde.Storage.FunctionalTests.References
                 using HttpClient httpClient = new HttpClient();
                 httpClient.BaseAddress = new Uri("http://localhost:80");
 
-                HttpResponseMessage result = await httpClient!.GetAsync(requestUri: $"api/v1/replication-log/incremental/{TestNamespace}");
+                HttpResponseMessage result = await httpClient!.GetAsync(new Uri($"api/v1/replication-log/incremental/{TestNamespace}", UriKind.Relative));
                 Assert.AreEqual(HttpStatusCode.NotFound, result.StatusCode);
             }
         }
diff --git a/Engine/Source/Programs/Horde/HordeStorage/Horde.Storage/Tests/Functional/Objects/ReplicationTests.cs b/Engine/Source/Programs/Horde/HordeStorage/Horde.Storage/Tests/Functional/Objects/ReplicationTests.cs
--- a/Engine/Source/Programs/Horde/HordeStorage/Horde.Storage/Tests/Functional/Objects/ReplicationTests.cs
+++ b/Engine/Source/Programs/Horde/HordeStorage/Horde.Storage/Tests/Functional/Objects/ReplicationTests.cs
@@ -51,7 +51,7 @@ namespace Horde.Storage.FunctionalTests.References
             IReplicationLog replicationLog = provider.GetService<IReplicationLog>()!;
             Assert.IsTrue(replicationLog.GetType() == typeof(ScyllaReplicationLog));
 
-            await SeedTestData(referencesStore);
+            await SeedTestData();
         }
 
         protected override async Task TeardownDb(IServiceProvider provider)
@@ -68,12 +68,9 @@ namespace Horde.Storage.FunctionalTests.References
                 // remove the namespaces
                 localKeyspace.ExecuteAsync(new SimpleStatement("DROP TABLE IF EXISTS replication_namespace;"))
             );
-
-
         }
     }
 
-
     [TestClass]
     public class MemoryReplicationTests : ReplicationTests
     {
@@ -95,7 +92,7 @@ namespace Horde.Storage.FunctionalTests.References
             IReplicationLog replicationLog = provider.GetService<IReplicationLog>()!;
             Assert.IsTrue(replicationLog.GetType() == typeof(MemoryReplicationLog));
 
-            await SeedTestData(referencesStore);
+            await SeedTestData();
         }
 
         protected override Task TeardownDb(IServiceProvider provider)
@@ -108,13 +105,12 @@ namespace Horde.Storage.FunctionalTests.References
     {
         private static TestServer? _server;
         private static HttpClient? _httpClient;
-        protected IBlobService _blobStore = null!;
-        protected IReferencesStore _referencesStore = null!;
+        private IBlobService _blobStore = null!;
         private IReplicationLog _replicationLog = null!;
 
-        protected readonly NamespaceId TestNamespace = new NamespaceId("test-namespace");
-        protected readonly NamespaceId SnapshotNamespace = new NamespaceId("snapshot-namespace");
-        protected readonly BucketId TestBucket = new BucketId("default");
+        private readonly NamespaceId TestNamespace = new NamespaceId("test-namespace");
+        private readonly NamespaceId SnapshotNamespace = new NamespaceId("snapshot-namespace");
+        private readonly BucketId TestBucket = new BucketId("default");
 
         [TestInitialize]
         public async Task Setup()
@@ -141,23 +137,23 @@ namespace Horde.Storage.FunctionalTests.References
             _server = server;
 
             _blobStore = _server.Services.GetService<IBlobService>()!;
-            _referencesStore = _server.Services.GetService<IReferencesStore>()!;
             _replicationLog = _server.Services.GetService<IReplicationLog>()!;
 
             await SeedDb(server.Services);
         }
 
-        protected async Task SeedTestData(IReferencesStore referencesStore)
+        protected virtual async Task SeedTestData()
         {
             await Task.CompletedTask;
         }
 
-
         [TestCleanup]
         public async Task Teardown()
         {
-            if (_server != null) 
+            if (_server != null)
+            {
                 await TeardownDb(_server.Services);
+            }
         }
 
         protected abstract IEnumerable<KeyValuePair<string, string>> GetSettings();
@@ -177,29 +173,28 @@ namespace Horde.Storage.FunctionalTests.References
             byte[] objectData = writer.ToByteArray();
             BlobIdentifier objectHash = BlobIdentifier.FromBlob(objectData);
 
-            HttpContent requestContent = new ByteArrayContent(objectData);
+            using HttpContent requestContent = new ByteArrayContent(objectData);
             requestContent.Headers.ContentType = new MediaTypeHeaderValue(CustomMediaTypeNames.UnrealCompactBinary);
             requestContent.Headers.Add(CommonHeaders.HashHeaderName, objectHash.ToString());
 
             {
-                HttpResponseMessage result = await _httpClient!.PutAsync(requestUri: $"api/v1/refs/{TestNamespace}/{TestBucket}/{IoHashKey.FromName("newReferenceObject")}.uecb", requestContent);
+                HttpResponseMessage result = await _httpClient!.PutAsync(new Uri($"api/v1/refs/{TestNamespace}/{TestBucket}/{IoHashKey.FromName("newReferenceObject")}.uecb", UriKind.Relative), requestContent);
                 result.EnsureSuccessStatusCode();
             }
 
             {
-                HttpResponseMessage result = await _httpClient!.PutAsync(requestUri: $"api/v1/refs/{TestNamespace}/{TestBucket}/{IoHashKey.FromName("secondObject")}.uecb", requestContent);
+                HttpResponseMessage result = await _httpClient!.PutAsync(new Uri($"api/v1/refs/{TestNamespace}/{TestBucket}/{IoHashKey.FromName("secondObject")}.uecb", UriKind.Relative), requestContent);
                 result.EnsureSuccessStatusCode();
             }
 
             {
-                HttpResponseMessage result = await _httpClient!.PutAsync(requestUri: $"api/v1/refs/{TestNamespace}/{TestBucket}/{IoHashKey.FromName("thirdObject")}.uecb", requestContent);
+                HttpResponseMessage result = await _httpClient!.PutAsync(new Uri($"api/v1/refs/{TestNamespace}/{TestBucket}/{IoHashKey.FromName("thirdObject")}.uecb", UriKind.Relative), requestContent);
                 result.EnsureSuccessStatusCode();
             }
             
             {
-                HttpResponseMessage result = await _httpClient!.GetAsync(requestUri: $"api/v1/replication-log/incremental/{TestNamespace}");
+                HttpResponseMessage result = await _httpClient!.GetAsync(new Uri($"api/v1/replication-log/incremental/{TestNamespace}", UriKind.Relative));
                 result.EnsureSuccessStatusCode();
-                string s = await result.Content.ReadAsStringAsync();
                 
                 Assert.AreEqual(result!.Content.Headers.ContentType!.MediaType, MediaTypeNames.Application.Json);
 
@@ -255,9 +250,8 @@ namespace Horde.Storage.FunctionalTests.References
             await _replicationLog.InsertAddEvent(TestNamespace, TestBucket, IoHashKey.FromName("fourthObject"), objectHash, oldestTimestamp.AddDays(0.9));
 
             {
-                HttpResponseMessage result = await _httpClient!.GetAsync(requestUri: $"api/v1/replication-log/incremental/{TestNamespace}?lastBucket={eventBucket}&lastEvent={eventId}");
+                HttpResponseMessage result = await _httpClient!.GetAsync(new Uri($"api/v1/replication-log/incremental/{TestNamespace}?lastBucket={eventBucket}&lastEvent={eventId}", UriKind.Relative));
                 result.EnsureSuccessStatusCode();
-                string s = await result.Content.ReadAsStringAsync();
                 
                 Assert.AreEqual(result!.Content.Headers.ContentType!.MediaType, MediaTypeNames.Application.Json);
 
@@ -317,9 +311,8 @@ namespace Horde.Storage.FunctionalTests.References
             await _replicationLog.InsertAddEvent(TestNamespace, TestBucket, IoHashKey.FromName("sixthObject"), objectHash, oldestTimestamp.AddDays(2.13));
 
             {
-                HttpResponseMessage result = await _httpClient!.GetAsync(requestUri: $"api/v1/replication-log/incremental/{TestNamespace}?lastBucket={eventBucket}&lastEvent={eventId}");
+                HttpResponseMessage result = await _httpClient!.GetAsync(new Uri($"api/v1/replication-log/incremental/{TestNamespace}?lastBucket={eventBucket}&lastEvent={eventId}", UriKind.Relative));
                 result.EnsureSuccessStatusCode();
-                string s = await result.Content.ReadAsStringAsync();
                 
                 Assert.AreEqual(result!.Content.Headers.ContentType!.MediaType, MediaTypeNames.Application.Json);
 
@@ -376,17 +369,16 @@ namespace Horde.Storage.FunctionalTests.References
             await _replicationLog.InsertAddEvent(TestNamespace, TestBucket, IoHashKey.FromName("fourthObject"), objectHash, oldestTimestamp.AddDays(0.9));
 
             // start from the second event
-            const int eventsToFetch = 2;
+            const int EventsToFetch = 2;
             {
-                HttpResponseMessage result = await _httpClient!.GetAsync(requestUri: $"api/v1/replication-log/incremental/{TestNamespace}?lastBucket={eventBucket}&lastEvent={eventId}&count={eventsToFetch}");
+                HttpResponseMessage result = await _httpClient!.GetAsync(new Uri($"api/v1/replication-log/incremental/{TestNamespace}?lastBucket={eventBucket}&lastEvent={eventId}&count={EventsToFetch}", UriKind.Relative));
                 result.EnsureSuccessStatusCode();
-                string s = await result.Content.ReadAsStringAsync();
                 
                 Assert.AreEqual(result!.Content.Headers.ContentType!.MediaType, MediaTypeNames.Application.Json);
 
                 ReplicationLogEvents? events = await result.Content.ReadAsAsync<ReplicationLogEvents>();
                 Assert.IsNotNull(events);
-                Assert.AreEqual(eventsToFetch, events!.Events.Count);
+                Assert.AreEqual(EventsToFetch, events!.Events.Count);
 
                 // parse the events returned, make sure they are in the right order
                 {
@@ -427,11 +419,9 @@ namespace Horde.Storage.FunctionalTests.References
             Guid eventId = Guid.NewGuid();
 
             {
-                HttpResponseMessage result = await _httpClient!.GetAsync(requestUri: $"api/v1/replication-log/incremental/{TestNamespace}?lastBucket={eventBucket}&lastEvent={eventId}");
+                HttpResponseMessage result = await _httpClient!.GetAsync(new Uri($"api/v1/replication-log/incremental/{TestNamespace}?lastBucket={eventBucket}&lastEvent={eventId}", UriKind.Relative));
                 Assert.AreEqual(HttpStatusCode.BadRequest, result.StatusCode);
                 
-                string s = await result.Content.ReadAsStringAsync();
-                
                 ProblemDetails? problem = await result.Content.ReadFromJsonAsync<ProblemDetails?>();
                 Assert.IsNotNull(problem);
             }
@@ -457,11 +447,9 @@ namespace Horde.Storage.FunctionalTests.References
             Guid eventId = Guid.NewGuid();
 
             {
-                HttpResponseMessage result = await _httpClient!.GetAsync(requestUri: $"api/v1/replication-log/incremental/{TestNamespace}?lastBucket={eventBucket}&lastEvent={eventId}");
+                HttpResponseMessage result = await _httpClient!.GetAsync(new Uri($"api/v1/replication-log/incremental/{TestNamespace}?lastBucket={eventBucket}&lastEvent={eventId}", UriKind.Relative));
                 Assert.AreEqual(HttpStatusCode.BadRequest, result.StatusCode);
                 
-                string s = await result.Content.ReadAsStringAsync();
-                
                 ProblemDetails? problem = await result.Content.ReadFromJsonAsync<ProblemDetails?>();
                 Assert.IsNotNull(problem);
             }
@@ -474,10 +462,8 @@ namespace Horde.Storage.FunctionalTests.References
         {
             // the namespace does not exist
             {
-                HttpResponseMessage result = await _httpClient!.GetAsync(requestUri: $"api/v1/replication-log/incremental/{TestNamespace}");
+                HttpResponseMessage result = await _httpClient!.GetAsync(new Uri($"api/v1/replication-log/incremental/{TestNamespace}", UriKind.Relative));
                 Assert.AreEqual(HttpStatusCode.NotFound, result.StatusCode);
-                string s = await result.Content.ReadAsStringAsync();
-                
                 ProblemDetails? problem = await result.Content.ReadFromJsonAsync<ProblemDetails?>();
                 Assert.IsNotNull(problem);
             }
@@ -511,17 +497,14 @@ namespace Horde.Storage.FunctionalTests.References
             // use a bucket that does not exist, should raise a message to use a snapshot instead
             string bucketThatDoesNotExist = "rep-0000";
             {
-                HttpResponseMessage result = await _httpClient!.GetAsync(requestUri: $"api/v1/replication-log/incremental/{TestNamespace}?lastBucket={bucketThatDoesNotExist}&lastEvent={Guid.NewGuid()}");
+                HttpResponseMessage result = await _httpClient!.GetAsync(new Uri($"api/v1/replication-log/incremental/{TestNamespace}?lastBucket={bucketThatDoesNotExist}&lastEvent={Guid.NewGuid()}", UriKind.Relative));
                 Assert.AreEqual(HttpStatusCode.BadRequest, result.StatusCode);
-                string s = await result.Content.ReadAsStringAsync();
-                
                 ProblemDetails? problem = await result.Content.ReadFromJsonAsync<ProblemDetails?>();
                 Assert.IsNotNull(problem);
                 Assert.AreEqual("http://jupiter.epicgames.com/replication/useSnapshot", problem!.Type);
                 Assert.IsTrue(problem.Extensions.ContainsKey("SnapshotId"));
                 Assert.AreEqual(snapshotBlobId, new BlobIdentifier(problem.Extensions["SnapshotId"]!.ToString()!));
             }
-            
         }
 
         [TestMethod]
@@ -599,14 +582,13 @@ namespace Horde.Storage.FunctionalTests.References
             Assert.AreEqual(lastEventId, snapshot.LastEvent);
             Assert.AreEqual(lastEventBucket, snapshot.LastBucket);
 
-            List<ReplicationLogSnapshot.SnapshotLiveObject> liveObjects = snapshot.GetLiveObjects().ToList();
+            List<SnapshotLiveObject> liveObjects = snapshot.GetLiveObjects().ToList();
             Assert.IsTrue(liveObjects.Any(o => o.Bucket == TestBucket && o.Key == IoHashKey.FromName("firstObject")));
             Assert.IsTrue(liveObjects.Any(o => o.Bucket == TestBucket && o.Key == IoHashKey.FromName("secondObject")));
             Assert.IsTrue(liveObjects.Any(o => o.Bucket == TestBucket && o.Key == IoHashKey.FromName("thirdObject")));
             Assert.IsTrue(liveObjects.Any(o => o.Bucket == TestBucket && o.Key == IoHashKey.FromName("fourthObject")));
         }
 
-
         [TestMethod]
         public async Task ReplicationLogSnapshotQuerying()
         {
@@ -642,10 +624,8 @@ namespace Horde.Storage.FunctionalTests.References
 
             // make sure the snapshot is returned by the rest api
             {
-                HttpResponseMessage result = await _httpClient!.GetAsync(requestUri: $"api/v1/replication-log/snapshots/{TestNamespace}");
+                HttpResponseMessage result = await _httpClient!.GetAsync(new Uri($"api/v1/replication-log/snapshots/{TestNamespace}", UriKind.Relative));
                 result.EnsureSuccessStatusCode();
-                string s = await result.Content.ReadAsStringAsync();
-                
                 Assert.AreEqual(result!.Content.Headers.ContentType!.MediaType, MediaTypeNames.Application.Json);
 
                 ReplicationLogSnapshots snapshots = await result.Content.ReadAsAsync<ReplicationLogSnapshots>();
@@ -706,7 +686,7 @@ namespace Horde.Storage.FunctionalTests.References
 
             // verify the new events can be found when resuming from the snapshot
             {
-                HttpResponseMessage result = await _httpClient!.GetAsync(requestUri: $"api/v1/replication-log/incremental/{TestNamespace}?lastBucket={snapshot.LastBucket}&lastEvent={snapshot.LastEvent}");
+                HttpResponseMessage result = await _httpClient!.GetAsync(new Uri($"api/v1/replication-log/incremental/{TestNamespace}?lastBucket={snapshot.LastBucket}&lastEvent={snapshot.LastEvent}", UriKind.Relative));
                 result.EnsureSuccessStatusCode();
                 
                 Assert.AreEqual(result!.Content.Headers.ContentType!.MediaType, MediaTypeNames.Application.Json);
diff --git a/Engine/Source/Programs/Horde/HordeStorage/Horde.Storage/Tests/Functional/PeerStatusTests.cs b/Engine/Source/Programs/Horde/HordeStorage/Horde.Storage/Tests/Functional/PeerStatusTests.cs
--- a/Engine/Source/Programs/Horde/HordeStorage/Horde.Storage/Tests/Functional/PeerStatusTests.cs
+++ b/Engine/Source/Programs/Horde/HordeStorage/Horde.Storage/Tests/Functional/PeerStatusTests.cs
@@ -1,5 +1,6 @@
 // Copyright Epic Games, Inc. All Rights Reserved.
 
+using System;
 using System.Linq;
 using System.Net;
 using System.Net.Http;
@@ -33,12 +34,12 @@ namespace Horde.Storage.FunctionalTests.Status
                         {
                             new PeerEndpoints()
                             {
-                                Url = "http://siteA.com/internal",
+                                Url = new Uri("http://siteA.com/internal"),
                                 IsInternal = true
                             },
                             new PeerEndpoints()
                             {
-                                Url = "http://siteA.com/public"
+                                Url = new Uri("http://siteA.com/public")
                             },
                         }.ToList()
                     },
@@ -50,12 +51,12 @@ namespace Horde.Storage.FunctionalTests.Status
                         {
                             new PeerEndpoints()
                             {
-                                Url = "http://siteB.com/internal",
+                                Url = new Uri("http://siteB.com/internal"),
                                 IsInternal = true
                             },
                             new PeerEndpoints()
                             {
-                                Url = "http://siteB.com/public"
+                                Url = new Uri("http://siteB.com/public")
                             },
                         }.ToList()
                     },
@@ -76,19 +77,19 @@ namespace Horde.Storage.FunctionalTests.Status
             handler.SetupRequest("http://siteb.com/internal" + endpoint).ReturnsResponse(HttpStatusCode.OK).Callback(() => Task.Delay(400).Wait()).Verifiable();
 
             IHttpClientFactory httpClientFactory = handler.CreateClientFactory();
-            PeerStatusService statusService = new PeerStatusService(settingsMock, jupiterSettingsMock, httpClientFactory);
+            await using PeerStatusService statusService = new(settingsMock, jupiterSettingsMock, httpClientFactory);
 
             await statusService.UpdatePeerStatus(CancellationToken.None);
 
             handler.Verify();
 
             // as we are actually measuring the time it takes to call our mocked handler the latency expected is not going to be exact so we allow some delta
-            IPeerStatusService.PeerStatus? siteAPeerStatus = statusService.GetPeerStatus("siteA");
+            PeerStatus? siteAPeerStatus = statusService.GetPeerStatus("siteA");
             Assert.IsNotNull(siteAPeerStatus);
             Assert.AreEqual(220, siteAPeerStatus.Latency, 50);
             Assert.IsTrue(siteAPeerStatus.Reachable);
 
-            IPeerStatusService.PeerStatus? siteBPeerStatus = statusService.GetPeerStatus("siteB");
+            PeerStatus? siteBPeerStatus = statusService.GetPeerStatus("siteB");
             Assert.IsNotNull(siteBPeerStatus);
             Assert.AreEqual(420, siteBPeerStatus.Latency, 50);
             Assert.IsTrue(siteBPeerStatus.Reachable);
@@ -109,12 +110,12 @@ namespace Horde.Storage.FunctionalTests.Status
                         {
                             new PeerEndpoints()
                             {
-                                Url = "http://siteA.com/internal",
+                                Url = new Uri("http://siteA.com/internal"),
                                 IsInternal = true
                             },
                             new PeerEndpoints()
                             {
-                                Url = "http://siteA.com/public"
+                                Url = new Uri("http://siteA.com/public")
                             },
                         }.ToList()
                     },
@@ -126,12 +127,12 @@ namespace Horde.Storage.FunctionalTests.Status
                         {
                             new PeerEndpoints()
                             {
-                                Url = "http://siteB.com/internal",
+                                Url = new Uri("http://siteB.com/internal"),
                                 IsInternal = true
                             },
                             new PeerEndpoints()
                             {
-                                Url = "http://siteB.com/public"
+                                Url = new Uri("http://siteB.com/public")
                             },
                         }.ToList()
                     },
@@ -152,20 +153,20 @@ namespace Horde.Storage.FunctionalTests.Status
             handler.SetupRequest("http://siteb.com/internal" + endpoint).Throws<SocketException>().Verifiable();
 
             IHttpClientFactory httpClientFactory = handler.CreateClientFactory();
-            PeerStatusService statusService = new PeerStatusService(settingsMock, jupiterSettingsMock, httpClientFactory);
+            await using PeerStatusService statusService = new(settingsMock, jupiterSettingsMock, httpClientFactory);
 
             await statusService.UpdatePeerStatus(CancellationToken.None);
 
             handler.Verify();
 
             // as we are actually measuring the time it takes to call our mocked handler the latency expected is not going to be exact so we allow some delta
-            IPeerStatusService.PeerStatus? siteAPeerStatus = statusService.GetPeerStatus("siteA");
+            PeerStatus? siteAPeerStatus = statusService.GetPeerStatus("siteA");
             Assert.IsNotNull(siteAPeerStatus);
             // verify that we ignore the failing internal endpoint
             Assert.AreEqual(220, siteAPeerStatus.Latency, 50);
             Assert.IsTrue(siteAPeerStatus.Reachable);
 
-            IPeerStatusService.PeerStatus? siteBPeerStatus = statusService.GetPeerStatus("siteB");
+            PeerStatus? siteBPeerStatus = statusService.GetPeerStatus("siteB");
             Assert.IsNotNull(siteBPeerStatus);
             Assert.AreEqual(int.MaxValue, siteBPeerStatus.Latency);
             Assert.IsFalse(siteBPeerStatus.Reachable);
