Horde Storage - Updated editorconfig to more closely match what we run for Horde and updated code to fix all static analysis errors

diff --git a/Engine/Source/Programs/Horde/HordeStorage/Horde.Storage/Tests/Functional/RefTests.cs b/Engine/Source/Programs/Horde/HordeStorage/Horde.Storage/Tests/Functional/RefTests.cs
--- a/Engine/Source/Programs/Horde/HordeStorage/Horde.Storage/Tests/Functional/RefTests.cs
+++ b/Engine/Source/Programs/Horde/HordeStorage/Horde.Storage/Tests/Functional/RefTests.cs
@@ -8,7 +8,6 @@ using System.Net;
 using System.Net.Http;
 using System.Net.Http.Headers;
 using System.Text;
-using System.Threading;
 using System.Threading.Tasks;
 using Amazon.DynamoDBv2;
 using Amazon.DynamoDBv2.Model;
@@ -59,7 +58,6 @@ namespace Horde.Storage.FunctionalTests.Ref
                 }
             }
 
-
             //verify we are using the expected refs store
             IRefsStore? refStore = provider.GetService<IRefsStore>();
             Assert.IsTrue(RefStoreIs(refStore, typeof(MongoRefsStore)));
@@ -78,9 +76,11 @@ namespace Horde.Storage.FunctionalTests.Ref
 
         protected override async Task<RefRecord> GetTestRecord(BucketId bucket, KeyId name, IRefsStore.ExtraFieldsFlag fields)
         {
-            var record = await _refsStore!.Get(TestNamespace, bucket, name, fields);
+            RefRecord? record = await _refsStore!.Get(TestNamespace, bucket, name, fields);
             if (record == null)
+            {
                 throw new Exception("Unable to find record");
+            }
 
             return record;
         }
@@ -122,9 +122,11 @@ namespace Horde.Storage.FunctionalTests.Ref
 
         protected override async Task<RefRecord> GetTestRecord(BucketId bucket, KeyId name, IRefsStore.ExtraFieldsFlag fields)
         {
-            var record = await _refsStore!.Get(TestNamespace, bucket, name, fields);
+            RefRecord? record = await _refsStore!.Get(TestNamespace, bucket, name, fields);
             if (record == null)
+            {
                 throw new Exception("Unable to find record");
+            }
 
             return record;
         }
@@ -135,7 +137,7 @@ namespace Horde.Storage.FunctionalTests.Ref
     }
 
     [TestClass]
-    public class DynamoRefTests : RefTests
+    public class DynamoRefTests : RefTests, IDisposable
     {
         private DynamoDbRefsStore? _refStore;
         private readonly string[] _tablesToDelete = new[] {"Europa_Cache", "Europa_Namespace"};
@@ -148,10 +150,9 @@ namespace Horde.Storage.FunctionalTests.Ref
         protected override async Task SeedDb(IServiceProvider provider)
         {
             IOptionsMonitor<DynamoDbSettings> dynamoSettings = provider.GetService<IOptionsMonitor<DynamoDbSettings>>()!;
-            IAmazonDynamoDB client = provider.GetService<IAmazonDynamoDB>()!;
 
             //verify we are using the expected refs store
-            var cacheStore = provider.GetService<IRefsStore>();
+            IRefsStore? cacheStore = provider.GetService<IRefsStore>();
             Assert.IsTrue(RefStoreIs(cacheStore, typeof(DynamoDbRefsStore)));
 
             _refStore = new DynamoDbRefsStore(dynamoSettings, provider.GetService<IAmazonDynamoDB>()!, provider.GetService<DynamoNamespaceStore>()!);
@@ -177,12 +178,28 @@ namespace Horde.Storage.FunctionalTests.Ref
 
         protected override async Task<RefRecord> GetTestRecord(BucketId bucket, KeyId name, IRefsStore.ExtraFieldsFlag fields)
         {
-            var record = await _refStore!.Get(TestNamespace, bucket, name, fields);
+            RefRecord? record = await _refStore!.Get(TestNamespace, bucket, name, fields);
             if (record == null)
+            {
                 throw new Exception("Unable to find record");
+            }
 
             return record;
         }
+
+        protected virtual void Dispose(bool disposing)
+        {
+            if (disposing)
+            {
+                _refStore?.Dispose();
+            }
+        }
+
+        public void Dispose()
+        {
+            Dispose(true);
+            System.GC.SuppressFinalize(this);
+        }
     }
     
     public abstract class RefTests
@@ -191,13 +208,13 @@ namespace Horde.Storage.FunctionalTests.Ref
         private static HttpClient? _httpClient;
         private static ILastAccessCache<RefRecord>? _lastAccessCache;
 
-        protected static readonly string TestObjectContents = "Test Object Data";
-        protected static readonly byte[] TestObjectData = Encoding.ASCII.GetBytes(TestObjectContents);
-        protected static readonly ContentHash TestObjectHash = ContentHash.FromBlob(TestObjectData);
-        protected static readonly BlobIdentifier TestObjectBlobHash = BlobIdentifier.FromContentHash(TestObjectHash);
-        
-        protected readonly NamespaceId TestNamespace = new NamespaceId("test-namespace");
-        protected readonly NamespaceId LastAccessTestNamespace = new NamespaceId("test-namespace-last-access");
+        private const string TestObjectContents = "Test Object Data";
+        private static readonly byte[] TestObjectData = Encoding.ASCII.GetBytes(TestObjectContents);
+        private static readonly ContentHash TestObjectHash = ContentHash.FromBlob(TestObjectData);
+        private static readonly BlobIdentifier TestObjectBlobHash = BlobIdentifier.FromContentHash(TestObjectHash);
+
+        protected NamespaceId TestNamespace { get; } = new NamespaceId("test-namespace");
+        protected NamespaceId LastAccessTestNamespace { get; } = new NamespaceId("test-namespace-last-access");
 
         [TestInitialize]
         public async Task Setup()
@@ -259,18 +276,21 @@ namespace Horde.Storage.FunctionalTests.Ref
             }
         }
 
-
         [TestCleanup]
         public async Task Teardown()
         {
-            if (_server != null) 
+            if (_server != null)
+            {
                 await TeardownDb(_server.Services);
+            }
         }
 
-        protected bool RefStoreIs(IRefsStore? refStore, Type refStoreType)
+        protected static bool RefStoreIs(IRefsStore? refStore, Type refStoreType)
         {
             if (refStore == null)
+            {
                 return false;
+            }
 
             if (refStore.GetType() == refStoreType)
             {
@@ -296,7 +316,7 @@ namespace Horde.Storage.FunctionalTests.Ref
         [TestMethod]
         public async Task ListNamespaces()
         {
-            HttpResponseMessage response = await _httpClient!.GetAsync("api/v1/c/ddc");
+            HttpResponseMessage response = await _httpClient!.GetAsync(new Uri("api/v1/c/ddc", UriKind.Relative));
             response.EnsureSuccessStatusCode();
             GetNamespacesResponse content = await response.Content.ReadAsAsync<GetNamespacesResponse>();
 
@@ -312,9 +332,7 @@ namespace Horde.Storage.FunctionalTests.Ref
             // fetch refs records once first so that any previous runs of tests are ignored
             List<(RefRecord, DateTime)> _ = await _lastAccessCache!.GetLastAccessedRecords();
 
-            Task<HttpResponseMessage> response = _httpClient!.GetAsync($"api/v1/c/ddc/{TestNamespace}/bucket/testObject");
-            await response;
-            HttpResponseMessage result = response.Result;
+            HttpResponseMessage result = await _httpClient!.GetAsync(new Uri($"api/v1/c/ddc/{TestNamespace}/bucket/testObject", UriKind.Relative));
 
             result.EnsureSuccessStatusCode();
             RefResponse content = await result.Content.ReadAsAsync<RefResponse>();
@@ -334,7 +352,7 @@ namespace Horde.Storage.FunctionalTests.Ref
         [TestMethod]
         public async Task GetFullRefObjectAsCompactBinary()
         {
-            HttpRequestMessage request = new HttpRequestMessage(HttpMethod.Get, $"api/v1/c/ddc/{TestNamespace}/bucket/testObjectWithMetadata.uecb");
+            using HttpRequestMessage request = new HttpRequestMessage(HttpMethod.Get, new Uri($"api/v1/c/ddc/{TestNamespace}/bucket/testObjectWithMetadata.uecb", UriKind.Relative));
             HttpResponseMessage response = await _httpClient!.SendAsync(request);
 
             Assert.AreEqual(CustomMediaTypeNames.UnrealCompactBinary, response.Content.Headers.ContentType!.MediaType!);
@@ -357,9 +375,8 @@ namespace Horde.Storage.FunctionalTests.Ref
         [TestMethod]
         public async Task GetRefWithMetadata()
         {
-            Task<HttpResponseMessage> response = _httpClient!.GetAsync($"api/v1/c/ddc/{TestNamespace}/bucket/testObjectWithMetadata");
-            await response;
-            HttpResponseMessage result = response.Result;
+            Task<HttpResponseMessage> response = _httpClient!.GetAsync(new Uri($"api/v1/c/ddc/{TestNamespace}/bucket/testObjectWithMetadata", UriKind.Relative));
+            HttpResponseMessage result = await response;
 
             result.EnsureSuccessStatusCode();
             RefResponse content = await result.Content.ReadAsAsync<RefResponse>();
@@ -375,9 +392,8 @@ namespace Horde.Storage.FunctionalTests.Ref
         [TestMethod]
         public async Task GetRefRawObject()
         {
-            Task<HttpResponseMessage> response = _httpClient!.GetAsync($"api/v1/c/ddc/{TestNamespace}/bucket/testObject.raw");
-            await response;
-            HttpResponseMessage result = response.Result;
+            Task<HttpResponseMessage> response = _httpClient!.GetAsync(new Uri($"api/v1/c/ddc/{TestNamespace}/bucket/testObject.raw", UriKind.Relative));
+            HttpResponseMessage result = await response;
 
             result.EnsureSuccessStatusCode();
             byte[] content = await result.Content.ReadAsByteArrayAsync();
@@ -390,7 +406,7 @@ namespace Horde.Storage.FunctionalTests.Ref
             // fetch refs records once first so that any previous runs of tests are ignored
             List<(RefRecord, DateTime)> _ = await _lastAccessCache!.GetLastAccessedRecords();
 
-            Task<HttpResponseMessage> response = _httpClient!.GetAsync($"api/v1/c/ddc/{TestNamespace}/bucket/testObject?fields=name&fields=blobIdentifiers");
+            Task<HttpResponseMessage> response = _httpClient!.GetAsync(new Uri($"api/v1/c/ddc/{TestNamespace}/bucket/testObject?fields=name&fields=blobIdentifiers", UriKind.Relative));
             HttpResponseMessage result = await response;
 
             result.EnsureSuccessStatusCode();
@@ -413,7 +429,7 @@ namespace Horde.Storage.FunctionalTests.Ref
             // fetch refs records once first so that any previous runs of tests are ignored
             List<(RefRecord, DateTime)> _ = await _lastAccessCache!.GetLastAccessedRecords();
 
-            Task<HttpResponseMessage> response = _httpClient!.GetAsync($"api/v1/c/ddc/{TestNamespace}/bucket/testObject?fields=");
+            Task<HttpResponseMessage> response = _httpClient!.GetAsync(new Uri($"api/v1/c/ddc/{TestNamespace}/bucket/testObject?fields=", UriKind.Relative));
             HttpResponseMessage result = await response;
 
             result.EnsureSuccessStatusCode();
@@ -425,11 +441,10 @@ namespace Horde.Storage.FunctionalTests.Ref
             CollectionAssert.AreEqual(TestObjectData, content.Blob);
         }
 
-
         [TestMethod]
         public async Task GetNonExistentIoObject()
         {
-            HttpResponseMessage response = await _httpClient!.GetAsync($"api/v1/c/ddc/{TestNamespace}/bucket/notUploadedToIo");
+            HttpResponseMessage response = await _httpClient!.GetAsync(new Uri($"api/v1/c/ddc/{TestNamespace}/bucket/notUploadedToIo", UriKind.Relative));
 
             Assert.AreEqual(HttpStatusCode.BadRequest, response.StatusCode);
 
@@ -438,7 +453,6 @@ namespace Horde.Storage.FunctionalTests.Ref
             Assert.AreEqual("Object EBA4BE7385F365E6AC2354096A6CA2C5AF412A0F in test-namespace not found", jsonResult!["title"]!.ToString());
         }
 
-
         [TestMethod]
         public async Task PutRefObject()
         {
@@ -449,8 +463,7 @@ namespace Horde.Storage.FunctionalTests.Ref
             requestContent.Headers.ContentLength = payload.Length;
             requestContent.Headers.Add(CommonHeaders.HashHeaderName, payloadHash.ToString());
 
-            HttpResponseMessage result =
-                await _httpClient!.PutAsync(requestUri: $"api/v1/c/ddc/{TestNamespace}/bucket/newObject", requestContent);
+            HttpResponseMessage result = await _httpClient!.PutAsync(new Uri($"api/v1/c/ddc/{TestNamespace}/bucket/newObject", UriKind.Relative), requestContent);
 
             result.EnsureSuccessStatusCode();
         }
@@ -464,10 +477,10 @@ namespace Horde.Storage.FunctionalTests.Ref
             requestContent.Headers.ContentLength = jsonObject.Length;
             requestContent.Headers.Add(CommonHeaders.HashHeaderName, TestObjectHash.ToString());
 
-            HttpResponseMessage result = await _httpClient!.PutAsync(requestUri: $"api/v1/c/ddc/{TestNamespace}/bucket/newStructuredObject", requestContent);
+            HttpResponseMessage result = await _httpClient!.PutAsync(new Uri($"api/v1/c/ddc/{TestNamespace}/bucket/newStructuredObject", UriKind.Relative), requestContent);
             result.EnsureSuccessStatusCode();
 
-            HttpResponseMessage getResponse = await _httpClient.GetAsync($"api/v1/c/ddc/{TestNamespace}/bucket/newStructuredObject.raw");
+            HttpResponseMessage getResponse = await _httpClient.GetAsync(new Uri($"api/v1/c/ddc/{TestNamespace}/bucket/newStructuredObject.raw", UriKind.Relative));
             getResponse.EnsureSuccessStatusCode();
             await using MemoryStream ms = new MemoryStream();
             await getResponse.Content.CopyToAsync(ms);
@@ -479,7 +492,6 @@ namespace Horde.Storage.FunctionalTests.Ref
             CollectionAssert.AreEqual(TestObjectData, roundTrippedBuffer);
         }
 
-
         [TestMethod]
         public async Task PutGetLargeObject()
         {
@@ -492,13 +504,11 @@ namespace Horde.Storage.FunctionalTests.Ref
             requestContent.Headers.ContentLength = payload.Length;
             requestContent.Headers.Add(CommonHeaders.HashHeaderName, payloadHash.ToString());
 
-            HttpResponseMessage result =
-                await _httpClient!.PutAsync(requestUri: $"api/v1/c/ddc/{TestNamespace}/bucket/largeObject", requestContent);
+            HttpResponseMessage result = await _httpClient!.PutAsync(new Uri($"api/v1/c/ddc/{TestNamespace}/bucket/largeObject", UriKind.Relative), requestContent);
 
             result.EnsureSuccessStatusCode();
 
-
-            HttpResponseMessage getResponse = await _httpClient.GetAsync($"api/v1/c/ddc/{TestNamespace}/bucket/largeObject.raw");
+            HttpResponseMessage getResponse = await _httpClient.GetAsync(new Uri($"api/v1/c/ddc/{TestNamespace}/bucket/largeObject.raw", UriKind.Relative));
             getResponse.EnsureSuccessStatusCode();
             await using MemoryStream ms = new MemoryStream();
             await getResponse.Content.CopyToAsync(ms);
@@ -513,20 +523,17 @@ namespace Horde.Storage.FunctionalTests.Ref
             CollectionAssert.AreEqual(payload, roundTrippedBuffer);
         }
 
-
         [TestMethod]
         public async Task DeleteRefObject()
         {
-            HttpResponseMessage deleteResponse =
-                await _httpClient!.DeleteAsync($"api/v1/c/ddc/{TestNamespace}/bucket/deletableObject");
+            HttpResponseMessage deleteResponse = await _httpClient!.DeleteAsync(new Uri($"api/v1/c/ddc/{TestNamespace}/bucket/deletableObject", UriKind.Relative));
             deleteResponse.EnsureSuccessStatusCode();
         }
 
         [TestMethod]
         public async Task DeleteNonExistentObject()
         {
-            HttpResponseMessage deleteResponse =
-                await _httpClient!.DeleteAsync($"api/v1/c/ddc/{TestNamespace}/bucket/thisDoesNotExistForDelete");
+            HttpResponseMessage deleteResponse = await _httpClient!.DeleteAsync(new Uri($"api/v1/c/ddc/{TestNamespace}/bucket/thisDoesNotExistForDelete", UriKind.Relative));
             Assert.AreEqual(System.Net.HttpStatusCode.BadRequest, deleteResponse.StatusCode);
         }
 
@@ -540,14 +547,11 @@ namespace Horde.Storage.FunctionalTests.Ref
             requestContent.Headers.ContentLength = payload.Length;
             requestContent.Headers.Add(CommonHeaders.HashHeaderName, payloadHash.ToString());
 
-            HttpResponseMessage putResponse =
-                await _httpClient!.PutAsync(requestUri: $"api/v1/c/ddc/{TestNamespace}/bucket/fullFlowObject",
-                    requestContent);
+            HttpResponseMessage putResponse = await _httpClient!.PutAsync(new Uri($"api/v1/c/ddc/{TestNamespace}/bucket/fullFlowObject", UriKind.Relative), requestContent);
             putResponse.EnsureSuccessStatusCode();
             await putResponse.Content.ReadAsStringAsync();
 
-            HttpResponseMessage getResponse =
-                await _httpClient.GetAsync($"api/v1/c/ddc/{TestNamespace}/bucket/fullFlowObject");
+            HttpResponseMessage getResponse = await _httpClient.GetAsync(new Uri($"api/v1/c/ddc/{TestNamespace}/bucket/fullFlowObject", UriKind.Relative));
             getResponse.EnsureSuccessStatusCode();
             RefResponse content = await getResponse.Content.ReadAsAsync<RefResponse>();
             Assert.AreEqual(expected: $"{TestNamespace}.bucket.fullFlowObject", content.Name);
@@ -555,8 +559,7 @@ namespace Horde.Storage.FunctionalTests.Ref
             Assert.IsNotNull(content.Blob);
             CollectionAssert.AreEqual(payload, content.Blob);
 
-            HttpResponseMessage deleteResponse =
-                await _httpClient.DeleteAsync($"api/v1/c/ddc/{TestNamespace}/bucket/fullFlowObject");
+            HttpResponseMessage deleteResponse = await _httpClient.DeleteAsync(new Uri($"api/v1/c/ddc/{TestNamespace}/bucket/fullFlowObject", UriKind.Relative));
             deleteResponse.EnsureSuccessStatusCode();
             Assert.AreEqual(System.Net.HttpStatusCode.NoContent, deleteResponse.StatusCode);
         }
@@ -573,15 +576,15 @@ namespace Horde.Storage.FunctionalTests.Ref
             requestContent.Headers.ContentLength = payload.Length;
             requestContent.Headers.Add(CommonHeaders.HashHeaderName, payloadHash.ToString());
 
-            HttpResponseMessage putResponse = await _httpClient!.PutAsync(requestUri: $"api/v1/c/ddc/{deleteNamespace}/bucket/deleteNamespaceObject", requestContent);
+            HttpResponseMessage putResponse = await _httpClient!.PutAsync(new Uri($"api/v1/c/ddc/{deleteNamespace}/bucket/deleteNamespaceObject", UriKind.Relative), requestContent);
             putResponse.EnsureSuccessStatusCode();
             await putResponse.Content.ReadAsStringAsync();
 
-            HttpResponseMessage deleteResponse = await _httpClient.DeleteAsync($"api/v1/c/ddc/{deleteNamespace}");
+            HttpResponseMessage deleteResponse = await _httpClient.DeleteAsync(new Uri($"api/v1/c/ddc/{deleteNamespace}", UriKind.Relative));
             deleteResponse.EnsureSuccessStatusCode();
             Assert.AreEqual(System.Net.HttpStatusCode.NoContent, deleteResponse.StatusCode);
 
-            HttpResponseMessage getResponse = await _httpClient.GetAsync($"api/v1/c/ddc/{TestNamespace}/bucket/deleteNamespaceObject");
+            HttpResponseMessage getResponse = await _httpClient.GetAsync(new Uri($"api/v1/c/ddc/{TestNamespace}/bucket/deleteNamespaceObject", UriKind.Relative));
             Assert.AreEqual(System.Net.HttpStatusCode.BadRequest, getResponse.StatusCode);
         }
 
@@ -599,9 +602,7 @@ namespace Horde.Storage.FunctionalTests.Ref
             requestContent.Headers.ContentLength = payload.Length;
             requestContent.Headers.Add(CommonHeaders.HashHeaderName, payloadHash.ToString());
 
-            HttpResponseMessage putResponse =
-                await _httpClient!.PutAsync(requestUri: $"api/v1/c/ddc/{TestNamespace}/{bucket}/{key}",
-                    requestContent);
+            HttpResponseMessage putResponse = await _httpClient!.PutAsync(new Uri($"api/v1/c/ddc/{TestNamespace}/{bucket}/{key}", UriKind.Relative), requestContent);
             putResponse.EnsureSuccessStatusCode();
             await putResponse.Content.ReadAsStringAsync();
 
@@ -610,15 +611,13 @@ namespace Horde.Storage.FunctionalTests.Ref
             Assert.IsTrue(initialLastAccessObject.LastAccessTime.HasValue);
             DateTime creationTime = initialLastAccessObject.LastAccessTime!.Value;
 
-            HttpResponseMessage getResponse =
-                await _httpClient.GetAsync($"api/v1/c/ddc/{TestNamespace}/{bucket}/{key}");
+            HttpResponseMessage getResponse = await _httpClient.GetAsync(new Uri($"api/v1/c/ddc/{TestNamespace}/{bucket}/{key}", UriKind.Relative));
             getResponse.EnsureSuccessStatusCode();
             RefResponse content = await getResponse.Content.ReadAsAsync<RefResponse>();
             Assert.AreEqual(expected: $"{TestNamespace}.bucket.lastAccessObject", content.Name);
 
-            StringContent jsonContent = new StringContent("", Encoding.UTF8, "application/json");
-            HttpResponseMessage lastAccessRollupResponse =
-                await _httpClient.PostAsync(requestUri: $"api/v1/admin/startLastAccessRollup", jsonContent);
+            using StringContent jsonContent = new StringContent("", Encoding.UTF8, "application/json");
+            HttpResponseMessage lastAccessRollupResponse = await _httpClient.PostAsync(new Uri($"api/v1/admin/startLastAccessRollup", UriKind.Relative), jsonContent);
             lastAccessRollupResponse.EnsureSuccessStatusCode();
 
             RefRecord doc = await GetTestRecord(bucket, key, IRefsStore.ExtraFieldsFlag.LastAccess);
@@ -630,7 +629,6 @@ namespace Horde.Storage.FunctionalTests.Ref
             Assert.AreEqual(0, records.Count);
         }
 
-
         [TestMethod]
         public async Task LastAccess()
         {
@@ -659,23 +657,13 @@ namespace Horde.Storage.FunctionalTests.Ref
             Assert.AreEqual(4, oldRecords.Length);
         }
 
-        public class LastAccessResponse
-        {
-            public class LastAccessRefRecord
-            {
-                public string Name { get; set; } = null!;
-                public string Bucket { get; set; } = null!;
-                public string Namespace { get; set; } = null!;
-            }
-
-            public LastAccessRefRecord[] Records { get; set; } = null!;
-        }
-
         [TestMethod]
         public async Task Batch()
         {
             if (this is DynamoRefTests)
+            {
                 Assert.Inconclusive();
+            }
 
             // we chunk at 1 kb to lets generate a 3kb string
             string content = string.Concat(Enumerable.Repeat("Duplicate string", 200));
@@ -848,7 +836,9 @@ namespace Horde.Storage.FunctionalTests.Ref
                     Assert.AreEqual(3, state);
                 }
                 else
+                {
                     throw new NotImplementedException();
+                }
             }
 
             // we have reached the end of the result thus the stream should be empty
diff --git a/Engine/Source/Programs/Horde/HordeStorage/Horde.Storage/Tests/Functional/StatusTests.cs b/Engine/Source/Programs/Horde/HordeStorage/Horde.Storage/Tests/Functional/StatusTests.cs
--- a/Engine/Source/Programs/Horde/HordeStorage/Horde.Storage/Tests/Functional/StatusTests.cs
+++ b/Engine/Source/Programs/Horde/HordeStorage/Horde.Storage/Tests/Functional/StatusTests.cs
@@ -1,6 +1,6 @@
 // Copyright Epic Games, Inc. All Rights Reserved.
 
-using System.Collections.Generic;
+using System;
 using System.Linq;
 using System.Net.Http;
 using System.Threading.Tasks;
@@ -9,8 +9,6 @@ using Microsoft.AspNetCore.TestHost;
 using Microsoft.AspNetCore.Hosting;
 using Microsoft.Extensions.Configuration;
 using Microsoft.Extensions.DependencyInjection;
-using Microsoft.Extensions.DependencyInjection.Extensions;
-using Microsoft.Extensions.Options;
 using Microsoft.VisualStudio.TestTools.UnitTesting;
 using Serilog;
 using Serilog.Core;
@@ -18,10 +16,10 @@ using Serilog.Core;
 namespace Horde.Storage.FunctionalTests.Status
 {
     [TestClass]
-    public class StatusTests
+    public class StatusTests: IDisposable
     {
-        protected TestServer? _server;
-        protected HttpClient? _httpClient;
+        private HttpClient? _httpClient;
+        private TestServer? _server;
 
         [TestInitialize]
         public async Task Setup()
@@ -36,7 +34,7 @@ namespace Horde.Storage.FunctionalTests.Status
                 .ReadFrom.Configuration(configuration)
                 .CreateLogger();
 
-            TestServer server = new TestServer(new WebHostBuilder()
+            _server = new TestServer(new WebHostBuilder()
                 .UseConfiguration(configuration)
                 .UseEnvironment("Testing")
                 .UseSerilog(logger)
@@ -52,8 +50,8 @@ namespace Horde.Storage.FunctionalTests.Status
                                 FullName = "us-east-1",
                                 Endpoints = new PeerEndpoints[]
                                 {
-                                    new PeerEndpoints { Url = "http://use-internal.jupiter.com", IsInternal = true },
-                                    new PeerEndpoints { Url = "http://use.jupiter.com", IsInternal = false },
+                                    new PeerEndpoints { Url = new Uri("http://use-internal.jupiter.com"), IsInternal = true },
+                                    new PeerEndpoints { Url = new Uri("http://use.jupiter.com"), IsInternal = false },
                                 }.ToList()
                             }
                         }.ToList();
@@ -61,8 +59,7 @@ namespace Horde.Storage.FunctionalTests.Status
                 })
                 .UseStartup<HordeStorageStartup>()
             );
-            _httpClient = server.CreateClient();
-            _server = server;
+            _httpClient = _server.CreateClient();
 
             await Task.CompletedTask;
         }
@@ -70,7 +67,7 @@ namespace Horde.Storage.FunctionalTests.Status
         [TestMethod]
         public async Task GetPeerConnection()
         {
-            HttpResponseMessage result = await _httpClient!.GetAsync($"api/v1/status/peers");
+            HttpResponseMessage result = await _httpClient!.GetAsync(new Uri($"api/v1/status/peers", UriKind.Relative));
             result.EnsureSuccessStatusCode();
             PeersResponse peersResponse = await result.Content.ReadAsAsync<PeersResponse>();
             Assert.AreEqual("test", peersResponse.CurrentSite);
@@ -80,14 +77,14 @@ namespace Horde.Storage.FunctionalTests.Status
             Assert.AreEqual("use", peersResponse.Peers[0].Site);
             Assert.AreEqual("us-east-1", peersResponse.Peers[0].FullName);
             Assert.AreEqual(1, peersResponse.Peers[0].Endpoints.Count);
-            Assert.AreEqual("http://use.jupiter.com", peersResponse.Peers[0].Endpoints[0]);
+            Assert.AreEqual(new Uri("http://use.jupiter.com"), peersResponse.Peers[0].Endpoints[0]);
         }
 
         
         [TestMethod]
         public async Task GetPeerConnectionInternal()
         {
-            HttpResponseMessage result = await _httpClient!.GetAsync($"api/v1/status/peers?includeInternalEndpoints=true");
+            HttpResponseMessage result = await _httpClient!.GetAsync(new Uri($"api/v1/status/peers?includeInternalEndpoints=true", UriKind.Relative));
             result.EnsureSuccessStatusCode();
             PeersResponse peersResponse = await result.Content.ReadAsAsync<PeersResponse>();
             Assert.AreEqual("test", peersResponse.CurrentSite);
@@ -98,8 +95,23 @@ namespace Horde.Storage.FunctionalTests.Status
             Assert.AreEqual("us-east-1", peersResponse.Peers[0].FullName);
             Assert.AreEqual(2, peersResponse.Peers[0].Endpoints.Count);
 
-            Assert.AreEqual("http://use-internal.jupiter.com", peersResponse.Peers[0].Endpoints[0]);
-            Assert.AreEqual("http://use.jupiter.com", peersResponse.Peers[0].Endpoints[1]);
+            Assert.AreEqual(new Uri("http://use-internal.jupiter.com"), peersResponse.Peers[0].Endpoints[0]);
+            Assert.AreEqual(new Uri("http://use.jupiter.com"), peersResponse.Peers[0].Endpoints[1]);
+        }
+
+        protected virtual void Dispose(bool disposing)
+        {
+            if (disposing)
+            {
+                _httpClient?.Dispose();
+                _server?.Dispose();
+            }
+        }
+
+        public void Dispose()
+        {
+            Dispose(true);
+            System.GC.SuppressFinalize(this);
         }
     }
 }
diff --git a/Engine/Source/Programs/Horde/HordeStorage/Horde.Storage/Tests/Functional/StorageTests.cs b/Engine/Source/Programs/Horde/HordeStorage/Horde.Storage/Tests/Functional/StorageTests.cs
--- a/Engine/Source/Programs/Horde/HordeStorage/Horde.Storage/Tests/Functional/StorageTests.cs
+++ b/Engine/Source/Programs/Horde/HordeStorage/Horde.Storage/Tests/Functional/StorageTests.cs
@@ -26,7 +26,6 @@ using Jupiter.Implementation;
 using Microsoft.AspNetCore.TestHost;
 using Microsoft.AspNetCore.Hosting;
 using Microsoft.AspNetCore.Mvc;
-using Microsoft.AspNetCore.Mvc.Formatters;
 using Microsoft.Extensions.Configuration;
 using Microsoft.Extensions.DependencyInjection;
 using Microsoft.Extensions.Options;
@@ -65,10 +64,10 @@ namespace Horde.Storage.FunctionalTests.Storage
             }
 
             await _s3.PutBucketAsync(S3BucketName);
-            await _s3.PutObjectAsync(new PutObjectRequest { BucketName = S3BucketName, Key = _smallFileHash.AsS3Key(), ContentBody = SmallFileContents });
-            await _s3.PutObjectAsync(new PutObjectRequest { BucketName = S3BucketName, Key = _anotherFileHash.AsS3Key(), ContentBody = AnotherFileContents });
-            await _s3.PutObjectAsync(new PutObjectRequest { BucketName = S3BucketName, Key = _deleteFileHash.AsS3Key(), ContentBody = DeletableFileContents });
-            await _s3.PutObjectAsync(new PutObjectRequest {BucketName = S3BucketName, Key = _oldBlobFileHash.AsS3Key(), ContentBody = OldFileContents});
+            await _s3.PutObjectAsync(new PutObjectRequest { BucketName = S3BucketName, Key = SmallFileHash.AsS3Key(), ContentBody = SmallFileContents });
+            await _s3.PutObjectAsync(new PutObjectRequest { BucketName = S3BucketName, Key = AnotherFileHash.AsS3Key(), ContentBody = AnotherFileContents });
+            await _s3.PutObjectAsync(new PutObjectRequest { BucketName = S3BucketName, Key = DeleteFileHash.AsS3Key(), ContentBody = DeletableFileContents });
+            await _s3.PutObjectAsync(new PutObjectRequest {BucketName = S3BucketName, Key = OldBlobFileHash.AsS3Key(), ContentBody = OldFileContents});
         }
 
         protected override async Task Teardown(IServiceProvider provider)
@@ -106,22 +105,22 @@ namespace Horde.Storage.FunctionalTests.Storage
             }
 
             await container.CreateAsync();
-            BlobClient smallBlob = container.GetBlobClient(_smallFileHash.ToString());
+            BlobClient smallBlob = container.GetBlobClient(SmallFileHash.ToString());
             byte[] smallContents = Encoding.ASCII.GetBytes(SmallFileContents);
             await using MemoryStream smallBlobStream = new MemoryStream(smallContents);
             await smallBlob.UploadAsync(smallBlobStream);
 
-            BlobClient anotherBlob = container.GetBlobClient(_anotherFileHash.ToString());
+            BlobClient anotherBlob = container.GetBlobClient(AnotherFileHash.ToString());
             byte[] anotherContents = Encoding.ASCII.GetBytes(AnotherFileContents);
             await using MemoryStream anotherContentsStream = new MemoryStream(anotherContents);
             await anotherBlob.UploadAsync(anotherContentsStream);
 
-            BlobClient deleteBlob = container.GetBlobClient(_deleteFileHash.ToString());
+            BlobClient deleteBlob = container.GetBlobClient(DeleteFileHash.ToString());
             byte[] deleteContents = Encoding.ASCII.GetBytes(DeletableFileContents);
             await using MemoryStream deleteContentsStream = new MemoryStream(deleteContents);
             await deleteBlob.UploadAsync(deleteContentsStream);
 
-            BlobClient oldBlob = container.GetBlobClient(_oldBlobFileHash.ToString());
+            BlobClient oldBlob = container.GetBlobClient(OldBlobFileHash.ToString());
             byte[] oldBlobContents = Encoding.ASCII.GetBytes(OldFileContents);
             await using MemoryStream oldBlobContentsSteam = new MemoryStream(oldBlobContents);
             await oldBlob.UploadAsync(oldBlobContentsSteam);
@@ -140,7 +139,7 @@ namespace Horde.Storage.FunctionalTests.Storage
     [TestClass]
     public class FileSystemStoreTests : StorageTests
     {
-        private string _localTestDir;
+        private readonly string _localTestDir;
         private readonly NamespaceId FooNamespace = new NamespaceId("foo");
 
         public FileSystemStoreTests()
@@ -155,28 +154,27 @@ namespace Horde.Storage.FunctionalTests.Storage
             };
         }
 
-
         protected override async Task Seed(IServiceProvider provider)
         {
             NamespaceId folderName = TestNamespaceName;
             
             Directory.CreateDirectory(_localTestDir);
 
-            FileInfo smallFileInfo = FileSystemStore.GetFilesystemPath(_localTestDir, folderName, _smallFileHash);
+            FileInfo smallFileInfo = FileSystemStore.GetFilesystemPath(_localTestDir, folderName, SmallFileHash);
             smallFileInfo.Directory?.Create();
             await File.WriteAllBytesAsync(
                 smallFileInfo.FullName,
                 Encoding.ASCII.GetBytes(SmallFileContents)
             );
 
-            FileInfo anotherFileInfo = FileSystemStore.GetFilesystemPath(_localTestDir, folderName, _anotherFileHash);
+            FileInfo anotherFileInfo = FileSystemStore.GetFilesystemPath(_localTestDir, folderName, AnotherFileHash);
             anotherFileInfo.Directory?.Create();
             await File.WriteAllBytesAsync(
                 anotherFileInfo.FullName,
                 Encoding.ASCII.GetBytes(AnotherFileContents)
             );
             
-            FileInfo deleteFileInfo = FileSystemStore.GetFilesystemPath(_localTestDir, folderName, _deleteFileHash);
+            FileInfo deleteFileInfo = FileSystemStore.GetFilesystemPath(_localTestDir, folderName, DeleteFileHash);
             deleteFileInfo.Directory?.Create();
             await File.WriteAllBytesAsync(
                 deleteFileInfo.FullName,
@@ -184,7 +182,7 @@ namespace Horde.Storage.FunctionalTests.Storage
             );
 
             // a old file used to verify cutoff filtering
-            FileInfo oldFileInfo = FileSystemStore.GetFilesystemPath(_localTestDir, folderName, _oldBlobFileHash);
+            FileInfo oldFileInfo = FileSystemStore.GetFilesystemPath(_localTestDir, folderName, OldBlobFileHash);
             oldFileInfo.Directory?.Create();
             await File.WriteAllBytesAsync(
                 oldFileInfo.FullName,
@@ -204,7 +202,7 @@ namespace Horde.Storage.FunctionalTests.Storage
         [TestMethod]
         public async Task ListOldBlobs()
         {
-            FileSystemStore? fsStore = _server!.Services.GetService<FileSystemStore>();
+            FileSystemStore? fsStore = Server!.Services.GetService<FileSystemStore>();
             Assert.IsNotNull(fsStore);
 
             // we fetch all objects that are more then a day old
@@ -212,20 +210,20 @@ namespace Horde.Storage.FunctionalTests.Storage
             DateTime cutoff = DateTime.Now.AddDays(-1);
             {
                 BlobIdentifier[] blobs = await fsStore.ListObjects(TestNamespaceName).Where(tuple => tuple.Item2 < cutoff).Select(tuple => tuple.Item1).ToArrayAsync();
-                Assert.AreEqual(_oldBlobFileHash, blobs[0]);
+                Assert.AreEqual(OldBlobFileHash, blobs[0]);
             }
         }
 
         [TestMethod]
         public async Task ListNamespaces()
         {
-            FileSystemStore? fsStore = _server!.Services.GetService<FileSystemStore>();
+            FileSystemStore? fsStore = Server!.Services.GetService<FileSystemStore>();
             Assert.IsNotNull(fsStore);
 
             List<NamespaceId> namespaces = await fsStore.ListNamespaces().ToListAsync();
             Assert.AreEqual(1, namespaces.Count);
             
-            await fsStore.PutObject(FooNamespace, Encoding.ASCII.GetBytes(SmallFileContents), _smallFileHash);
+            await fsStore.PutObject(FooNamespace, Encoding.ASCII.GetBytes(SmallFileContents), SmallFileHash);
             
             namespaces = await fsStore.ListNamespaces().ToListAsync();
             Assert.AreEqual(2, namespaces.Count);
@@ -238,12 +236,12 @@ namespace Horde.Storage.FunctionalTests.Storage
             Directory.Delete(_localTestDir, true);
             Directory.CreateDirectory(_localTestDir);
             
-            FilesystemSettings fsSettings = _server!.Services.GetService<IOptionsMonitor<FilesystemSettings>>()!.CurrentValue;
+            FilesystemSettings fsSettings = Server!.Services.GetService<IOptionsMonitor<FilesystemSettings>>()!.CurrentValue;
             fsSettings.MaxSizeBytes = 600;
-            FileSystemStore? fsStore = _server!.Services.GetService<FileSystemStore>();
+            FileSystemStore? fsStore = Server!.Services.GetService<FileSystemStore>();
             Assert.IsNotNull(fsStore);
 
-            CancellationTokenSource cts = new CancellationTokenSource();
+            using CancellationTokenSource cts = new CancellationTokenSource();
             Assert.IsTrue(await fsStore.CleanupInternal(cts.Token, batchSize: 2) == 0); // No garbage to collect, should return false
             
             FileInfo[] fooFiles = CreateFilesInNamespace(FooNamespace, 10);
@@ -273,7 +271,7 @@ namespace Horde.Storage.FunctionalTests.Storage
             FileInfo[] fooFiles = CreateFilesInNamespace(FooNamespace, 10);
             CreateFilesInNamespace(new NamespaceId("bar"), 10);
 
-            FileSystemStore? fsStore = _server!.Services.GetService<FileSystemStore>();
+            FileSystemStore? fsStore = Server!.Services.GetService<FileSystemStore>();
             Assert.IsNotNull(fsStore);
 
             Assert.AreEqual(10, (fsStore.GetLeastRecentlyAccessedObjects(FooNamespace)).ToArray().Length);
@@ -288,10 +286,10 @@ namespace Horde.Storage.FunctionalTests.Storage
         [TestMethod]
         public async Task CalculateUsedDiskSpace()
         {
-            FileSystemStore? fsStore = _server!.Services.GetService<FileSystemStore>();
+            FileSystemStore? fsStore = Server!.Services.GetService<FileSystemStore>();
             Assert.IsNotNull(fsStore);
-            await fsStore.PutObject(FooNamespace, Encoding.ASCII.GetBytes(SmallFileContents), _smallFileHash);
-            await fsStore.PutObject(FooNamespace, Encoding.ASCII.GetBytes(AnotherFileContents), _anotherFileHash);
+            await fsStore.PutObject(FooNamespace, Encoding.ASCII.GetBytes(SmallFileContents), SmallFileHash);
+            await fsStore.PutObject(FooNamespace, Encoding.ASCII.GetBytes(AnotherFileContents), AnotherFileHash);
             
             Assert.AreEqual(SmallFileContents.Length + AnotherFileContents.Length, await fsStore.CalculateDiskSpaceUsed(FooNamespace));
         }
@@ -317,19 +315,20 @@ namespace Horde.Storage.FunctionalTests.Storage
 
     public abstract class StorageTests
     {
-        protected TestServer? _server;
-        protected HttpClient? _httpClient;
+        protected TestServer? Server { get; set; }
+        protected NamespaceId TestNamespaceName { get; } = new NamespaceId("testbucket");
+
+        private HttpClient? _httpClient;
 
-        protected readonly NamespaceId TestNamespaceName = new NamespaceId("testbucket");
         protected const string SmallFileContents = "Small file contents";
         protected const string AnotherFileContents = "Another file with contents";
         protected const string DeletableFileContents = "Delete Me";
         protected const string OldFileContents = "a old blob used for testing cutoff filtering";
 
-        protected readonly BlobIdentifier _smallFileHash = BlobIdentifier.FromBlob( Encoding.ASCII.GetBytes(SmallFileContents));
-        protected readonly BlobIdentifier _anotherFileHash = BlobIdentifier.FromBlob(Encoding.ASCII.GetBytes(AnotherFileContents));
-        protected readonly BlobIdentifier _deleteFileHash = BlobIdentifier.FromBlob(Encoding.ASCII.GetBytes(DeletableFileContents));
-        protected readonly BlobIdentifier _oldBlobFileHash = BlobIdentifier.FromBlob(Encoding.ASCII.GetBytes(OldFileContents));
+        protected BlobIdentifier SmallFileHash { get; } = BlobIdentifier.FromBlob(Encoding.ASCII.GetBytes(SmallFileContents));
+        protected BlobIdentifier AnotherFileHash { get; } = BlobIdentifier.FromBlob(Encoding.ASCII.GetBytes(AnotherFileContents));
+        protected BlobIdentifier DeleteFileHash { get; } = BlobIdentifier.FromBlob(Encoding.ASCII.GetBytes(DeletableFileContents));
+        protected BlobIdentifier OldBlobFileHash { get; } = BlobIdentifier.FromBlob(Encoding.ASCII.GetBytes(OldFileContents));
 
         [TestInitialize]
         public async Task Setup()
@@ -352,10 +351,10 @@ namespace Horde.Storage.FunctionalTests.Storage
                 .UseStartup<HordeStorageStartup>()
             );
             _httpClient = server.CreateClient();
-            _server = server;
+            Server = server;
 
             // Seed storage
-            await Seed(_server.Services);
+            await Seed(Server.Services);
         }
 
         protected abstract IEnumerable<KeyValuePair<string, string>> GetSettings();
@@ -363,31 +362,27 @@ namespace Horde.Storage.FunctionalTests.Storage
         protected abstract Task Seed(IServiceProvider serverServices);
         protected abstract Task Teardown(IServiceProvider serverServices);
 
-
         [TestCleanup]
         public async Task MyTeardown()
         {
-            await Teardown(_server!.Services);
+            await Teardown(Server!.Services);
         }
 
-
         [TestMethod]
         public async Task GetSmallFile()
         {
-            HttpResponseMessage result = await _httpClient!.GetAsync($"api/v1/s/{TestNamespaceName}/{_smallFileHash}");
+            HttpResponseMessage result = await _httpClient!.GetAsync(new Uri($"api/v1/s/{TestNamespaceName}/{SmallFileHash}", UriKind.Relative));
             result.EnsureSuccessStatusCode();
             string content = await result.Content.ReadAsStringAsync();
             Assert.AreEqual(SmallFileContents, content);
         }
 
-
-
         [TestMethod]
         public async Task GetNotExistentFile()
         {
             byte[] payload = Encoding.ASCII.GetBytes("This content does not exist");
             ContentHash contentHash = ContentHash.FromBlob(payload);
-            HttpResponseMessage result = await _httpClient!.GetAsync($"api/v1/s/{TestNamespaceName}/{contentHash}");
+            HttpResponseMessage result = await _httpClient!.GetAsync(new Uri($"api/v1/s/{TestNamespaceName}/{contentHash}", UriKind.Relative));
 
             Assert.AreEqual(HttpStatusCode.NotFound, result.StatusCode);
         }
@@ -395,9 +390,7 @@ namespace Horde.Storage.FunctionalTests.Storage
         [TestMethod]
         public async Task GetInvalidHash()
         {
-            Task<HttpResponseMessage> response = _httpClient!.GetAsync($"api/v1/s/{TestNamespaceName}/smallFile");
-            await response;
-            HttpResponseMessage result = response.Result;
+            HttpResponseMessage result = await _httpClient!.GetAsync(new Uri($"api/v1/s/{TestNamespaceName}/smallFile", UriKind.Relative));
 
             Assert.AreEqual(HttpStatusCode.BadRequest, result.StatusCode);
         }
@@ -410,10 +403,7 @@ namespace Horde.Storage.FunctionalTests.Storage
             requestContent.Headers.ContentType = new MediaTypeHeaderValue("application/octet-stream");
             requestContent.Headers.ContentLength = payload.Length;
             BlobIdentifier contentHash = BlobIdentifier.FromBlob(payload);
-            Task<HttpResponseMessage> response =
-                _httpClient!.PutAsync(requestUri: $"api/v1/s/{TestNamespaceName}/{contentHash}", requestContent);
-            await response;
-            HttpResponseMessage result = response.Result;
+            HttpResponseMessage result = await _httpClient!.PutAsync(new Uri($"api/v1/s/{TestNamespaceName}/{contentHash}", UriKind.Relative), requestContent);
 
             result.EnsureSuccessStatusCode();
             InsertResponse content = await result.Content.ReadAsAsync<InsertResponse>();
@@ -428,7 +418,7 @@ namespace Horde.Storage.FunctionalTests.Storage
             requestContent.Headers.ContentType = new MediaTypeHeaderValue("application/octet-stream");
             requestContent.Headers.ContentLength = payload.Length;
             BlobIdentifier contentHash = BlobIdentifier.FromBlob(payload);
-            HttpResponseMessage result = await _httpClient!.PostAsync(requestUri: $"api/v1/s/{TestNamespaceName}", requestContent);
+            HttpResponseMessage result = await _httpClient!.PostAsync(new Uri($"api/v1/s/{TestNamespaceName}", UriKind.Relative), requestContent);
 
             result.EnsureSuccessStatusCode();
             InsertResponse content = await result.Content.ReadAsAsync<InsertResponse>();
@@ -438,7 +428,7 @@ namespace Horde.Storage.FunctionalTests.Storage
         [TestMethod]
         public async Task DeleteBlob()
         {
-            HttpResponseMessage result = await  _httpClient!.DeleteAsync($"api/v1/s/{TestNamespaceName}/{_deleteFileHash}");
+            HttpResponseMessage result = await  _httpClient!.DeleteAsync(new Uri($"api/v1/s/{TestNamespaceName}/{DeleteFileHash}", UriKind.Relative));
             result.EnsureSuccessStatusCode();
 
             Assert.AreEqual(HttpStatusCode.NoContent, result.StatusCode);
@@ -449,13 +439,13 @@ namespace Horde.Storage.FunctionalTests.Storage
         {
             List<BlobIdentifier> validBlobHashes = new List<BlobIdentifier>
             {
-                _smallFileHash,
-                _anotherFileHash,
-                _deleteFileHash,
-                _oldBlobFileHash
+                SmallFileHash,
+                AnotherFileHash,
+                DeleteFileHash,
+                OldBlobFileHash
             };
 
-            IBlobService blobService = _server!.Services.GetService<IBlobService>()!;
+            IBlobService blobService = Server!.Services.GetService<IBlobService>()!;
 
             BlobIdentifier[] oldObjects = await blobService.ListObjects(TestNamespaceName).Select(tuple => tuple.Item1).ToArrayAsync();
 
@@ -470,15 +460,15 @@ namespace Horde.Storage.FunctionalTests.Storage
         public async Task BlobExists()
         {
             {
-                HttpResponseMessage result = await _httpClient!.SendAsync(new HttpRequestMessage(HttpMethod.Head,
-                    requestUri: $"api/v1/s/{TestNamespaceName}/{_smallFileHash}"));
+                using HttpRequestMessage message = new(HttpMethod.Head, new Uri($"api/v1/s/{TestNamespaceName}/{SmallFileHash}", UriKind.Relative));
+                HttpResponseMessage result = await _httpClient!.SendAsync(message);
                 Assert.AreEqual(HttpStatusCode.OK, result.StatusCode);
             }
 
             {
                 ContentHash newContent = ContentHash.FromBlob(Encoding.ASCII.GetBytes("this content has never been submitted"));
-                HttpResponseMessage resultNew = await _httpClient!.SendAsync(new HttpRequestMessage(HttpMethod.Head,
-                    requestUri: $"api/v1/s/{TestNamespaceName}/{newContent}"));
+                using HttpRequestMessage message = new (HttpMethod.Head, new Uri($"api/v1/s/{TestNamespaceName}/{newContent}", UriKind.Relative));
+                HttpResponseMessage resultNew = await _httpClient!.SendAsync(message);
                 Assert.AreEqual(HttpStatusCode.NotFound, resultNew.StatusCode);
                 string content = await resultNew.Content.ReadAsStringAsync();
                 ValidationProblemDetails result = JsonConvert.DeserializeObject<ValidationProblemDetails>(content)!;
@@ -486,7 +476,6 @@ namespace Horde.Storage.FunctionalTests.Storage
             }
         }
 
-
         [TestMethod]
         public async Task BlobExistsBatch()
         {
@@ -500,7 +489,7 @@ namespace Horde.Storage.FunctionalTests.Storage
                     {
                         Op = "HEAD",
                         Namespace = TestNamespaceName,
-                        Id = _smallFileHash,
+                        Id = SmallFileHash,
                     },
                     new
                     {
@@ -523,7 +512,6 @@ namespace Horde.Storage.FunctionalTests.Storage
             Assert.AreEqual(newContent, id);
         }
 
-
         [TestMethod]
         public async Task BatchOp()
         {
@@ -535,13 +523,13 @@ namespace Horde.Storage.FunctionalTests.Storage
                     {
                         Op = "GET",
                         Namespace = TestNamespaceName,
-                        Id = _smallFileHash,
+                        Id = SmallFileHash,
                     },
                     new
                     {
                         Op = "PUT",
                         Namespace = TestNamespaceName,
-                        Id = _deleteFileHash,
+                        Id = DeleteFileHash,
                         Content = Encoding.ASCII.GetBytes(DeletableFileContents)
                     }
                 }
@@ -559,7 +547,7 @@ namespace Horde.Storage.FunctionalTests.Storage
             string convertedContent = Encoding.ASCII.GetString(Convert.FromBase64String(base64Content));
             Assert.AreEqual(SmallFileContents, convertedContent);
             BlobIdentifier identifier = new BlobIdentifier(array[1].Value<string>()!);
-            Assert.AreEqual(_deleteFileHash, identifier);
+            Assert.AreEqual(DeleteFileHash, identifier);
         }
 
         [TestMethod]
@@ -629,8 +617,8 @@ namespace Horde.Storage.FunctionalTests.Storage
         public async Task MultipleBlobChecks()
         {
             BlobIdentifier newContent = BlobIdentifier.FromBlob(Encoding.ASCII.GetBytes("this content has never been submitted"));
-            HttpResponseMessage response = await _httpClient!.SendAsync(new HttpRequestMessage(HttpMethod.Post,
-                requestUri: $"api/v1/s/{TestNamespaceName}/exists?id={_smallFileHash}&id={newContent}"));
+            using HttpRequestMessage message = new(HttpMethod.Post, new Uri($"api/v1/s/{TestNamespaceName}/exists?id={SmallFileHash}&id={newContent}", UriKind.Relative));
+            HttpResponseMessage response = await _httpClient!.SendAsync(message);
             Assert.AreEqual(HttpStatusCode.OK, response.StatusCode);
             Assert.AreEqual("application/json", response.Content.Headers.ContentType!.MediaType);
 
@@ -646,10 +634,10 @@ namespace Horde.Storage.FunctionalTests.Storage
         public async Task MultipleBlobChecksBodyCB()
         {
             BlobIdentifier newContent = BlobIdentifier.FromBlob(Encoding.ASCII.GetBytes("this content has never been submitted"));
-            HttpRequestMessage request = new HttpRequestMessage(HttpMethod.Post, requestUri: $"api/v1/s/{TestNamespaceName}/exist");
+            using HttpRequestMessage request = new(HttpMethod.Post, new Uri($"api/v1/s/{TestNamespaceName}/exist", UriKind.Relative));
             CbWriter writer = new CbWriter();
             writer.BeginUniformArray(CbFieldType.Hash);
-            writer.WriteHashValue(_smallFileHash.AsIoHash());
+            writer.WriteHashValue(SmallFileHash.AsIoHash());
             writer.WriteHashValue(newContent.AsIoHash());
             writer.EndUniformArray();
             byte[] buf = writer.ToByteArray();
@@ -672,8 +660,8 @@ namespace Horde.Storage.FunctionalTests.Storage
         public async Task MultipleBlobChecksBodyJson()
         {
             BlobIdentifier newContent = BlobIdentifier.FromBlob(Encoding.ASCII.GetBytes("this content has never been submitted"));
-            HttpRequestMessage request = new HttpRequestMessage(HttpMethod.Post, requestUri: $"api/v1/s/{TestNamespaceName}/exist");
-            string jsonBody = JsonConvert.SerializeObject(new BlobIdentifier[] { _smallFileHash, newContent });
+            using HttpRequestMessage request = new(HttpMethod.Post, new Uri($"api/v1/s/{TestNamespaceName}/exist", UriKind.Relative));
+            string jsonBody = JsonConvert.SerializeObject(new BlobIdentifier[] { SmallFileHash, newContent });
             request.Content = new StringContent(jsonBody, Encoding.UTF8, MediaTypeNames.Application.Json);
 
             HttpResponseMessage response = await _httpClient!.SendAsync(request);
@@ -692,7 +680,7 @@ namespace Horde.Storage.FunctionalTests.Storage
         public async Task MultipleBlobCompactBinaryResponse()
         {
             BlobIdentifier newContent = BlobIdentifier.FromBlob(Encoding.ASCII.GetBytes("this content has never been submitted"));
-            HttpRequestMessage request = new HttpRequestMessage(HttpMethod.Post, requestUri: $"api/v1/s/{TestNamespaceName}/exists?id={_smallFileHash}&id={newContent}");
+            using HttpRequestMessage request = new(HttpMethod.Post, new Uri($"api/v1/s/{TestNamespaceName}/exists?id={SmallFileHash}&id={newContent}", UriKind.Relative));
             request.Headers.Accept.Add(new MediaTypeWithQualityHeaderValue(CustomMediaTypeNames.UnrealCompactBinary));
             HttpResponseMessage response = await _httpClient!.SendAsync(request);
             Assert.AreEqual(HttpStatusCode.OK, response.StatusCode);
@@ -717,17 +705,16 @@ namespace Horde.Storage.FunctionalTests.Storage
             requestContent.Headers.ContentType = new MediaTypeHeaderValue("application/octet-stream");
             requestContent.Headers.ContentLength = payload.Length;
             BlobIdentifier contentHash = BlobIdentifier.FromBlob(payload);
-            HttpResponseMessage putResponse =
-                await _httpClient!.PutAsync(requestUri: $"api/v1/s/{TestNamespaceName}/{contentHash}", requestContent);
+            HttpResponseMessage putResponse = await _httpClient!.PutAsync(new Uri($"api/v1/s/{TestNamespaceName}/{contentHash}", UriKind.Relative), requestContent);
             putResponse.EnsureSuccessStatusCode();
             InsertResponse response = await putResponse.Content.ReadAsAsync<InsertResponse>();
             Assert.AreEqual(contentHash, response.Identifier);
 
-            HttpResponseMessage getResponse = await _httpClient.GetAsync($"api/v1/s/{TestNamespaceName}/{contentHash}");
+            HttpResponseMessage getResponse = await _httpClient.GetAsync(new Uri($"api/v1/s/{TestNamespaceName}/{contentHash}", UriKind.Relative));
             getResponse.EnsureSuccessStatusCode();
             CollectionAssert.AreEqual(payload, actual: await getResponse.Content.ReadAsByteArrayAsync());
 
-            HttpResponseMessage deleteResponse = await _httpClient.DeleteAsync($"api/v1/s/{TestNamespaceName}/{contentHash}");
+            HttpResponseMessage deleteResponse = await _httpClient.DeleteAsync(new Uri($"api/v1/s/{TestNamespaceName}/{contentHash}", UriKind.Relative));
             deleteResponse.EnsureSuccessStatusCode();
             Assert.AreEqual(HttpStatusCode.NoContent, deleteResponse.StatusCode);
         }
@@ -740,11 +727,14 @@ namespace Horde.Storage.FunctionalTests.Storage
         public async Task PutGetLargePayload()
         {
             // we submit a blob so large that it can not fit using the memory blob store
-            IBlobStore? blobStore = _server?.Services.GetService<IBlobStore>();
+            IBlobStore? blobStore = Server?.Services.GetService<IBlobStore>();
             Assert.IsFalse(blobStore is MemoryBlobStore);
 
             if (blobStore is AzureBlobStore)
+            {
                 Assert.Inconclusive("Azure blob store gets internal server errors when receiving large blobs");
+            }
+
             FileInfo fi = new FileInfo(Path.GetTempFileName());
 
             FileInfo tempOutputFile = new FileInfo(Path.GetTempFileName());
@@ -772,9 +762,9 @@ namespace Horde.Storage.FunctionalTests.Storage
 
                 {
                     await using FileStream fs = fi.OpenRead();
-                    StreamContent content = new StreamContent(fs);
+                    using StreamContent content = new StreamContent(fs);
                     content.Headers.ContentType = new MediaTypeHeaderValue(MediaTypeNames.Application.Octet);
-                    HttpResponseMessage result = await _httpClient!.PutAsync($"api/v1/blobs/{TestNamespaceName}/{blobIdentifier}", content);
+                    HttpResponseMessage result = await _httpClient!.PutAsync(new Uri($"api/v1/blobs/{TestNamespaceName}/{blobIdentifier}", UriKind.Relative), content);
                     result.EnsureSuccessStatusCode();
 
                     InsertResponse response = await result.Content.ReadAsAsync<InsertResponse>();
@@ -783,7 +773,7 @@ namespace Horde.Storage.FunctionalTests.Storage
                 
                 {
                     // verify we can fetch the blob again
-                    HttpResponseMessage result = await _httpClient!.GetAsync($"api/v1/blobs/{TestNamespaceName}/{blobIdentifier}", HttpCompletionOption.ResponseHeadersRead);
+                    HttpResponseMessage result = await _httpClient!.GetAsync(new Uri($"api/v1/blobs/{TestNamespaceName}/{blobIdentifier}", UriKind.Relative), HttpCompletionOption.ResponseHeadersRead);
                     result.EnsureSuccessStatusCode();
                     Stream s = await result.Content.ReadAsStreamAsync();
 
@@ -801,74 +791,24 @@ namespace Horde.Storage.FunctionalTests.Storage
                     Assert.AreEqual(blobIdentifier, downloadedBlobIdentifier);
                     s.Close();
                 }
-
             }
             finally
             {
                 if (fi.Exists)
+                {
                     fi.Delete();
+                }
 
                 if (tempOutputFile.Exists)
+                {
                     tempOutputFile.Delete();
+                }
             }
         }
-
-        /// <summary>
-        /// Uploads a compact binary with a lot of attachments that are missing
-        /// </summary>
-        /// <returns></returns>
-        /*[TestMethod]
-        public async Task LargeCBUpload()
-        {
-            // we submit a blob so large that it can not fit using the memory blob store
-            IBlobStore? blobStore = _server?.Services.GetService<IBlobStore>();
-            Assert.IsFalse(blobStore is MemoryBlobStore);
-
-            if (blobStore is AzureBlobStore)
-                Assert.Inconclusive("Azure blob store gets internal server errors when receiving large blobs");
-            CompactBinaryWriter writer = new CompactBinaryWriter();
-            writer.BeginObject();
-            const int countOfFields = 100_000;
-            for (int i = 0; i < countOfFields; i++)
-            {
-                string s = i.ToString("x40");
-                writer.AddBinaryAttachment(new BlobIdentifier(s));
-            }
-            writer.EndObject();
-
-            byte[] objectData = writer.Save();
-            string path = Path.GetTempFileName();
-            File.WriteAllBytes(path, objectData);
-
-            Console.WriteLine("CB Object output to: " + path);
-            BlobIdentifier objectHash = BlobIdentifier.FromBlob(objectData);
-
-            HttpContent requestContent = new ByteArrayContent(objectData);
-            requestContent.Headers.ContentType = new MediaTypeHeaderValue(CustomMediaTypeNames.UnrealCompactBinary);
-            requestContent.Headers.Add(CommonHeaders.HashHeaderName, objectHash.ToString());
-
-            HttpResponseMessage result = await _httpClient!.PutAsync(requestUri: $"api/v1/refs/{TestNamespaceName}/bucket/newLargeRef.uecb", requestContent);
-            result.EnsureSuccessStatusCode();
-
-            {
-                Assert.AreEqual(CustomMediaTypeNames.UnrealCompactBinary, result!.Content.Headers.ContentType!.MediaType);
-                // check that no blobs are missing
-                await using MemoryStream ms = new MemoryStream();
-                await result.Content.CopyToAsync(ms);
-                byte[] roundTrippedBuffer = ms.ToArray();
-                ReadOnlyMemory<byte> localMemory = new ReadOnlyMemory<byte>(roundTrippedBuffer);
-                CompactBinaryObject cb = CompactBinaryObject.Load(ref localMemory);
-                CompactBinaryField? needsField = cb["needs"];
-                Assert.IsNotNull(needsField);
-                List<BlobIdentifier?> missingBlobs = needsField!.AsArray().Select(field => field.AsHash()).ToList();
-                Assert.AreEqual(countOfFields, missingBlobs.Count);
-            }
-        }*/
     }
 
-
     public class InsertResponse
     {
-        public BlobIdentifier? Identifier;
+        public BlobIdentifier? Identifier { get; set; }
     }
 }
diff --git a/Engine/Source/Programs/Horde/HordeStorage/Horde.Storage/Tests/UnitTests/BlobCleanupServiceTest.cs b/Engine/Source/Programs/Horde/HordeStorage/Horde.Storage/Tests/UnitTests/BlobCleanupServiceTest.cs
--- a/Engine/Source/Programs/Horde/HordeStorage/Horde.Storage/Tests/UnitTests/BlobCleanupServiceTest.cs
+++ b/Engine/Source/Programs/Horde/HordeStorage/Horde.Storage/Tests/UnitTests/BlobCleanupServiceTest.cs
@@ -1,10 +1,8 @@
 ﻿// Copyright Epic Games, Inc. All Rights Reserved.
 
 using System;
-using System.Collections.Generic;
 using System.Threading;
 using System.Threading.Tasks;
-using Dasync.Collections;
 using Horde.Storage.Implementation;
 using Microsoft.VisualStudio.TestTools.UnitTesting;
 using Moq;
@@ -22,7 +20,7 @@ namespace Horde.Storage.UnitTests
                 CleanOldBlobs = false
             };
             TestOptionsMonitor<GCSettings> gcSettingsMon = new TestOptionsMonitor<GCSettings>(gcSettings);
-            BlobCleanupService blobCleanupService = new BlobCleanupService(Mock.Of<IServiceProvider>(), gcSettingsMon);
+            await using BlobCleanupService blobCleanupService = new BlobCleanupService(Mock.Of<IServiceProvider>(), gcSettingsMon);
 
             Mock<IBlobCleanup> store1 = new Mock<IBlobCleanup>();
             store1.Setup(cleanup => cleanup.ShouldRun()).Returns(true);
@@ -31,7 +29,8 @@ namespace Horde.Storage.UnitTests
             blobCleanupService.RegisterCleanup(store1.Object);
             blobCleanupService.RegisterCleanup(store2.Object);
 
-            await blobCleanupService.OnPoll(blobCleanupService.State, new CancellationTokenSource().Token);
+            using CancellationTokenSource tokenSource = new CancellationTokenSource();
+            await blobCleanupService.OnPoll(blobCleanupService.State, tokenSource.Token);
             
             store1.Verify(m => m.Cleanup(It.IsAny<CancellationToken>()), Times.Once);
             store2.Verify(m => m.Cleanup(It.IsAny<CancellationToken>()), Times.Once);
diff --git a/Engine/Source/Programs/Horde/HordeStorage/Horde.Storage/Tests/UnitTests/CompactBinaryTests.cs b/Engine/Source/Programs/Horde/HordeStorage/Horde.Storage/Tests/UnitTests/CompactBinaryTests.cs
--- a/Engine/Source/Programs/Horde/HordeStorage/Horde.Storage/Tests/UnitTests/CompactBinaryTests.cs
+++ b/Engine/Source/Programs/Horde/HordeStorage/Horde.Storage/Tests/UnitTests/CompactBinaryTests.cs
@@ -1,12 +1,10 @@
 ﻿// Copyright Epic Games, Inc. All Rights Reserved.
 
-using System;
 using System.Collections.Generic;
 using System.IO;
 using System.Linq;
 using EpicGames.Core;
 using EpicGames.Serialization;
-using Jupiter.Implementation;
 using Microsoft.VisualStudio.TestTools.UnitTesting;
 
 namespace EuropaUnit
@@ -65,7 +63,7 @@ namespace EuropaUnit
         }
 
         [TestMethod]
-        public void compact_binary()
+        public void CompactBinary()
         {
             byte[] bytes = File.ReadAllBytes("CompactBinaryObjects/compact_binary");
 
@@ -80,8 +78,6 @@ namespace EuropaUnit
             Assert.AreEqual("{\"Key\":{\"Bucket\":\"EditorDomainPackage\",\"Hash\":\"37dbaa409ef30ba67f18c8fc2faaf606636cb915\"},\"Meta\":{\"FileSize\":24789},\"Attachments\":[{\"Id\":\"000000000000000000000001\",\"RawHash\":\"da6fc57e4b9f91377c9509ea0ad567bacb3796c5\",\"RawSize\":24789}]}", o.ToJson());
         }
 
-
-
         [TestMethod]
         public void WriteArray()
         {
@@ -109,7 +105,6 @@ namespace EuropaUnit
             CollectionAssert.AreEqual(new IoHash[] {hash1, hash2}, blobs);
         }
 
-
         [TestMethod]
         public void WriteObject()
         {
@@ -136,5 +131,4 @@ namespace EuropaUnit
             Assert.AreEqual(hash1, hashField!.AsAttachment());
         }
     }
-
 }
diff --git a/Engine/Source/Programs/Horde/HordeStorage/Horde.Storage/Tests/UnitTests/HierarchicalBlobStorageTest.cs b/Engine/Source/Programs/Horde/HordeStorage/Horde.Storage/Tests/UnitTests/HierarchicalBlobStorageTest.cs
--- a/Engine/Source/Programs/Horde/HordeStorage/Horde.Storage/Tests/UnitTests/HierarchicalBlobStorageTest.cs
+++ b/Engine/Source/Programs/Horde/HordeStorage/Horde.Storage/Tests/UnitTests/HierarchicalBlobStorageTest.cs
@@ -7,7 +7,6 @@ using System.Linq;
 using System.Net.Http;
 using System.Text;
 using System.Threading.Tasks;
-using Dasync.Collections;
 using EpicGames.Horde.Storage;
 using Horde.Storage.Implementation;
 using Horde.Storage.Implementation.Blob;
@@ -51,10 +50,11 @@ namespace Horde.Storage.UnitTests
         public async Task Setup()
         {
             Mock<IServiceProvider> serviceProviderMock = new Mock<IServiceProvider>();
-            serviceProviderMock.Setup(x => x.GetService(typeof(MemoryCacheBlobStore))).Returns(new MemoryCacheBlobStore(Mock.Of<IOptionsMonitor<MemoryCacheBlobSettings>>(_ => _.CurrentValue == new MemoryCacheBlobSettings())));
+            using MemoryCacheBlobStore blobStore =  new MemoryCacheBlobStore(Mock.Of<IOptionsMonitor<MemoryCacheBlobSettings>>(_ => _.CurrentValue == new MemoryCacheBlobSettings()));
+            serviceProviderMock.Setup(x => x.GetService(typeof(MemoryCacheBlobStore))).Returns(blobStore);
             IOptionsMonitor<HordeStorageSettings> settingsMonitor = Mock.Of<IOptionsMonitor<HordeStorageSettings>>(_ => _.CurrentValue == new HordeStorageSettings());
             Mock<INamespacePolicyResolver> mockPolicyResolver = new Mock<INamespacePolicyResolver>();
-            mockPolicyResolver.Setup(x => x.GetPoliciesForNs(It.IsAny<NamespaceId>())).Returns(new NamespaceSettings.PerNamespaceSettings());
+            mockPolicyResolver.Setup(x => x.GetPoliciesForNs(It.IsAny<NamespaceId>())).Returns(new NamespacePolicy());
 
             _chained = new BlobService(serviceProviderMock.Object, settingsMonitor, Mock.Of<IBlobIndex>(), Mock.Of<IPeerStatusService>(), Mock.Of<IHttpClientFactory>(), Mock.Of<IServiceCredentials>(), mockPolicyResolver.Object);
             _chained.BlobStore = new List<IBlobStore> { _first, _second, _third };
@@ -156,7 +156,7 @@ namespace Horde.Storage.UnitTests
             await Assert.ThrowsExceptionAsync<NamespaceNotFoundException>(() => _second.DeleteNamespace(Ns));
         }
 
-        private string BlobToString(BlobContents contents)
+        private static string BlobToString(BlobContents contents)
         {
             using StreamReader reader = new StreamReader(contents.Stream);
             return reader.ReadToEnd();
diff --git a/Engine/Source/Programs/Horde/HordeStorage/Horde.Storage/Tests/UnitTests/HordeStorageStartupTest.cs b/Engine/Source/Programs/Horde/HordeStorage/Horde.Storage/Tests/UnitTests/HordeStorageStartupTest.cs
--- a/Engine/Source/Programs/Horde/HordeStorage/Horde.Storage/Tests/UnitTests/HordeStorageStartupTest.cs
+++ b/Engine/Source/Programs/Horde/HordeStorage/Horde.Storage/Tests/UnitTests/HordeStorageStartupTest.cs
@@ -16,7 +16,7 @@ namespace Horde.Storage.UnitTests
     [TestClass]
     public class HordeStorageStartupTest
     {
-        private string localTestDir;
+        private readonly string localTestDir;
         
         public HordeStorageStartupTest()
         {
@@ -45,7 +45,6 @@ namespace Horde.Storage.UnitTests
                 {"Horde_Storage:StorageImplementations:1", "Memory"},
             });
 
-
             List<IBlobStore> blobStores = blobService.BlobStore.ToList();
             Assert.IsTrue(blobStores[0] is FileSystemStore);
             Assert.IsTrue(blobStores[1] is MemoryCacheBlobStore);
@@ -58,7 +57,7 @@ namespace Horde.Storage.UnitTests
                 .AddInMemoryCollection(new Dictionary<string, string> {{"Filesystem:RootDir", localTestDir}})
                 .AddInMemoryCollection(configDict)
                 .Build();
-            TestServer server = new TestServer(new WebHostBuilder()
+            using TestServer server = new TestServer(new WebHostBuilder()
                 .UseConfiguration(configuration)
                 .UseEnvironment("Testing")
                 .UseSerilog(new LoggerConfiguration().ReadFrom.Configuration(configuration).CreateLogger())
diff --git a/Engine/Source/Programs/Horde/HordeStorage/Horde.Storage/Tests/UnitTests/MemoryCacheBlobStoreTest.cs b/Engine/Source/Programs/Horde/HordeStorage/Horde.Storage/Tests/UnitTests/MemoryCacheBlobStoreTest.cs
--- a/Engine/Source/Programs/Horde/HordeStorage/Horde.Storage/Tests/UnitTests/MemoryCacheBlobStoreTest.cs
+++ b/Engine/Source/Programs/Horde/HordeStorage/Horde.Storage/Tests/UnitTests/MemoryCacheBlobStoreTest.cs
@@ -34,7 +34,7 @@ namespace Horde.Storage.UnitTests
     }
     
     [TestClass]
-    public class MemoryCacheBlobStoreTest
+    public class MemoryCacheBlobStoreTest: IDisposable
     {
         private readonly NamespaceId Ns1 = new NamespaceId("my-namespace-1");
         private readonly NamespaceId Ns2 = new NamespaceId("my-namespace-2");
@@ -116,5 +116,19 @@ namespace Horde.Storage.UnitTests
         public void ListOldObjects()
         {
         }
+
+        protected virtual void Dispose(bool disposing)
+        {
+            if (disposing)
+            {
+                _store.Dispose();
+            }
+        }
+
+        public void Dispose()
+        {
+            Dispose(true);
+            GC.SuppressFinalize(this);
+        }
     }
 }
diff --git a/Engine/Source/Programs/Horde/HordeStorage/Horde.Storage/Tests/UnitTests/ReplicatorTests.cs b/Engine/Source/Programs/Horde/HordeStorage/Horde.Storage/Tests/UnitTests/ReplicatorTests.cs
--- a/Engine/Source/Programs/Horde/HordeStorage/Horde.Storage/Tests/UnitTests/ReplicatorTests.cs
+++ b/Engine/Source/Programs/Horde/HordeStorage/Horde.Storage/Tests/UnitTests/ReplicatorTests.cs
@@ -18,7 +18,6 @@ using Microsoft.Extensions.Options;
 using Moq;
 using RestSharp;
 using EpicGames.Horde.Storage;
-using Horde.Storage.Controllers;
 using Moq.Contrib.HttpClient;
 
 namespace Horde.Storage.UnitTests
@@ -60,9 +59,9 @@ namespace Horde.Storage.UnitTests
             Mock<ITransactionLogWriter> transactionLogWriter = new Mock<ITransactionLogWriter>();
             Mock<IRestClient> remoteClientMock = new Mock<IRestClient> { DefaultValue = DefaultValue.Empty };
 
-            CallistoReader.CallistoGetResponse[] responses = new[]
+            CallistoGetResponse[] responses = new[]
             {
-                new CallistoReader.CallistoGetResponse(
+                new CallistoGetResponse(
                     generation: _currentLogGeneration,
                     currentOffset: 0,
                     events: new List<TransactionEvent>(new[]
@@ -70,17 +69,17 @@ namespace Horde.Storage.UnitTests
                         new AddTransactionEvent("", "", new[] {blobToReplication}, identifier:0, nextIdentifier: 100)
                     })
                 ),
-                new CallistoReader.CallistoGetResponse(
+                new CallistoGetResponse(
                     generation: _currentLogGeneration,
                     currentOffset: 100,
                     events: new List<TransactionEvent>()
                 )
             };
 
-            Mock<IRestResponse<CallistoReader.CallistoGetResponse>>[] mockedCallistoResponse = responses.Select(response => CreateMockResponse(HttpStatusCode.OK, response)).ToArray();
+            Mock<IRestResponse<CallistoGetResponse>>[] mockedCallistoResponse = responses.Select(response => CreateMockResponse(HttpStatusCode.OK, response)).ToArray();
             // mock of callisto calls
             remoteClientMock
-                .SetupSequence(x => x.ExecuteGetAsync<CallistoReader.CallistoGetResponse>(It.IsAny<IRestRequest>(), It.IsAny<CancellationToken>()))
+                .SetupSequence(x => x.ExecuteGetAsync<CallistoGetResponse>(It.IsAny<IRestRequest>(), It.IsAny<CancellationToken>()))
                 // first we will query for a callisto event
                 .ReturnsAsync(mockedCallistoResponse[0].Object)
                 // then it will iterate again and should now reach the stop event
@@ -115,7 +114,7 @@ namespace Horde.Storage.UnitTests
             blobStoreMock.Verify(blobStore => blobStore.Exists(ns, blobToReplication), Times.Once);
 
             // we should have attempted to store the replicated object in our callisto
-            remoteClientMock.Verify(client => client.ExecuteGetAsync<CallistoReader.CallistoGetResponse>(It.IsAny<IRestRequest>(), It.IsAny<CancellationToken>()), Times.Exactly(3));
+            remoteClientMock.Verify(client => client.ExecuteGetAsync<CallistoGetResponse>(It.IsAny<IRestRequest>(), It.IsAny<CancellationToken>()), Times.Exactly(3));
             transactionLogWriter.Verify(writer => writer.Add(ns, It.Is<AddTransactionEvent>(e => e.Blobs[0].Equals(blobToReplication))), Times.Once);
 
             // as this was a add operation we should have transferred the blob from the remote blob store to the local
@@ -126,7 +125,6 @@ namespace Horde.Storage.UnitTests
             remoteClientMock.VerifyNoOtherCalls();
         }
 
-
         [TestMethod]
         public async Task ReplicatorSkipLogReplicatorV1()
         {
@@ -137,18 +135,18 @@ namespace Horde.Storage.UnitTests
 
             Mock<IRestClient> remoteClientMock = new Mock<IRestClient> { DefaultValue = DefaultValue.Empty };
             IServiceCredentials serviceCredentials = Mock.Of<IServiceCredentials>();
-            CallistoReader.CallistoGetResponse[] responses = new[]
+            CallistoGetResponse[] responses = new[]
             {
-                new CallistoReader.CallistoGetResponse(
+                new CallistoGetResponse(
                     generation: _currentLogGeneration,
                     currentOffset: 100000,
                     events: new List<TransactionEvent>()
                 ),
             };
-            Mock<IRestResponse<CallistoReader.CallistoGetResponse>>[] mockedCallistoResponse = responses.Select(response => CreateMockResponse(HttpStatusCode.OK, response)).ToArray();
+            Mock<IRestResponse<CallistoGetResponse>>[] mockedCallistoResponse = responses.Select(response => CreateMockResponse(HttpStatusCode.OK, response)).ToArray();
             // mock of callisto calls
             remoteClientMock
-                .SetupSequence(x => x.ExecuteGetAsync<CallistoReader.CallistoGetResponse>(It.IsAny<IRestRequest>(), It.IsAny<CancellationToken>()))
+                .SetupSequence(x => x.ExecuteGetAsync<CallistoGetResponse>(It.IsAny<IRestRequest>(), It.IsAny<CancellationToken>()))
                 .ReturnsAsync(mockedCallistoResponse[0].Object)
                 .ReturnsAsync(mockedCallistoResponse[0].Object);
 
@@ -175,7 +173,7 @@ namespace Horde.Storage.UnitTests
             Assert.AreEqual(_currentLogGeneration, replicator.State.ReplicatingGeneration);
 
             // we should called the transaction log 2
-            remoteClientMock.Verify(client => client.ExecuteGetAsync<CallistoReader.CallistoGetResponse>(It.IsAny<IRestRequest>(), It.IsAny<CancellationToken>()), Times.Exactly(2));
+            remoteClientMock.Verify(client => client.ExecuteGetAsync<CallistoGetResponse>(It.IsAny<IRestRequest>(), It.IsAny<CancellationToken>()), Times.Exactly(2));
 
             remoteClientMock.VerifyNoOtherCalls();
         }
@@ -196,9 +194,9 @@ namespace Horde.Storage.UnitTests
             Mock<ITransactionLogWriter> transactionLogWriter = new Mock<ITransactionLogWriter>();
             Mock<IRestClient> remoteClientMock = new Mock<IRestClient> { DefaultValue = DefaultValue.Empty };
 
-            CallistoReader.CallistoGetResponse[] responses = new[]
+            CallistoGetResponse[] responses = new[]
             {
-                new CallistoReader.CallistoGetResponse(
+                new CallistoGetResponse(
                     generation: _currentLogGeneration,
                     currentOffset: 0,
                     events: new List<TransactionEvent>(new[]
@@ -206,17 +204,17 @@ namespace Horde.Storage.UnitTests
                         new AddTransactionEvent("", "", new[] {blobToReplication}, identifier:0, nextIdentifier: 100)
                     })
                 ),
-                new CallistoReader.CallistoGetResponse(
+                new CallistoGetResponse(
                     generation: _currentLogGeneration,
                     currentOffset: 100,
                     events: new List<TransactionEvent>()
                 )
             };
 
-            Mock<IRestResponse<CallistoReader.CallistoGetResponse>>[] mockedCallistoResponse = responses.Select(response => CreateMockResponse(HttpStatusCode.OK, response)).ToArray();
+            Mock<IRestResponse<CallistoGetResponse>>[] mockedCallistoResponse = responses.Select(response => CreateMockResponse(HttpStatusCode.OK, response)).ToArray();
             // mock of callisto calls
             remoteClientMock
-                .SetupSequence(x => x.ExecuteGetAsync<CallistoReader.CallistoGetResponse>(It.IsAny<IRestRequest>(), It.IsAny<CancellationToken>()))
+                .SetupSequence(x => x.ExecuteGetAsync<CallistoGetResponse>(It.IsAny<IRestRequest>(), It.IsAny<CancellationToken>()))
                 // first we will query for a callisto event
                 .ReturnsAsync(mockedCallistoResponse[0].Object)
                 // then it will iterate again and should now reach the stop event
@@ -253,7 +251,7 @@ namespace Horde.Storage.UnitTests
             blobStoreMock.Verify(blobStore => blobStore.Exists(ns, blobToReplication), Times.Once);
 
             // we should have attempted to store the replicated object in our callisto
-            remoteClientMock.Verify(client => client.ExecuteGetAsync<CallistoReader.CallistoGetResponse>(It.IsAny<IRestRequest>(), It.IsAny<CancellationToken>()), Times.Exactly(3));
+            remoteClientMock.Verify(client => client.ExecuteGetAsync<CallistoGetResponse>(It.IsAny<IRestRequest>(), It.IsAny<CancellationToken>()), Times.Exactly(3));
             transactionLogWriter.Verify(writer => writer.Add(ns, It.Is<AddTransactionEvent>(e => e.Blobs[0].Equals(blobToReplication))), Times.Once);
 
             // as this was a add operation we should have transferred the blob from the remote blob store to the local
@@ -264,7 +262,6 @@ namespace Horde.Storage.UnitTests
             remoteClientMock.VerifyNoOtherCalls();
         }
 
-
         [TestMethod]
         [Ignore("Fails intermittently on build machines, this is a test for our legacy replicator so disabling as this code will be removed soon")]
         public async Task ReplicatorDeleteTransactionV1()
@@ -281,9 +278,9 @@ namespace Horde.Storage.UnitTests
             Mock<ITransactionLogWriter> transactionLogWriter = new Mock<ITransactionLogWriter>();
             Mock<IRestClient> remoteClientMock = new Mock<IRestClient> { DefaultValue = DefaultValue.Empty };
 
-            CallistoReader.CallistoGetResponse[] responses = new[]
+            CallistoGetResponse[] responses = new[]
             {
-                new CallistoReader.CallistoGetResponse(
+                new CallistoGetResponse(
                     generation: _currentLogGeneration,
                     currentOffset: 0,
                     events: new List<TransactionEvent>(new[]
@@ -291,17 +288,17 @@ namespace Horde.Storage.UnitTests
                         new RemoveTransactionEvent("", "", 0, 100)
                     })
                 ),
-                new CallistoReader.CallistoGetResponse(
+                new CallistoGetResponse(
                     generation: _currentLogGeneration,
                     currentOffset: 100,
                     events: new List<TransactionEvent>()
                 )
             };
 
-            Mock<IRestResponse<CallistoReader.CallistoGetResponse>>[] mockedCallistoResponse = responses.Select(response => CreateMockResponse(HttpStatusCode.OK, response)).ToArray();
+            Mock<IRestResponse<CallistoGetResponse>>[] mockedCallistoResponse = responses.Select(response => CreateMockResponse(HttpStatusCode.OK, response)).ToArray();
             // mock of callisto calls
             remoteClientMock
-                .SetupSequence(x => x.ExecuteGetAsync<CallistoReader.CallistoGetResponse>(It.IsAny<IRestRequest>(), It.IsAny<CancellationToken>()))
+                .SetupSequence(x => x.ExecuteGetAsync<CallistoGetResponse>(It.IsAny<IRestRequest>(), It.IsAny<CancellationToken>()))
                 // first we will query for a callisto event
                 .ReturnsAsync(mockedCallistoResponse[0].Object)
                 // then it will iterate again and should now reach the stop event
@@ -322,7 +319,7 @@ namespace Horde.Storage.UnitTests
             Assert.IsTrue(ran);
 
             // we should have fetched the replicated object
-            remoteClientMock.Verify(client => client.ExecuteGetAsync<CallistoReader.CallistoGetResponse>(It.IsAny<IRestRequest>(), It.IsAny<CancellationToken>()), Times.Exactly(3));
+            remoteClientMock.Verify(client => client.ExecuteGetAsync<CallistoGetResponse>(It.IsAny<IRestRequest>(), It.IsAny<CancellationToken>()), Times.Exactly(3));
             transactionLogWriter.Verify(writer => writer.Add(ns, It.IsAny<RemoveTransactionEvent>()), Times.Once);
             // other then copying the remove event we should have written nothing else
             transactionLogWriter.VerifyNoOtherCalls();
@@ -334,7 +331,7 @@ namespace Horde.Storage.UnitTests
             remoteClientMock.VerifyNoOtherCalls();
         }
 
-        private Mock<IRestResponse<T>> CreateMockResponse<T>(HttpStatusCode httpStatusCode, T data)
+        private static Mock<IRestResponse<T>> CreateMockResponse<T>(HttpStatusCode httpStatusCode, T data)
         {
             Mock<IRestResponse<T>> response = new Mock<IRestResponse<T>>();
             response.Setup(_ => _.StatusCode).Returns(httpStatusCode);
@@ -353,6 +350,5 @@ namespace Horde.Storage.UnitTests
 
             return response;
         }
-
     }
 }
diff --git a/Engine/Source/Programs/Horde/HordeStorage/Horde.Storage/Tests/UnitTests/S3StoreTests.cs b/Engine/Source/Programs/Horde/HordeStorage/Horde.Storage/Tests/UnitTests/S3StoreTests.cs
--- a/Engine/Source/Programs/Horde/HordeStorage/Horde.Storage/Tests/UnitTests/S3StoreTests.cs
+++ b/Engine/Source/Programs/Horde/HordeStorage/Horde.Storage/Tests/UnitTests/S3StoreTests.cs
@@ -1,8 +1,6 @@
 ﻿// Copyright Epic Games, Inc. All Rights Reserved.
 
 using System;
-using System.Collections.Generic;
-using System.IO;
 using System.Text;
 using System.Threading.Tasks;
 using Amazon.S3;
@@ -22,7 +20,7 @@ namespace Horde.Storage.UnitTests
     [TestClass]
     public class S3StoreTests
     {
-        private NamespaceId Namespace = new NamespaceId("foo");
+        private readonly NamespaceId Namespace = new NamespaceId("foo");
 
         [TestMethod]
         public async Task PutBufferSuccess()
@@ -32,7 +30,7 @@ namespace Horde.Storage.UnitTests
             IOptionsMonitor<S3Settings> settingsMock = Mock.Of<IOptionsMonitor<S3Settings>>(_ => _.CurrentValue == settings);
 
             INamespacePolicyResolver namespacePolicyResolver = Mock.Of<INamespacePolicyResolver>(_ => _.GetPoliciesForNs(Namespace) ==
-                                                  new NamespaceSettings.PerNamespaceSettings
+                                                  new NamespacePolicy
                                                   {
                                                       StoragePool = ""
                                                   });
@@ -60,7 +58,7 @@ namespace Horde.Storage.UnitTests
             IOptionsMonitor<S3Settings> settingsMock = Mock.Of<IOptionsMonitor<S3Settings>>(_ => _.CurrentValue == settings);
             
             INamespacePolicyResolver namespacePolicyResolver = Mock.Of<INamespacePolicyResolver>(_ => _.GetPoliciesForNs(Namespace) ==
-                new NamespaceSettings.PerNamespaceSettings
+                new NamespacePolicy
                 {
                     StoragePool = "storagepool"
                 });
@@ -91,7 +89,7 @@ namespace Horde.Storage.UnitTests
             S3Settings settings = new S3Settings{BucketName = "tests-foo"};
             IOptionsMonitor<S3Settings> settingsMock = Mock.Of<IOptionsMonitor<S3Settings>>(_ => _.CurrentValue == settings);
             INamespacePolicyResolver namespacePolicyResolver = Mock.Of<INamespacePolicyResolver>(_ => _.GetPoliciesForNs(Namespace) ==
-                new NamespaceSettings.PerNamespaceSettings
+                new NamespacePolicy
                 {
                     StoragePool = ""
                 });
@@ -118,7 +116,7 @@ namespace Horde.Storage.UnitTests
             IOptionsMonitor<S3Settings> settingsMock = Mock.Of<IOptionsMonitor<S3Settings>>(_ => _.CurrentValue == settings);
 
             INamespacePolicyResolver namespacePolicyResolver = Mock.Of<INamespacePolicyResolver>(_ => _.GetPoliciesForNs(Namespace) ==
-                new NamespaceSettings.PerNamespaceSettings
+                new NamespacePolicy
                 {
                     StoragePool = ""
                 });
@@ -138,7 +136,7 @@ namespace Horde.Storage.UnitTests
             IOptionsMonitor<S3Settings> settingsMock = Mock.Of<IOptionsMonitor<S3Settings>>(_ => _.CurrentValue == settings);
 
             INamespacePolicyResolver namespacePolicyResolver = Mock.Of<INamespacePolicyResolver>(_ => _.GetPoliciesForNs(Namespace) ==
-                new NamespaceSettings.PerNamespaceSettings
+                new NamespacePolicy
                 {
                     StoragePool = ""
                 });
diff --git a/Engine/Source/Programs/Horde/HordeStorage/Jupiter.Common/AWSCredentialsHelper.cs b/Engine/Source/Programs/Horde/HordeStorage/Jupiter.Common/AWSCredentialsHelper.cs
--- a/Engine/Source/Programs/Horde/HordeStorage/Jupiter.Common/AWSCredentialsHelper.cs
+++ b/Engine/Source/Programs/Horde/HordeStorage/Jupiter.Common/AWSCredentialsHelper.cs
@@ -4,11 +4,10 @@ using System;
 using System.Collections.Generic;
 using System.ComponentModel.DataAnnotations;
 using Amazon.Runtime;
-using Amazon.SecurityToken.Model;
 
 namespace Jupiter
 {
-    public class AWSCredentialsHelper
+    public static class AWSCredentialsHelper
     {
         public static AWSCredentials GetCredentials(AWSCredentialsSettings options, string sessionName)
         {
@@ -21,7 +20,7 @@ namespace Jupiter
                 case AWSCredentialsType.AssumeRoleWebIdentity:
                     return AssumeRoleWithWebIdentityCredentials.FromEnvironmentVariables();
                 default:
-                    throw new ArgumentOutOfRangeException();
+                    throw new NotImplementedException($"Unhandled enum option {options.AWSCredentialsType}");
             }
         }
     }
@@ -36,7 +35,7 @@ namespace Jupiter
 
         public IEnumerable<ValidationResult> Validate(ValidationContext validationContext)
         {
-            var results = new List<ValidationResult>();
+            List<ValidationResult> results = new();
 
             switch (AWSCredentialsType)
             {
@@ -62,7 +61,7 @@ namespace Jupiter
                     // the environment variables will be verified when we create the credentials type, as there is no configuration in our appsettings we do not do any checks
                     break;
                 default:
-                    throw new ArgumentOutOfRangeException();
+	                throw new NotImplementedException($"Unhandled enum option {AWSCredentialsType}");
             }
 
             return results;
diff --git a/Engine/Source/Programs/Horde/HordeStorage/Jupiter.Common/Authentication/NamespaceAuthorizationHandler.cs b/Engine/Source/Programs/Horde/HordeStorage/Jupiter.Common/Authentication/NamespaceAuthorizationHandler.cs
--- a/Engine/Source/Programs/Horde/HordeStorage/Jupiter.Common/Authentication/NamespaceAuthorizationHandler.cs
+++ b/Engine/Source/Programs/Horde/HordeStorage/Jupiter.Common/Authentication/NamespaceAuthorizationHandler.cs
@@ -1,14 +1,11 @@
 ﻿// Copyright Epic Games, Inc. All Rights Reserved.
 
 using System;
-using System.Collections.Generic;
-using System.ComponentModel.DataAnnotations;
 using System.Threading.Tasks;
 using EpicGames.Horde.Storage;
 using Jupiter.Common;
 using Jupiter.Implementation;
 using Microsoft.AspNetCore.Authorization;
-using Microsoft.Extensions.Options;
 
 namespace Jupiter
 {
@@ -33,10 +30,10 @@ namespace Jupiter
 
             try
             {
-                NamespaceSettings.PerNamespaceSettings settings = _namespacePolicyResolver.GetPoliciesForNs(namespaceName);
+                NamespacePolicy policy = _namespacePolicyResolver.GetPoliciesForNs(namespaceName);
 
                 // These are ANDed, e.g. all claims needs to be present
-                foreach (string expectedClaim in settings.Claims)
+                foreach (string expectedClaim in policy.Claims)
                 {
                     // if expected claim is * then everyone is allowed to use the namespace
                     if (expectedClaim == "*")
@@ -45,9 +42,9 @@ namespace Jupiter
                         continue;
                     }
 
-                    if (expectedClaim.Contains('='))
+                    if (expectedClaim.Contains('=', StringComparison.InvariantCultureIgnoreCase))
                     {
-                        int separatorIndex = expectedClaim.IndexOf('=');
+                        int separatorIndex = expectedClaim.IndexOf('=', StringComparison.InvariantCultureIgnoreCase);
                         string claimName = expectedClaim.Substring(0, separatorIndex);
                         string claimValue = expectedClaim.Substring(separatorIndex + 1);
                         if (context.User.HasClaim(claim => claim.Type == claimName && claim.Value == claimValue))
diff --git a/Engine/Source/Programs/Horde/HordeStorage/Jupiter.Common/BaseProgram.cs b/Engine/Source/Programs/Horde/HordeStorage/Jupiter.Common/BaseProgram.cs
--- a/Engine/Source/Programs/Horde/HordeStorage/Jupiter.Common/BaseProgram.cs
+++ b/Engine/Source/Programs/Horde/HordeStorage/Jupiter.Common/BaseProgram.cs
@@ -1,13 +1,7 @@
 // Copyright Epic Games, Inc. All Rights Reserved.
 
 using System;
-using System.Collections.Generic;
 using System.IO;
-using System.IO.Pipelines;
-using System.Net.Sockets;
-using System.Reflection;
-using System.Threading.Tasks;
-using Microsoft.AspNetCore.Connections;
 using Microsoft.AspNetCore.Hosting;
 using Microsoft.Extensions.Configuration;
 using Microsoft.Extensions.Hosting;
@@ -47,6 +41,7 @@ namespace Jupiter
 
         }
 
+        [System.Diagnostics.CodeAnalysis.SuppressMessage("Design", "CA1000:Do not declare static members on generic types", Justification = "Keep logic for the base program together")]
         public static int BaseMain(string[] args)
         {
             Log.Logger = new LoggerConfiguration()
@@ -70,6 +65,7 @@ namespace Jupiter
             }
         }
 
+        [System.Diagnostics.CodeAnalysis.SuppressMessage("Design", "CA1000:Do not declare static members on generic types", Justification = "Keep logic for the base program together")]
         public static IHostBuilder CreateHostBuilder(string[] args)
         {
             return Host.CreateDefaultBuilder(args)
