Horde Storage - Updated editorconfig to more closely match what we run for Horde and updated code to fix all static analysis errors

diff --git a/Engine/Source/Programs/Horde/HordeStorage/Jupiter.Common/BaseStartup.cs b/Engine/Source/Programs/Horde/HordeStorage/Jupiter.Common/BaseStartup.cs
--- a/Engine/Source/Programs/Horde/HordeStorage/Jupiter.Common/BaseStartup.cs
+++ b/Engine/Source/Programs/Horde/HordeStorage/Jupiter.Common/BaseStartup.cs
@@ -39,19 +39,17 @@ namespace Jupiter
 {
     public abstract class BaseStartup
     {
-        protected readonly ILogger _logger = Log.ForContext<BaseStartup>();
+	    protected ILogger Logger { get; } = Log.ForContext<BaseStartup>();
 
-        public BaseStartup(IConfiguration configuration, IWebHostEnvironment environment)
+        protected BaseStartup(IConfiguration configuration)
         {
             Configuration = configuration;
-            Environment = environment;
             Auth = new AuthSettings();
         }
 
-        public IConfiguration Configuration { get; }
-        public IWebHostEnvironment Environment { get; }
+        protected IConfiguration Configuration { get; }
 
-        protected AuthSettings Auth { get; }
+        private AuthSettings Auth { get; }
 
         // This method gets called by the runtime. Use this method to add services to the container.
         public void ConfigureServices(IServiceCollection services)
@@ -96,7 +94,7 @@ namespace Jupiter
                 {
                     options.InvalidModelStateResponseFactory = context =>
                     {
-                        var result = new BadRequestObjectResult(context.ModelState);
+                        BadRequestObjectResult result = new BadRequestObjectResult(context.ModelState);
                         // always return errors as json objects
                         // we could allow more types here, but we do not want raw for instance
                         result.ContentTypes.Add(MediaTypeNames.Application.Json);
@@ -128,7 +126,7 @@ namespace Jupiter
                             options.DefaultChallengeScheme = DisabledAuthenticationHandler.AuthenticateScheme;
                             break;
                         default:
-                            throw new ArgumentOutOfRangeException();
+                            throw new NotImplementedException($"Method {Auth.Method} not implemented");
                     }
                 });
             if (Auth.Method == AuthMethod.JWTBearer)
@@ -178,7 +176,6 @@ namespace Jupiter
                 options.ForwardedHeaders = ForwardedHeaders.XForwardedFor | ForwardedHeaders.XForwardedProto;
             });
 
-
             services.AddSwaggerGen(settings =>
             {
                 string? assemblyName = Assembly.GetEntryAssembly()?.GetName().Name;
@@ -213,7 +210,7 @@ namespace Jupiter
             OnAddHealthChecks(services, healthChecks);
 
             string? ddAgentHost = System.Environment.GetEnvironmentVariable("DD_AGENT_HOST");
-            if (!String.IsNullOrEmpty(ddAgentHost))
+            if (!string.IsNullOrEmpty(ddAgentHost))
             {
                 healthChecks.AddDatadogPublisher("jupiter.healthchecks");
             }
@@ -246,7 +243,7 @@ namespace Jupiter
 
             if (jupiterSettings.ShowPII)
             {
-                _logger.Error("Personally Identifiable information being shown. This should not be generally enabled in prod.");
+                Logger.Error("Personally Identifiable information being shown. This should not be generally enabled in prod.");
 
                 // do not hide personal information during development
                 Microsoft.IdentityModel.Logging.IdentityModelEventSource.ShowPII = true;
@@ -344,7 +341,7 @@ namespace Jupiter
 
     public class MvcJsonOptionsWrapper : IConfigureOptions<MvcNewtonsoftJsonOptions>
     {
-        IServiceProvider ServiceProvider;
+        readonly IServiceProvider ServiceProvider;
 
         public MvcJsonOptionsWrapper(IServiceProvider serviceProvider)
         {
@@ -358,7 +355,7 @@ namespace Jupiter
 
     public class FieldFilteringResolver : DefaultContractResolver
     {
-        private IHttpContextAccessor _httpContextAccessor;
+        private readonly IHttpContextAccessor _httpContextAccessor;
 
         public FieldFilteringResolver(IServiceProvider sp)
         {
@@ -377,11 +374,15 @@ namespace Jupiter
                 HttpContext? httpContext = _httpContextAccessor.HttpContext;
 
                 if (httpContext == null)
+                {
                     return true;
+                }
 
                 // if no fields are being filtered we should serialize the property
                 if (!httpContext.Request.Query.ContainsKey("fields"))
+                {
                     return true;
+                }
 
                 StringValues fields = httpContext.Request.Query["fields"];
                 bool ignore = true;
@@ -392,7 +393,7 @@ namespace Jupiter
                     {
                         return true;
                     }
-                    if (string.Equals(field, property.PropertyName, StringComparison.InvariantCultureIgnoreCase))
+                    if (string.Equals(field, property.PropertyName, StringComparison.OrdinalIgnoreCase))
                     {
                         ignore = false;
                     }
@@ -405,7 +406,6 @@ namespace Jupiter
         }
     }
 
-
     public enum AuthMethod
     {
         JWTBearer,
@@ -418,7 +418,6 @@ namespace Jupiter
         [Required] 
         public AuthMethod Method { get; set; } = AuthMethod.JWTBearer;
 
-
         [Required]
         public string OktaDomain { get; set; } = "";
         public string AuthorizationServerId { get; set; } = OktaWebOptions.DefaultAuthorizationServerId;
@@ -438,10 +437,10 @@ namespace Jupiter
         /// <summary>
         /// If the request is smaller then MemoryBufferSize we buffer it in memory rather then as a file
         /// </summary>
-        public long MemoryBufferSize { get; set; } = Int32.MaxValue;
+        public long MemoryBufferSize { get; set; } = int.MaxValue;
 
         // enable to unhide potentially personal information, see https://aka.ms/IdentityModel/PII
-        public bool ShowPII = false;
+        public bool ShowPII { get; set; } = false;
         public bool DisableHealthChecks { get; set; } = false;
         public bool HostSwaggerDocumentation { get; set; } = true;
 
@@ -460,19 +459,20 @@ namespace Jupiter
 
     public class NamespaceSettings
     {
-        public Dictionary<string, PerNamespaceSettings> Policies { get; set; } = new Dictionary<string, PerNamespaceSettings>();
+        [System.Diagnostics.CodeAnalysis.SuppressMessage("Usage", "CA2227:Collection properties should be read only", Justification = "Used by the configuration system")]
+        public Dictionary<string, NamespacePolicy> Policies { get; set; } = new Dictionary<string, NamespacePolicy>();
+    }
 
-        public class PerNamespaceSettings
-        {
-            public string[] Claims { get; set; } = Array.Empty<string>();
-            public string StoragePool { get; set; } = "";
-
-            public bool LastAccessTracking { get; set; } = true;
-            public bool OnDemandReplication { get; set; } = false;
-            public bool UseBlobIndexForExists { get; set; } = false;
-            public bool UseBlobIndexForSlowExists { get; set; } = false;
-            public bool? IsLegacyNamespace { get; set; } = null;
-        }
+    public class NamespacePolicy
+    {
+	    public string[] Claims { get; set; } = Array.Empty<string>();
+	    public string StoragePool { get; set; } = "";
+
+	    public bool LastAccessTracking { get; set; } = true;
+	    public bool OnDemandReplication { get; set; } = false;
+	    public bool UseBlobIndexForExists { get; set; } = false;
+	    public bool UseBlobIndexForSlowExists { get; set; } = false;
+	    public bool? IsLegacyNamespace { get; set; } = null;
     }
 
     public class DatadogTraceMiddleware
diff --git a/Engine/Source/Programs/Horde/HordeStorage/Jupiter.Common/Controllers/ErrorController.cs b/Engine/Source/Programs/Horde/HordeStorage/Jupiter.Common/Controllers/ErrorController.cs
--- a/Engine/Source/Programs/Horde/HordeStorage/Jupiter.Common/Controllers/ErrorController.cs
+++ b/Engine/Source/Programs/Horde/HordeStorage/Jupiter.Common/Controllers/ErrorController.cs
@@ -12,10 +12,12 @@ namespace Jupiter.Controllers
         [Route("/error")]
         public IActionResult Error()
         {
-            var context = HttpContext.Features.Get<IExceptionHandlerFeature>();
+            IExceptionHandlerFeature? context = HttpContext.Features.Get<IExceptionHandlerFeature>();
 
             if (context == null)
+            {
                 return NoContent();
+            }
 
             return Problem(
                 detail: context.Error.StackTrace,
diff --git a/Engine/Source/Programs/Horde/HordeStorage/Jupiter.Common/CustomMediaTypeNames.cs b/Engine/Source/Programs/Horde/HordeStorage/Jupiter.Common/CustomMediaTypeNames.cs
--- a/Engine/Source/Programs/Horde/HordeStorage/Jupiter.Common/CustomMediaTypeNames.cs
+++ b/Engine/Source/Programs/Horde/HordeStorage/Jupiter.Common/CustomMediaTypeNames.cs
@@ -2,12 +2,11 @@
 
 namespace Jupiter
 {
-    public class CustomMediaTypeNames
+    public static class CustomMediaTypeNames
     {
         public const string UnrealCompactBinary = "application/x-ue-cb";
         public const string UnrealCompressedBuffer = "application/x-ue-comp";
 
         public const string JupiterInlinedPayload = "application/x-jupiter-inline";
     }
-
 }
diff --git a/Engine/Source/Programs/Horde/HordeStorage/Jupiter.Common/Implementation/BatchWriter.cs b/Engine/Source/Programs/Horde/HordeStorage/Jupiter.Common/Implementation/BatchWriter.cs
--- a/Engine/Source/Programs/Horde/HordeStorage/Jupiter.Common/Implementation/BatchWriter.cs
+++ b/Engine/Source/Programs/Horde/HordeStorage/Jupiter.Common/Implementation/BatchWriter.cs
@@ -8,30 +8,29 @@ using System.Threading;
 using System.Threading.Tasks;
 using Datadog.Trace;
 using EpicGames.Horde.Storage;
-using Serilog;
 
 namespace Jupiter.Implementation
 {
     public class BatchWriter
     {
-        public enum OpState : byte
-        {
+#pragma warning disable CA1028 // Enum Storage should be Int32
+		public enum OpState : byte
+		{
             OK = 0,
             Failed = 1,
             NotFound = 2,
             Exists = 3 // special case for HEAD checks, same as OK but indicates that there will be no payload
         };
+#pragma warning restore CA1028 // Enum Storage should be Int32
 
-        private readonly ILogger _logger = Log.ForContext<BatchWriter>();
-
-        public async Task WriteToStream(Stream outStream, NamespaceId ns, List<Tuple<OpVerb, BucketId, KeyId>> ops, Func<OpVerb, NamespaceId, BucketId, KeyId, string[], Task<Tuple<ContentHash?, BlobContents?, OpState>>> fetchAction)
+		public async Task WriteToStream(Stream outStream, NamespaceId ns, List<Tuple<OpVerb, BucketId, KeyId>> ops, Func<OpVerb, NamespaceId, BucketId, KeyId, string[], Task<Tuple<ContentHash?, BlobContents?, OpState>>> fetchAction)
         {
             await outStream.WriteAsync(Encoding.ASCII.GetBytes("JPTR"));
             await outStream.WriteAsync(BitConverter.GetBytes((uint) ops.Count));
 
             byte[] entryHeaderBytes = Encoding.ASCII.GetBytes("JPEE");
 
-            SemaphoreSlim streamLock = new SemaphoreSlim(1, 1);
+            using SemaphoreSlim streamLock = new SemaphoreSlim(1, 1);
 
             foreach ((OpVerb opVerb, BucketId bucket, KeyId key) in ops)
             {
@@ -64,7 +63,9 @@ namespace Jupiter.Implementation
                     // if we have a failure we do not know if the hash data or blob is valid so we only write the state code
                     // also for head checks we do not write anything more
                     if (opState != OpState.OK || opState == OpState.Exists || contentHash == null || blob == null)
+                    {
                         continue;
+                    }
 
                     await outStream.WriteAsync(contentHash.HashData);
                     await outStream.WriteAsync(BitConverter.GetBytes((ulong) blob.Length));
diff --git a/Engine/Source/Programs/Horde/HordeStorage/Jupiter.Common/Implementation/BlobContents.cs b/Engine/Source/Programs/Horde/HordeStorage/Jupiter.Common/Implementation/BlobContents.cs
--- a/Engine/Source/Programs/Horde/HordeStorage/Jupiter.Common/Implementation/BlobContents.cs
+++ b/Engine/Source/Programs/Horde/HordeStorage/Jupiter.Common/Implementation/BlobContents.cs
@@ -6,7 +6,7 @@ using System.Threading.Tasks;
 
 namespace Jupiter.Implementation
 {
-    public class BlobContents : IDisposable, IAsyncDisposable
+    public sealed class BlobContents : IDisposable, IAsyncDisposable
     {
         public BlobContents(Stream stream, long length)
         {
diff --git a/Engine/Source/Programs/Horde/HordeStorage/Jupiter.Common/Implementation/BlobIdentifier.cs b/Engine/Source/Programs/Horde/HordeStorage/Jupiter.Common/Implementation/BlobIdentifier.cs
--- a/Engine/Source/Programs/Horde/HordeStorage/Jupiter.Common/Implementation/BlobIdentifier.cs
+++ b/Engine/Source/Programs/Horde/HordeStorage/Jupiter.Common/Implementation/BlobIdentifier.cs
@@ -9,7 +9,6 @@ using Blake3;
 using EpicGames.Core;
 using EpicGames.Serialization;
 using Newtonsoft.Json;
-using SharpYaml.Tokens;
 using JsonWriter = Newtonsoft.Json.JsonWriter;
 
 namespace Jupiter.Implementation
@@ -41,7 +40,9 @@ namespace Jupiter.Implementation
         public bool Equals(BlobIdentifier? other)
         {
             if (other == null)
+            {
                 return false;
+            }
 
             return Comparer.Equals(Identifier, other.Identifier);
         }
@@ -58,7 +59,7 @@ namespace Jupiter.Implementation
                 return true;
             }
 
-            if (obj.GetType() != this.GetType())
+            if (obj.GetType() != GetType())
             {
                 return false;
             }
@@ -76,7 +77,7 @@ namespace Jupiter.Implementation
             return _stringIdentifier;
         }
 
-        public new static BlobIdentifier FromBlob(byte[] blobMemory)
+        public static new BlobIdentifier FromBlob(byte[] blobMemory)
         {
             Hash blake3Hash;
             if (blobMemory.Length < MultiThreadedSize)
@@ -165,9 +166,9 @@ namespace Jupiter.Implementation
   
         public override object? ConvertFrom(ITypeDescriptorContext? context, CultureInfo? culture, object value)  
         {
-            if (value is string)
+            if (value is string s)
             {
-                return new BlobIdentifier((string) value);
+                return new BlobIdentifier(s);
             }
 
             return base.ConvertFrom(context, culture, value);  
@@ -183,12 +184,18 @@ namespace Jupiter.Implementation
 
         public override BlobIdentifier? ReadJson(JsonReader reader, Type objectType, BlobIdentifier? existingValue, bool hasExistingValue, JsonSerializer serializer)
         {
-            if (reader.TokenType == JsonToken.Null) 
+            if (reader.TokenType == JsonToken.Null)
+            {
                 return null;
+            }
+
             string? s = (string?)reader.Value;
 
             if (s == null)
+            {
                 return null;
+            }
+
             return new BlobIdentifier(s!);
         }
     }
diff --git a/Engine/Source/Programs/Horde/HordeStorage/Jupiter.Common/Implementation/CallistoReader.cs b/Engine/Source/Programs/Horde/HordeStorage/Jupiter.Common/Implementation/CallistoReader.cs
--- a/Engine/Source/Programs/Horde/HordeStorage/Jupiter.Common/Implementation/CallistoReader.cs
+++ b/Engine/Source/Programs/Horde/HordeStorage/Jupiter.Common/Implementation/CallistoReader.cs
@@ -14,32 +14,32 @@ using RestSharp;
 
 namespace Jupiter.Implementation
 {
-    public class CallistoReader
+    public class CallistoGetResponse
     {
-        public class CallistoGetResponse
+        [JsonConstructor]
+        public CallistoGetResponse(List<TransactionEvent> events, Guid generation, long currentOffset)
         {
-            [JsonConstructor]
-            public CallistoGetResponse(List<TransactionEvent> events, Guid generation, long currentOffset)
-            {
-                Events = events;
-                Generation = generation;
-                CurrentOffset = currentOffset;
-            }
-
-            public CallistoGetResponse()
-            {
-                Events = null!;
-                Generation = Guid.Empty;
-                CurrentOffset = 0;
-            }
+            Events = events;
+            Generation = generation;
+            CurrentOffset = currentOffset;
+        }
 
-            public Guid Generation { get; set; }
-            public long CurrentOffset { get; }
-            public List<TransactionEvent> Events { get; set; }
+        public CallistoGetResponse()
+        {
+            Events = null!;
+            Generation = Guid.Empty;
+            CurrentOffset = 0;
         }
 
+        public Guid Generation { get; set; }
+        public long CurrentOffset { get; set; }
+        public List<TransactionEvent> Events { get; init; }
+    }
+
+	public class CallistoReader
+    {
         private readonly NamespaceId _ns;
-        private IRestClient _client;
+        private readonly IRestClient _client;
 
         public CallistoReader(IRestClient callistoClient, NamespaceId ns)
         {
@@ -56,7 +56,9 @@ namespace Jupiter.Implementation
             }
 
             if (cancellationToken.IsCancellationRequested)
+            {
                 yield break;
+            }
 
             bool hadNewEvents;
             long currentOffset = startingOffset;
@@ -74,23 +76,23 @@ namespace Jupiter.Implementation
                 }
 
                 if (currentBatch.Events == null)
+                {
                     yield break;
+                }
 
                 foreach (TransactionEvent transactionEvent in currentBatch.Events)
                 {
                     yield return transactionEvent;
                 }
-
-
             } while (hadNewEvents);
         }
 
         private async Task<Transactions> FetchTransactions(long startingOffset, int countOfEventsToFetch, string? siteFilter, Guid? replicatingGeneration, int? maxOffsetsAttempted, CancellationToken replicationToken)
         {
-            const int maxAttempts = 3;
+            const int MaxAttempts = 3;
             CallistoGetResponse? callistoGetResponse = null;
             IRestResponse<CallistoGetResponse>? response = null;
-            for (int i = 0; i < maxAttempts; i++)
+            for (int i = 0; i < MaxAttempts; i++)
             {
                 RestRequest request = new RestRequest("api/v1/t/{ns}/{offset}");
                 request.AddUrlSegment("ns", _ns);
@@ -99,13 +101,19 @@ namespace Jupiter.Implementation
                 request.AddQueryParameter("count", countOfEventsToFetch.ToString());
 
                 if (siteFilter != null)
+                {
                     request.AddQueryParameter("notSeenAtSite", siteFilter);
+                }
 
                 if (replicatingGeneration != null)
+                {
                     request.AddQueryParameter("expectedGeneration", replicatingGeneration.ToString()!);
+                }
 
                 if (maxOffsetsAttempted.HasValue)
+                {
                     request.AddQueryParameter("maxOffsetsAttempted", maxOffsetsAttempted.Value.ToString());
+                }
 
                 response = await _client.ExecuteGetAsync<CallistoGetResponse>(request, replicationToken);
                 if (response.StatusCode == HttpStatusCode.BadRequest)
@@ -130,7 +138,7 @@ namespace Jupiter.Implementation
                         maxOffsetsAttempted = maxOffsetsAttempted.GetValueOrDefault(100) * 10;
 
                         // if we reached max attempts throw a exception other retry with more offsets attempted
-                        if (i == maxAttempts - 1)
+                        if (i == MaxAttempts - 1)
                         {
                             string errorMsg = badRequestObj["title"]!.Value<string>()!;
                             throw new UnknownNamespaceException($"Transaction log mismatch in {_ns}. Error message {errorMsg}");
@@ -149,7 +157,9 @@ namespace Jupiter.Implementation
                 }
 
                 if (!response.IsSuccessful)
+                {
                     throw new Exception($"Unsuccessful response from remote callisto. Status Code: {response.StatusCode}. Url: {response.ResponseUri} . Message: {response.ErrorMessage}");
+                }
 
                 // we have data so stop retrying
                 callistoGetResponse = response.Data;
@@ -157,11 +167,15 @@ namespace Jupiter.Implementation
             }
 
             if (callistoGetResponse == null)
-                throw new Exception($"Unsuccessful response from remote callisto, gave up after {maxAttempts} attempts. Status Code: {response!.StatusCode}. Url: {response.ResponseUri} . Message: {response.ErrorMessage}");
+            {
+                throw new Exception($"Unsuccessful response from remote callisto, gave up after {MaxAttempts} attempts. Status Code: {response!.StatusCode}. Url: {response.ResponseUri} . Message: {response.ErrorMessage}");
+            }
 
             List<TransactionEvent> events = callistoGetResponse.Events;
             if (events == null)
+            {
                 throw new Exception($"Invalid response from remote callisto: {_client} expected objects with events array");
+            }
 
             if (events.Count == 0)
             {
diff --git a/Engine/Source/Programs/Horde/HordeStorage/Jupiter.Common/Implementation/ClientCredentialOAuthAuthenticator.cs b/Engine/Source/Programs/Horde/HordeStorage/Jupiter.Common/Implementation/ClientCredentialOAuthAuthenticator.cs
--- a/Engine/Source/Programs/Horde/HordeStorage/Jupiter.Common/Implementation/ClientCredentialOAuthAuthenticator.cs
+++ b/Engine/Source/Programs/Horde/HordeStorage/Jupiter.Common/Implementation/ClientCredentialOAuthAuthenticator.cs
@@ -9,7 +9,7 @@ namespace Jupiter.Implementation
 {
     public class ClientCredentialOAuthAuthenticator: IAuthenticator
     {
-        public ClientCredentialOAuthAuthenticator(string authUrl, string clientId, string clientSecret, string scope)
+        public ClientCredentialOAuthAuthenticator(Uri authUrl, string clientId, string clientSecret, string scope)
         {
             _authUrl = authUrl;
             _clientId = clientId;
@@ -19,7 +19,7 @@ namespace Jupiter.Implementation
 
         private string? _accessToken;
 
-        private readonly string _authUrl;
+        private readonly Uri _authUrl;
         private readonly string _clientId;
         private readonly string _clientSecret;
         private readonly string _scope;
@@ -78,13 +78,11 @@ namespace Jupiter.Implementation
                 }
                 throw;
             }
-
         }
 
         private IRestResponse<ClientCredentialsResponse> DoAuthenticationRequest()
         {
-            Uri requestUri = new Uri(_authUrl);
-            RestClient client = new RestClient(requestUri);
+            RestClient client = new RestClient(_authUrl);
             RestRequest request = new RestRequest(Method.POST);
 
             request.AddParameter("grant_type", "client_credentials", ParameterType.GetOrPost);
@@ -109,7 +107,7 @@ namespace Jupiter.Implementation
         public string? scope { get; set; }
     }
 
-    internal class AuthenticationFailedException : Exception
+    public class AuthenticationFailedException : Exception
     {
         public AuthenticationFailedException(object errorResult) : base(errorResult.ToString()) { }
     }
diff --git a/Engine/Source/Programs/Horde/HordeStorage/Jupiter.Common/Implementation/ContentHash.cs b/Engine/Source/Programs/Horde/HordeStorage/Jupiter.Common/Implementation/ContentHash.cs
--- a/Engine/Source/Programs/Horde/HordeStorage/Jupiter.Common/Implementation/ContentHash.cs
+++ b/Engine/Source/Programs/Horde/HordeStorage/Jupiter.Common/Implementation/ContentHash.cs
@@ -1,9 +1,6 @@
 ﻿// Copyright Epic Games, Inc. All Rights Reserved.
 
 using System;
-using System.IO;
-using System.Text;
-using System.Threading.Tasks;
 using Blake3;
 using EpicGames.Core;
 using EpicGames.Serialization;
@@ -16,32 +13,33 @@ namespace Jupiter.Implementation
     [CbConverter(typeof(ContentHashCbConverter))]
     public class ContentHash : IEquatable<ContentHash>, IEquatable<byte[]>
     {
-        protected readonly ByteArrayComparer Comparer = new ByteArrayComparer();
-        protected readonly byte[] Identifier;
-        public static int HashLength = 20;
-
+	    protected ByteArrayComparer Comparer { get; } = new ByteArrayComparer();
+        protected byte[] Identifier { get; init; }
+        public const int HashLength = 20;
+        public byte[] HashData => Identifier;
         public ContentHash(byte[] identifier)
         {
             Identifier = identifier;
             if (identifier.Length != HashLength)
+            {
                 throw new ArgumentException("Supplied identifier was not 20 bytes, this is not a valid identifier", nameof(identifier));
+            }
         }
 
         [JsonConstructor]
         public ContentHash(string identifier)
         {
             if (identifier == null)
+            {
                 throw new ArgumentNullException(nameof(identifier));
+            }
 
             byte[] byteIdentifier = StringUtils.ToHashFromHexString(identifier);
             Identifier = byteIdentifier;
             if (byteIdentifier.Length != HashLength)
+            {
                 throw new ArgumentException("Supplied identifier was not 20 bytes, this is not a valid identifier", nameof(identifier));
-        }
-
-        public byte[] HashData
-        {
-            get { return Identifier; }
+            }
         }
 
         public override int GetHashCode()
@@ -52,7 +50,9 @@ namespace Jupiter.Implementation
         public bool Equals(ContentHash? other)
         {
             if (other == null)
+            {
                 return false;
+            }
 
             return Comparer.Equals(Identifier, other.Identifier);
         }
@@ -60,7 +60,10 @@ namespace Jupiter.Implementation
         public bool Equals(byte[]? other)
         {
             if (other == null)
+            {
                 return false;
+            }
+
             return Comparer.Equals(Identifier, other);
         }
 
@@ -76,7 +79,7 @@ namespace Jupiter.Implementation
                 return true;
             }
 
-            if (obj.GetType() != this.GetType())
+            if (obj.GetType() != GetType())
             {
                 return false;
             }
diff --git a/Engine/Source/Programs/Horde/HordeStorage/Jupiter.Common/Implementation/ContentId.cs b/Engine/Source/Programs/Horde/HordeStorage/Jupiter.Common/Implementation/ContentId.cs
--- a/Engine/Source/Programs/Horde/HordeStorage/Jupiter.Common/Implementation/ContentId.cs
+++ b/Engine/Source/Programs/Horde/HordeStorage/Jupiter.Common/Implementation/ContentId.cs
@@ -31,7 +31,9 @@ namespace Jupiter.Implementation
         public bool Equals(ContentId? other)
         {
             if (other == null)
+            {
                 return false;
+            }
 
             return Comparer.Equals(Identifier, other.Identifier);
         }
@@ -48,7 +50,7 @@ namespace Jupiter.Implementation
                 return true;
             }
 
-            if (obj.GetType() != this.GetType())
+            if (obj.GetType() != GetType())
             {
                 return false;
             }
@@ -95,9 +97,9 @@ namespace Jupiter.Implementation
   
         public override object? ConvertFrom(ITypeDescriptorContext? context, CultureInfo? culture, object value)  
         {
-            if (value is string)
+            if (value is string s)
             {
-                return new ContentId((string) value);
+                return new ContentId(s);
             }
 
             return base.ConvertFrom(context, culture, value);  
@@ -113,12 +115,18 @@ namespace Jupiter.Implementation
 
         public override ContentId? ReadJson(JsonReader reader, Type objectType, ContentId? existingValue, bool hasExistingValue, JsonSerializer serializer)
         {
-            if (reader.TokenType == JsonToken.Null) 
+            if (reader.TokenType == JsonToken.Null)
+            {
                 return null;
+            }
+
             string? s = (string?)reader.Value;
 
             if (s == null)
+            {
                 return null;
+            }
+
             return new ContentId(s!);
         }
     }
diff --git a/Engine/Source/Programs/Horde/HordeStorage/Jupiter.Common/Implementation/InstrumentedStream.cs b/Engine/Source/Programs/Horde/HordeStorage/Jupiter.Common/Implementation/InstrumentedStream.cs
--- a/Engine/Source/Programs/Horde/HordeStorage/Jupiter.Common/Implementation/InstrumentedStream.cs
+++ b/Engine/Source/Programs/Horde/HordeStorage/Jupiter.Common/Implementation/InstrumentedStream.cs
@@ -1,12 +1,9 @@
 ﻿// Copyright Epic Games, Inc. All Rights Reserved.
 
 using System;
-using System.Collections.Generic;
 using System.IO;
-using System.Linq;
 using System.Threading;
 using System.Threading.Tasks;
-using Serilog;
 
 namespace Jupiter.Common.Implementation
 {
@@ -21,14 +18,19 @@ namespace Jupiter.Common.Implementation
         private long _totalBytesWritten = 0;
         private DateTime _consumeStartedAt;
 
-        private readonly ILogger _logger = Log.ForContext<InstrumentedStream>();
-
         public InstrumentedStream(Stream s, string sourceIdentifier)
         {
             _streamImplementation = s;
             _sourceIdentifier = sourceIdentifier;
         }
 
+        public override async ValueTask DisposeAsync()
+        {
+            await _streamImplementation.DisposeAsync();
+            await base.DisposeAsync();
+            GC.SuppressFinalize(this);
+        }
+
         public override void Close()
         {
             TimeSpan duration = DateTime.Now - _consumeStartedAt;
@@ -77,30 +79,18 @@ namespace Jupiter.Common.Implementation
             _streamImplementation.Write(buffer, offset, count);
         }
 
-        public override bool CanRead
-        {
-            get { return _streamImplementation.CanRead; }
-        }
+        public override bool CanRead => _streamImplementation.CanRead;
 
-        public override bool CanSeek
-        {
-            get { return _streamImplementation.CanSeek; }
-        }
+        public override bool CanSeek => _streamImplementation.CanSeek;
 
-        public override bool CanWrite
-        {
-            get { return _streamImplementation.CanWrite; }
-        }
+        public override bool CanWrite => _streamImplementation.CanWrite;
 
-        public override long Length
-        {
-            get { return _streamImplementation.Length; }
-        }
+        public override long Length => _streamImplementation.Length;
 
         public override long Position
         {
-            get { return _streamImplementation.Position; }
-            set { _streamImplementation.Position = value; }
+            get => _streamImplementation.Position;
+            set => _streamImplementation.Position = value;
         }
     }
 }
diff --git a/Engine/Source/Programs/Horde/HordeStorage/Jupiter.Common/Implementation/IoHashKey.cs b/Engine/Source/Programs/Horde/HordeStorage/Jupiter.Common/Implementation/IoHashKey.cs
--- a/Engine/Source/Programs/Horde/HordeStorage/Jupiter.Common/Implementation/IoHashKey.cs
+++ b/Engine/Source/Programs/Horde/HordeStorage/Jupiter.Common/Implementation/IoHashKey.cs
@@ -61,7 +61,7 @@ namespace Jupiter.Implementation
 
         public override int GetHashCode()
         {
-            return _text.GetHashCode();
+            return _text.GetHashCode(StringComparison.Ordinal);
         }
 
         public override string ToString()
@@ -106,7 +106,9 @@ namespace Jupiter.Implementation
         public override bool CanConvertFrom(ITypeDescriptorContext? context, Type sourceType)
         {
             if (sourceType == typeof(string))
+            {
                 return true;
+            }
 
             return base.CanConvertFrom(context, sourceType);
         }
diff --git a/Engine/Source/Programs/Horde/HordeStorage/Jupiter.Common/Implementation/KeyId.cs b/Engine/Source/Programs/Horde/HordeStorage/Jupiter.Common/Implementation/KeyId.cs
--- a/Engine/Source/Programs/Horde/HordeStorage/Jupiter.Common/Implementation/KeyId.cs
+++ b/Engine/Source/Programs/Horde/HordeStorage/Jupiter.Common/Implementation/KeyId.cs
@@ -42,7 +42,7 @@ namespace Jupiter.Implementation
 
         public override int GetHashCode()
         {
-            return _text.GetHashCode();
+            return _text.GetHashCode(StringComparison.Ordinal);
         }
 
         public override string ToString()
@@ -82,7 +82,9 @@ namespace Jupiter.Implementation
         public override bool CanConvertFrom(ITypeDescriptorContext? context, Type sourceType)
         {
             if (sourceType == typeof(string))
+            {
                 return true;
+            }
 
             return base.CanConvertFrom(context, sourceType);
         }
diff --git a/Engine/Source/Programs/Horde/HordeStorage/Jupiter.Common/Implementation/NamespacePolicyResolver.cs b/Engine/Source/Programs/Horde/HordeStorage/Jupiter.Common/Implementation/NamespacePolicyResolver.cs
--- a/Engine/Source/Programs/Horde/HordeStorage/Jupiter.Common/Implementation/NamespacePolicyResolver.cs
+++ b/Engine/Source/Programs/Horde/HordeStorage/Jupiter.Common/Implementation/NamespacePolicyResolver.cs
@@ -1,6 +1,5 @@
 // Copyright Epic Games, Inc. All Rights Reserved.
 
-using System;
 using EpicGames.Horde.Storage;
 using Jupiter.Implementation;
 using Microsoft.Extensions.Options;
@@ -9,12 +8,9 @@ namespace Jupiter.Common
 {
     public interface INamespacePolicyResolver
     {
-        public NamespaceSettings.PerNamespaceSettings GetPoliciesForNs(NamespaceId ns);
+        public NamespacePolicy GetPoliciesForNs(NamespaceId ns);
 
-        static NamespaceId JupiterInternalNamespace
-        {
-            get { return new NamespaceId("jupiter-internal"); }
-        }
+        static NamespaceId JupiterInternalNamespace => new NamespaceId("jupiter-internal");
     }
     public class NamespacePolicyResolver : INamespacePolicyResolver
     {
@@ -25,11 +21,11 @@ namespace Jupiter.Common
             _namespaceSettings = namespaceSettings;
         }
 
-        public NamespaceSettings.PerNamespaceSettings GetPoliciesForNs(NamespaceId ns)
+        public NamespacePolicy GetPoliciesForNs(NamespaceId ns)
         {
             if (ns == INamespacePolicyResolver.JupiterInternalNamespace)
             {
-                return new NamespaceSettings.PerNamespaceSettings()
+                return new NamespacePolicy()
                 {
                     // we grant the full storage access claim access to the internal namespace, this is typically set of admins and service accounts that can access everything
                     Claims = new string[] {"Storage=full" },
@@ -38,14 +34,14 @@ namespace Jupiter.Common
             }
 
             if (_namespaceSettings.CurrentValue.Policies.TryGetValue(ns.ToString(),
-                    out NamespaceSettings.PerNamespaceSettings? settings))
+                    out NamespacePolicy? settings))
             {
                 return settings;
             }
 
             // attempt to find the default mapping
             if (_namespaceSettings.CurrentValue.Policies.TryGetValue("*",
-                    out NamespaceSettings.PerNamespaceSettings? defaultSettings))
+                    out NamespacePolicy? defaultSettings))
             {
                 return defaultSettings;
             }
diff --git a/Engine/Source/Programs/Horde/HordeStorage/Jupiter.Common/Implementation/SecretResolver.cs b/Engine/Source/Programs/Horde/HordeStorage/Jupiter.Common/Implementation/SecretResolver.cs
--- a/Engine/Source/Programs/Horde/HordeStorage/Jupiter.Common/Implementation/SecretResolver.cs
+++ b/Engine/Source/Programs/Horde/HordeStorage/Jupiter.Common/Implementation/SecretResolver.cs
@@ -45,11 +45,11 @@ namespace Jupiter.Common.Implementation
 
         private string? ResolveUsingProvider(string providerId, string providerPath, string originalValue)
         {
-            switch (providerId.ToLowerInvariant())
+            switch (providerId.ToUpperInvariant())
             {
-                case "aws":
+                case "AWS":
                     return ResolveAWSSecret(providerPath);
-                case "akv":
+                case "AKV":
                     return ResolveAKVSecret(providerPath);
                 default:
                     // no provider matches so just return the original value
@@ -63,7 +63,9 @@ namespace Jupiter.Common.Implementation
 
             IAmazonSecretsManager? secretsManager = _serviceProvider.GetService<IAmazonSecretsManager>();
             if (secretsManager == null)
+            {
                 throw new Exception($"Unable to get AWSSecretsManager when resolving aws secret resource: {providerPath}");
+            }
 
             string secretValue;
             try
@@ -81,14 +83,21 @@ namespace Jupiter.Common.Implementation
             }
 
             if (key == null)
+            {
                 return secretValue;
+            }
 
             Dictionary<string, string>? keyCollection = JsonSerializer.Deserialize<Dictionary<string, string>>(secretValue);
 
             if (keyCollection == null)
+            {
                 throw new Exception($"Unable to deserialize secret to a json payload for path: {providerPath}");
+            }
+
             if (keyCollection.TryGetValue(key, out string? s))
+            {
                 return s;
+            }
 
             throw new Exception($"Unable to find key {key} in blob returned for secret {arn}");
         }
@@ -102,7 +111,7 @@ namespace Jupiter.Common.Implementation
         /// <param name="providerPath"></param>
         /// <returns></returns>
         /// <exception cref="Exception"></exception>
-        private string? ResolveAKVSecret(string providerPath)
+        private static string? ResolveAKVSecret(string providerPath)
         {
             if (!SplitByFirstSeparator(providerPath, '|', out string vaultName, out string? secretName))
             {
@@ -118,7 +127,7 @@ namespace Jupiter.Common.Implementation
 
         private static bool SplitByFirstSeparator(string fullPath, char sep, out string left, out string? right)
         {
-            int keySeparator = fullPath.IndexOf(sep);
+            int keySeparator = fullPath.IndexOf(sep, StringComparison.OrdinalIgnoreCase);
             left = fullPath;
             right = null;
             if (keySeparator != -1)
diff --git a/Engine/Source/Programs/Horde/HordeStorage/Jupiter.Common/Implementation/StreamingBodyPayload.cs b/Engine/Source/Programs/Horde/HordeStorage/Jupiter.Common/Implementation/StreamingBodyPayload.cs
--- a/Engine/Source/Programs/Horde/HordeStorage/Jupiter.Common/Implementation/StreamingBodyPayload.cs
+++ b/Engine/Source/Programs/Horde/HordeStorage/Jupiter.Common/Implementation/StreamingBodyPayload.cs
@@ -18,9 +18,9 @@ namespace Jupiter.Common.Implementation
     /// <summary>
     /// Streaming request that is streamed into memory
     /// </summary>
-    public class MemoryBufferedPayload : IBufferedPayload
+    public sealed class MemoryBufferedPayload : IBufferedPayload
     {
-        private byte[] _buffer;
+        private readonly byte[] _buffer;
 
         public MemoryBufferedPayload(byte[] source)
         {
@@ -44,13 +44,13 @@ namespace Jupiter.Common.Implementation
             return new MemoryStream(_buffer);
         }
 
-        public long Length { get { return _buffer.LongLength; } }
+        public long Length => _buffer.LongLength;
     }
 
     /// <summary>
     /// A streaming request backed by a temporary file on disk
     /// </summary>
-    public class FilesystemBufferedPayload : IBufferedPayload
+    public sealed class FilesystemBufferedPayload : IBufferedPayload
     {
         private readonly FileInfo _tempFile;
         private long _length;
@@ -84,7 +84,9 @@ namespace Jupiter.Common.Implementation
         public void Dispose()
         {
             if (_tempFile.Exists)
+            {
                 _tempFile.Delete();
+            }
         }
 
         public Stream GetStream()
@@ -92,10 +94,7 @@ namespace Jupiter.Common.Implementation
             return _tempFile.OpenRead();
         }
 
-        public long Length
-        {
-            get { return _length; }
-        }
+        public long Length => _length;
 
         public static FilesystemBufferedPayload FromTempFile(FileInfo tempFile)
         {
diff --git a/Engine/Source/Programs/Horde/HordeStorage/Jupiter.Common/Implementation/TreeHash.cs b/Engine/Source/Programs/Horde/HordeStorage/Jupiter.Common/Implementation/TreeHash.cs
--- a/Engine/Source/Programs/Horde/HordeStorage/Jupiter.Common/Implementation/TreeHash.cs
+++ b/Engine/Source/Programs/Horde/HordeStorage/Jupiter.Common/Implementation/TreeHash.cs
@@ -1,7 +1,6 @@
 ﻿// Copyright Epic Games, Inc. All Rights Reserved.
 
 using System;
-using System.Collections.Generic;
 using System.ComponentModel;
 using System.Globalization;
 using Blake3;
@@ -30,7 +29,10 @@ namespace Jupiter.Implementation
         public bool Equals(TreeHash? other)
         {
             if (other == null)
+            {
                 return false;
+            }
+
             return Comparer.Equals(Identifier, other.Identifier);
         }
 
@@ -46,7 +48,7 @@ namespace Jupiter.Implementation
                 return true;
             }
 
-            if (obj.GetType() != this.GetType())
+            if (obj.GetType() != GetType())
             {
                 return false;
             }
@@ -93,14 +95,15 @@ namespace Jupiter.Implementation
         }
     }
 
-
     public sealed class TreeHashTypeConverter : TypeConverter
     {
         /// <inheritdoc/>
         public override bool CanConvertFrom(ITypeDescriptorContext? context, Type sourceType)
         {
             if (sourceType == typeof(string))
+            {
                 return true;
+            }
 
             return base.CanConvertFrom(context, sourceType);
         }
diff --git a/Engine/Source/Programs/Horde/HordeStorage/Jupiter.Common/PollingService.cs b/Engine/Source/Programs/Horde/HordeStorage/Jupiter.Common/PollingService.cs
--- a/Engine/Source/Programs/Horde/HordeStorage/Jupiter.Common/PollingService.cs
+++ b/Engine/Source/Programs/Horde/HordeStorage/Jupiter.Common/PollingService.cs
@@ -49,15 +49,9 @@ namespace Jupiter
             }, -1, -1);
         }
 
-        public bool Running
-        {
-            get { return _timer != null; }
-        }
+        public bool Running => _timer != null;
 
-        public T State
-        {
-            get { return _state; }
-        }
+        public T State => _state;
 
         protected virtual bool ShouldStartPolling()
         {
@@ -88,7 +82,7 @@ namespace Jupiter
             ThreadState? ts = (ThreadState?)state;
             if (ts == null)
             {
-                throw new ArgumentNullException(nameof(ts));
+                throw new ArgumentNullException("", "Null thread state passed to polling service");
             }
             ThreadState threadState = ts.Value;
 
@@ -98,7 +92,9 @@ namespace Jupiter
             ILogger logger = Log.ForContext<PollingService<T>>();
 
             if (instance._alreadyRunning)
+            {
                 return;
+            }
 
             try
             {
@@ -106,7 +102,9 @@ namespace Jupiter
                 instance._hasFinishedRunning.Reset();
 
                 if (stopPollingToken.IsCancellationRequested)
+                {
                     return;
+                }
 
                 bool _ = instance.OnPoll(threadState.ServiceState, stopPollingToken).Result;
             }
@@ -140,6 +138,7 @@ namespace Jupiter
             finally
             {
                 instance._alreadyRunning = false;
+                instance._hasFinishedRunning.Set();
             }
         }
 
@@ -155,7 +154,10 @@ namespace Jupiter
             _logger.Information("{Service} poll service stopping.", _serviceName);
 
             if (_timer != null)
+            {
                 await _timer.DisposeAsync();
+            }
+
             _timer = null;
 
             await OnStopping(_state);
@@ -167,10 +169,18 @@ namespace Jupiter
         public async ValueTask DisposeAsync()
         {
             if (_disposed)
+            {
                 return;
+            }
 
             _disposed = true;
+
             await StopAsync(CancellationToken.None);
+
+            _stopPolling.Dispose();
+            _hasFinishedRunning.Dispose();
+
+            GC.SuppressFinalize(this);
         }
     }
 }
diff --git a/Engine/Source/Programs/Horde/HordeStorage/Jupiter.Common/RawOutputFormatter.cs b/Engine/Source/Programs/Horde/HordeStorage/Jupiter.Common/RawOutputFormatter.cs
--- a/Engine/Source/Programs/Horde/HordeStorage/Jupiter.Common/RawOutputFormatter.cs
+++ b/Engine/Source/Programs/Horde/HordeStorage/Jupiter.Common/RawOutputFormatter.cs
@@ -24,7 +24,9 @@ namespace Jupiter
         public override async Task WriteResponseBodyAsync(OutputFormatterWriteContext context)
         {
             if (context.Object == null)
+            {
                 return;
+            }
 
             HttpResponse response = context.HttpContext.Response;
 
@@ -74,7 +76,8 @@ namespace Jupiter
     /// <summary>
     ///  Used to indicate that the field in question is what we want to return as a raw set of bytes
     /// </summary>
-    public class RawOutputPropertyAttribute : Attribute
+    [AttributeUsage(AttributeTargets.Property)]
+    public sealed class RawOutputPropertyAttribute : Attribute
     {
     }
 }
diff --git a/Engine/Source/Programs/Horde/HordeStorage/Jupiter.Common/RequiredContentTypeAttribute.cs b/Engine/Source/Programs/Horde/HordeStorage/Jupiter.Common/RequiredContentTypeAttribute.cs
--- a/Engine/Source/Programs/Horde/HordeStorage/Jupiter.Common/RequiredContentTypeAttribute.cs
+++ b/Engine/Source/Programs/Horde/HordeStorage/Jupiter.Common/RequiredContentTypeAttribute.cs
@@ -7,7 +7,8 @@ using Microsoft.Extensions.Primitives;
 
 namespace Jupiter
 {
-    public class RequiredContentTypeAttribute : Attribute, IActionConstraint
+    [AttributeUsage(AttributeTargets.Method)]
+    public sealed class RequiredContentTypeAttribute : Attribute, IActionConstraint
     {
         private readonly string _mediaTypeName;
 
@@ -16,10 +17,9 @@ namespace Jupiter
             _mediaTypeName = mediaTypeName;
         }
 
-        public int Order
-        {
-            get { return 0; }
-        }
+        public string MediaTypeName => _mediaTypeName;
+
+        public int Order => 0;
 
         public bool Accept(ActionConstraintContext context)
         {
diff --git a/Engine/Source/Programs/Horde/HordeStorage/Jupiter.Common/Utils/OodleCompressor.cs b/Engine/Source/Programs/Horde/HordeStorage/Jupiter.Common/Utils/OodleCompressor.cs
--- a/Engine/Source/Programs/Horde/HordeStorage/Jupiter.Common/Utils/OodleCompressor.cs
+++ b/Engine/Source/Programs/Horde/HordeStorage/Jupiter.Common/Utils/OodleCompressor.cs
@@ -7,6 +7,12 @@ using IntPtr = System.IntPtr;
 
 namespace Jupiter.Utils
 {
+
+#pragma warning disable CA1008 // Add a member that has a value of zero with a suggested name of None
+#pragma warning disable CA1027 // Mark enums with FlagsAttribute
+#pragma warning disable CA1707 // Identifiers should not contain underscores. We use the same identifiers as Oodle uses in its source
+#pragma warning disable CA1712 // Do not prefix enum values with the name of the enum type.  We use the same identifiers as Oodle uses in its source
+
     public enum OodleLZ_Compressor
     {
         OodleLZ_Compressor_Invalid = -1,
@@ -51,10 +57,10 @@ namespace Jupiter.Utils
         OodleLZ_CompressionLevel_HyperFast1=-1, // faster than SuperFast, less compression
         OodleLZ_CompressionLevel_HyperFast2=-2, // faster than HyperFast1, less compression
         OodleLZ_CompressionLevel_HyperFast3=-3, // faster than HyperFast2, less compression
-        OodleLZ_CompressionLevel_HyperFast4=-4, // fastest, less compression
-
+        OodleLZ_CompressionLevel_HyperFast4 =-4, // fastest, less compression
+        
         // aliases :
-        OodleLZ_CompressionLevel_HyperFast=OodleLZ_CompressionLevel_HyperFast1, // alias hyperfast base level
+        OodleLZ_CompressionLevel_HyperFast =OodleLZ_CompressionLevel_HyperFast1, // alias hyperfast base level
         OodleLZ_CompressionLevel_Optimal = OodleLZ_CompressionLevel_Optimal2,   // alias optimal standard level
         OodleLZ_CompressionLevel_Max     = OodleLZ_CompressionLevel_Optimal5,   // maximum compression level
         OodleLZ_CompressionLevel_Min     = OodleLZ_CompressionLevel_HyperFast4, // fastest compression level
@@ -92,6 +98,10 @@ namespace Jupiter.Utils
         OodleLZ_Verbosity_Lots = 3,
         OodleLZ_Verbosity_Force32 = 0x40000000
     }
+#pragma warning restore CA1008 // Add a member that has a value of zero with a suggested name of None
+#pragma warning restore CA1027 // Mark enums with FlagsAttribute
+#pragma warning restore CA1707 // Identifiers should not contain underscores
+#pragma warning restore CA1712 // Do not prefix enum values with the name of the enum type
 
     [StructLayout(LayoutKind.Sequential)]
     public class OodleConfigValues
@@ -176,7 +186,9 @@ namespace Jupiter.Utils
         public long Compress(OodleLZ_Compressor compressor, byte[] rawBuf, OodleLZ_CompressionLevel compressionLevel, out byte[] compBuf)
         {
             if (!_initialized)
+            {
                 throw new InvalidOperationException("Initialize Oodle before using it");
+            }
 
             long bufferSizeNeeded = OodleLZ_GetCompressedBufferSizeNeeded(compressor, rawBuf.LongLength);
             compBuf = new byte[bufferSizeNeeded];
@@ -185,10 +197,12 @@ namespace Jupiter.Utils
             return result;
         }
 
-        public long Decompress(OodleLZ_Compressor compressor, byte[] compBuf, OodleLZ_CompressionLevel compressionLevel, long uncompressedSize, out byte[] rawBuf)
+        public long Decompress(byte[] compBuf, long uncompressedSize, out byte[] rawBuf)
         {
             if (!_initialized)
+            {
                 throw new InvalidOperationException("Initialize Oodle before using it");
+            }
 
             rawBuf = new byte[uncompressedSize];
 
@@ -206,13 +220,21 @@ namespace Jupiter.Utils
 
         public void Dispose()
         {
-            ReleaseUnmanagedResources();
+            Dispose(true);
             GC.SuppressFinalize(this);
         }
 
+        protected virtual void Dispose(bool disposing)
+        {
+            if (disposing)
+            {
+                ReleaseUnmanagedResources();
+            }
+        }
+
         ~OodleCompressor()
         {
-            ReleaseUnmanagedResources();
+            Dispose(false);
         }
     }
 }
diff --git a/Engine/Source/Programs/Horde/HordeStorage/Jupiter.Common/Utils/Sha1Utils.cs b/Engine/Source/Programs/Horde/HordeStorage/Jupiter.Common/Utils/Sha1Utils.cs
--- a/Engine/Source/Programs/Horde/HordeStorage/Jupiter.Common/Utils/Sha1Utils.cs
+++ b/Engine/Source/Programs/Horde/HordeStorage/Jupiter.Common/Utils/Sha1Utils.cs
@@ -43,18 +43,31 @@ namespace Jupiter
             }
 
             if (b == null)
+            {
+                return false;
+            }
+
+            if (a.Length != b.Length)
+            {
                 return false;
+            }
 
-            if (a.Length != b.Length) return false;
             for (int i = 0; i < a.Length; i++)
-                if (a[i] != b[i]) return false;
+            {
+                if (a[i] != b[i])
+                {
+                    return false;
+                }
+            }
             return true;
         }
         public override int GetHashCode(byte[] a)
         {
             uint b = 0;
             foreach (byte t in a)
-                b = ((b << 23) | (b >> 9)) ^ t;
+            {
+	            b = ((b << 23) | (b >> 9)) ^ t;
+            }
 
             return unchecked((int)b);
         }
diff --git a/Engine/Source/Programs/Horde/HordeStorage/Jupiter.Common/Utils/StreamUtils.cs b/Engine/Source/Programs/Horde/HordeStorage/Jupiter.Common/Utils/StreamUtils.cs
--- a/Engine/Source/Programs/Horde/HordeStorage/Jupiter.Common/Utils/StreamUtils.cs
+++ b/Engine/Source/Programs/Horde/HordeStorage/Jupiter.Common/Utils/StreamUtils.cs
@@ -1,7 +1,5 @@
 ﻿// Copyright Epic Games, Inc. All Rights Reserved.
 
-using System;
-using System.Collections.Generic;
 using System.IO;
 using System.Threading.Tasks;
 
diff --git a/Engine/Source/Programs/Horde/HordeStorage/Jupiter.Common/Utils/TaskUtils.cs b/Engine/Source/Programs/Horde/HordeStorage/Jupiter.Common/Utils/TaskUtils.cs
--- a/Engine/Source/Programs/Horde/HordeStorage/Jupiter.Common/Utils/TaskUtils.cs
+++ b/Engine/Source/Programs/Horde/HordeStorage/Jupiter.Common/Utils/TaskUtils.cs
@@ -11,7 +11,10 @@ namespace Jupiter
         public static byte[] CombineIntoSingleBuffer(List<byte[]> tasks)
         {
             if (tasks.Count == 1)
+            {
                 return tasks[0];
+            }
+
             int totalLength = tasks.Sum(task => task.Length);
             byte[] buffer = new byte[totalLength];
             int index = 0;
diff --git a/Engine/Source/Programs/Shared/EpicGames.AspNet/CbConvertersAspnet.cs b/Engine/Source/Programs/Shared/EpicGames.AspNet/CbConvertersAspnet.cs
--- a/Engine/Source/Programs/Shared/EpicGames.AspNet/CbConvertersAspnet.cs
+++ b/Engine/Source/Programs/Shared/EpicGames.AspNet/CbConvertersAspnet.cs
@@ -1,72 +1,86 @@
 // Copyright Epic Games, Inc. All Rights Reserved.
 
+using System;
 using EpicGames.Core;
 using EpicGames.Serialization;
 using Microsoft.AspNetCore.Mvc;
 
 namespace EpicGames.AspNet
 {
-	public static class CbConvertersAspNet
-	{
-		public static void AddAspnetConverters()
-		{
-			CbConverter.TypeToConverter[typeof(ProblemDetails)] = new CbProblemDetailsConverter();
-		}
-	}
+    public static class CbConvertersAspNet
+    {
+        public static void AddAspnetConverters()
+        {
+            CbConverter.TypeToConverter[typeof(ProblemDetails)] = new CbProblemDetailsConverter();
+        }
+    }
 
-	/// <summary>
-	/// Converter for asp.net problem details type
-	/// </summary>
-	/// <typeparam name="T"></typeparam>
-	class CbProblemDetailsConverter : CbConverterBase<ProblemDetails>
-	{
-		/// <inheritdoc/>
-		public override ProblemDetails Read(CbField field)
-		{
-			if (!field.IsObject())
-				throw new CbException($"Error converting field \"{field.Name}\" to ProblemDetails. Expected CbObject.");
-			
-			ProblemDetails result = new ProblemDetails
-			{
-				Title = field["title"].AsString(),
-				Detail = field["detail"].AsString(),
-				Type = field["type"].AsString(),
-				Instance = field["instance"].AsString(),
-				Status = field["status"].AsInt32()
-			};
-			return result;
-		}
+    /// <summary>
+    /// Converter for asp.net problem details type
+    /// </summary>
+    /// <typeparam name="T"></typeparam>
+    class CbProblemDetailsConverter : CbConverterBase<ProblemDetails>
+    {
+        /// <inheritdoc/>
+        public override ProblemDetails Read(CbField field)
+        {
+            if (!field.IsObject())
+            {
+                throw new CbException($"Error converting field \"{field.Name}\" to ProblemDetails. Expected CbObject.");
+            }
+            
+            ProblemDetails result = new ProblemDetails
+            {
+                Title = field["title"].AsString(),
+                Detail = field["detail"].AsString(),
+                Type = field["type"].AsString(),
+                Instance = field["instance"].AsString(),
+                Status = field["status"].AsInt32()
+            };
+            return result;
+        }
 
-		/// <inheritdoc/>
-		public override void Write(CbWriter writer, ProblemDetails problemDetails)
-		{
-			writer.WriteObject(ToCbObject(problemDetails));
-		}
+        /// <inheritdoc/>
+        public override void Write(CbWriter writer, ProblemDetails problemDetails)
+        {
+            writer.WriteObject(ToCbObject(problemDetails));
+        }
 
-		/// <inheritdoc/>
-		public override void WriteNamed(CbWriter writer, Utf8String name, ProblemDetails problemDetails)
-		{
-			writer.WriteField(name, ToCbObject(problemDetails).AsField());
-		}
+        /// <inheritdoc/>
+        public override void WriteNamed(CbWriter writer, Utf8String name, ProblemDetails problemDetails)
+        {
+            writer.WriteField(name, ToCbObject(problemDetails).AsField());
+        }
 
-		private CbObject ToCbObject(ProblemDetails problemDetails)
-		{
-			CbWriter objectWriter = new CbWriter();
-			objectWriter.BeginObject();
+        private CbObject ToCbObject(ProblemDetails problemDetails)
+        {
+            CbWriter objectWriter = new CbWriter();
+            objectWriter.BeginObject();
 
-			objectWriter.WriteString("title", problemDetails.Title);
-			if (!string.IsNullOrEmpty(problemDetails.Detail))
-				objectWriter.WriteString("detail", problemDetails.Detail);
-			if (!string.IsNullOrEmpty(problemDetails.Type))
-				objectWriter.WriteString("type", problemDetails.Type);
-			if (!string.IsNullOrEmpty(problemDetails.Instance))
-				objectWriter.WriteString("instance", problemDetails.Instance);
-			if (problemDetails.Status.HasValue)
-				objectWriter.WriteInteger("status", problemDetails.Status.Value);
-			
-			objectWriter.EndObject();
+            objectWriter.WriteString("title", problemDetails.Title);
+            if (!String.IsNullOrEmpty(problemDetails.Detail))
+            {
+                objectWriter.WriteString("detail", problemDetails.Detail);
+            }
 
-			return objectWriter.ToObject();
-		}
-	}
+            if (!String.IsNullOrEmpty(problemDetails.Type))
+            {
+                objectWriter.WriteString("type", problemDetails.Type);
+            }
+
+            if (!String.IsNullOrEmpty(problemDetails.Instance))
+            {
+                objectWriter.WriteString("instance", problemDetails.Instance);
+            }
+
+            if (problemDetails.Status.HasValue)
+            {
+                objectWriter.WriteInteger("status", problemDetails.Status.Value);
+            }
+            
+            objectWriter.EndObject();
+
+            return objectWriter.ToObject();
+        }
+    }
 }
