Horde Storage - Added option to control if namespaces have to be required from non-public endpoints (meaning not exposed to the internet).
Namespaces are public by default for backwards compatability.

These non-public endpoints is mostly intended for buildfarm only storage - were we are guaranteed to be on a appopriate network.

diff --git a/Engine/Source/Programs/Horde/HordeStorage/Horde.Storage/Controllers/AuthController.cs b/Engine/Source/Programs/Horde/HordeStorage/Horde.Storage/Controllers/AuthController.cs
--- a/Engine/Source/Programs/Horde/HordeStorage/Horde.Storage/Controllers/AuthController.cs
+++ b/Engine/Source/Programs/Horde/HordeStorage/Horde.Storage/Controllers/AuthController.cs
@@ -13,11 +13,11 @@ namespace Horde.Storage.Controllers
     [Route("api/v1/auth")]
     public class AuthController : ControllerBase
     {
-        private readonly IAuthorizationService _authorizationService;
+        private readonly RequestHelper _requestHelper;
 
-        public AuthController(IAuthorizationService authorizationService)
+        public AuthController(RequestHelper requestHelper)
         {
-            _authorizationService = authorizationService;
+            _requestHelper = requestHelper;
         }
 
         [HttpGet("{ns}")]
@@ -26,11 +26,10 @@ namespace Horde.Storage.Controllers
             [FromRoute] [Required] NamespaceId ns
             )
         {
-            AuthorizationResult authorizationResult = await _authorizationService.AuthorizeAsync(User, ns, NamespaceAccessRequirement.Name);
-
-            if (!authorizationResult.Succeeded)
+            ActionResult? result = await _requestHelper.HasAccessToNamespace(User, Request, ns);
+            if (result != null)
             {
-                return Forbid();
+                return result;
             }
 
             return Ok();
diff --git a/Engine/Source/Programs/Horde/HordeStorage/Horde.Storage/Controllers/CompressedBlobController.cs b/Engine/Source/Programs/Horde/HordeStorage/Horde.Storage/Controllers/CompressedBlobController.cs
--- a/Engine/Source/Programs/Horde/HordeStorage/Horde.Storage/Controllers/CompressedBlobController.cs
+++ b/Engine/Source/Programs/Horde/HordeStorage/Horde.Storage/Controllers/CompressedBlobController.cs
@@ -31,16 +31,16 @@ namespace Horde.Storage.Controllers
         private readonly IBlobService _storage;
         private readonly IContentIdStore _contentIdStore;
         private readonly IDiagnosticContext _diagnosticContext;
-        private readonly IAuthorizationService _authorizationService;
+        private readonly RequestHelper _requestHelper;
         private readonly CompressedBufferUtils _compressedBufferUtils;
         private readonly BufferedPayloadFactory _bufferedPayloadFactory;
 
-        public CompressedBlobController(IBlobService storage, IContentIdStore contentIdStore, IDiagnosticContext diagnosticContext, IAuthorizationService authorizationService, CompressedBufferUtils compressedBufferUtils, BufferedPayloadFactory bufferedPayloadFactory)
+        public CompressedBlobController(IBlobService storage, IContentIdStore contentIdStore, IDiagnosticContext diagnosticContext, RequestHelper requestHelper, CompressedBufferUtils compressedBufferUtils, BufferedPayloadFactory bufferedPayloadFactory)
         {
             _storage = storage;
             _contentIdStore = contentIdStore;
             _diagnosticContext = diagnosticContext;
-            _authorizationService = authorizationService;
+            _requestHelper = requestHelper;
             _compressedBufferUtils = compressedBufferUtils;
             _bufferedPayloadFactory = bufferedPayloadFactory;
         }
@@ -55,11 +55,10 @@ namespace Horde.Storage.Controllers
             [Required] NamespaceId ns,
             [Required] ContentId id)
         {
-            AuthorizationResult authorizationResult = await _authorizationService.AuthorizeAsync(User, ns, NamespaceAccessRequirement.Name);
-
-            if (!authorizationResult.Succeeded)
+            ActionResult? result = await _requestHelper.HasAccessToNamespace(User, Request, ns);
+            if (result != null)
             {
-                return Forbid();
+                return result;
             }
 
             try
@@ -91,12 +90,12 @@ namespace Horde.Storage.Controllers
             [Required] NamespaceId ns,
             [Required] ContentId id)
         {
-            AuthorizationResult authorizationResult = await _authorizationService.AuthorizeAsync(User, ns, NamespaceAccessRequirement.Name);
-
-            if (!authorizationResult.Succeeded)
+            ActionResult? result = await _requestHelper.HasAccessToNamespace(User, Request, ns);
+            if (result != null)
             {
-                return Forbid();
+                return result;
             }
+
             BlobIdentifier[]? chunks = await _contentIdStore.Resolve(ns, id, mustBeContentId: false);
             if (chunks == null || chunks.Length == 0)
             {
@@ -128,11 +127,10 @@ namespace Horde.Storage.Controllers
             [Required] NamespaceId ns,
             [Required] [FromQuery] List<ContentId> id)
         {
-            AuthorizationResult authorizationResult = await _authorizationService.AuthorizeAsync(User, ns, NamespaceAccessRequirement.Name);
-
-            if (!authorizationResult.Succeeded)
+            ActionResult? result = await _requestHelper.HasAccessToNamespace(User, Request, ns);
+            if (result != null)
             {
-                return Forbid();
+                return result;
             }
 
             ConcurrentBag<ContentId> partialContentIds = new ConcurrentBag<ContentId>();
@@ -172,11 +170,10 @@ namespace Horde.Storage.Controllers
             [Required] NamespaceId ns,
             [FromBody] ContentId[] bodyIds)
         {
-            AuthorizationResult authorizationResult = await _authorizationService.AuthorizeAsync(User, ns, NamespaceAccessRequirement.Name);
-
-            if (!authorizationResult.Succeeded)
+            ActionResult? result = await _requestHelper.HasAccessToNamespace(User, Request, ns);
+            if (result != null)
             {
-                return Forbid();
+                return result;
             }
 
             ConcurrentBag<ContentId> partialContentIds = new ConcurrentBag<ContentId>();
@@ -261,11 +258,10 @@ namespace Horde.Storage.Controllers
             [Required] NamespaceId ns,
             [Required] ContentId id)
         {
-            AuthorizationResult authorizationResult = await _authorizationService.AuthorizeAsync(User, ns, NamespaceAccessRequirement.Name);
-
-            if (!authorizationResult.Succeeded)
+            ActionResult? result = await _requestHelper.HasAccessToNamespace(User, Request, ns);
+            if (result != null)
             {
-                return Forbid();
+                return result;
             }
 
             _diagnosticContext.Set("Content-Length", Request.ContentLength ?? -1);
@@ -328,11 +324,10 @@ namespace Horde.Storage.Controllers
             [Required] string ns,
             [Required] BlobIdentifier id)
         {
-            AuthorizationResult authorizationResult = await _authorizationService.AuthorizeAsync(User, ns, NamespaceAccessRequirement.Name);
-
-            if (!authorizationResult.Succeeded)
+            ActionResult? result = await _requestHelper.HasAccessToNamespace(User, Request, ns);
+            if (result != null)
             {
-                return Forbid();
+                return result;
             }
 
             await DeleteImpl(ns, id);
@@ -346,11 +341,10 @@ namespace Horde.Storage.Controllers
         public async Task<IActionResult> DeleteNamespace(
             [Required] string ns)
         {
-            AuthorizationResult authorizationResult = await _authorizationService.AuthorizeAsync(User, ns, NamespaceAccessRequirement.Name);
-
-            if (!authorizationResult.Succeeded)
+            ActionResult? result = await _requestHelper.HasAccessToNamespace(User, Request, ns);
+            if (result != null)
             {
-                return Forbid();
+                return result;
             }
 
             int deletedCount = await  _storage.DeleteNamespace(ns);
@@ -365,7 +359,7 @@ namespace Horde.Storage.Controllers
         }*/
     }
 
-    
+
     public class ExistCheckMultipleContentIdResponse
     {
         public ContentId[] Needs { get; set; } = null!;
diff --git a/Engine/Source/Programs/Horde/HordeStorage/Horde.Storage/Controllers/ContentIdController.cs b/Engine/Source/Programs/Horde/HordeStorage/Horde.Storage/Controllers/ContentIdController.cs
--- a/Engine/Source/Programs/Horde/HordeStorage/Horde.Storage/Controllers/ContentIdController.cs
+++ b/Engine/Source/Programs/Horde/HordeStorage/Horde.Storage/Controllers/ContentIdController.cs
@@ -18,12 +18,12 @@ namespace Horde.Storage.Controllers
     [Route("api/v1/content-id")]
     public class ContentIdController : ControllerBase
     {
-        private readonly IAuthorizationService _authorizationService;
+        private readonly RequestHelper _requestHelper;
         private readonly IContentIdStore _contentIdStore;
 
-        public ContentIdController(IAuthorizationService authorizationService, IContentIdStore contentIdStore)
+        public ContentIdController(RequestHelper requestHelper, IContentIdStore contentIdStore)
         {
-            _authorizationService = authorizationService;
+            _requestHelper = requestHelper;
             _contentIdStore = contentIdStore;
         }
 
@@ -38,14 +38,10 @@ namespace Horde.Storage.Controllers
         [Authorize("Cache.read")]
         public async Task<IActionResult> Resolve(NamespaceId ns, ContentId contentId)
         {
-            using (IScope _ = Tracer.Instance.StartActive("authorize"))
+            ActionResult? result = await _requestHelper.HasAccessToNamespace(User, Request, ns);
+            if (result != null)
             {
-                AuthorizationResult authorizationResult = await _authorizationService.AuthorizeAsync(User, ns, NamespaceAccessRequirement.Name);
-
-                if (!authorizationResult.Succeeded)
-                {
-                    return Forbid();
-                }
+                return result;
             }
             BlobIdentifier[]? blobs = await _contentIdStore.Resolve(ns, contentId, mustBeContentId: true);
 
@@ -70,14 +66,10 @@ namespace Horde.Storage.Controllers
         [Authorize("Cache.write")]
         public async Task<IActionResult> UpdateContentIdMapping(NamespaceId ns, ContentId contentId, BlobIdentifier blobIdentifier, int contentWeight)
         {
-            using (IScope _ = Tracer.Instance.StartActive("authorize"))
+            ActionResult? result = await _requestHelper.HasAccessToNamespace(User, Request, ns);
+            if (result != null)
             {
-                AuthorizationResult authorizationResult = await _authorizationService.AuthorizeAsync(User, ns, NamespaceAccessRequirement.Name);
-
-                if (!authorizationResult.Succeeded)
-                {
-                    return Forbid();
-                }
+                return result;
             }
             await _contentIdStore.Put(ns, contentId, blobIdentifier, contentWeight);
 
diff --git a/Engine/Source/Programs/Horde/HordeStorage/Horde.Storage/Controllers/DDCRefController.cs b/Engine/Source/Programs/Horde/HordeStorage/Horde.Storage/Controllers/DDCRefController.cs
--- a/Engine/Source/Programs/Horde/HordeStorage/Horde.Storage/Controllers/DDCRefController.cs
+++ b/Engine/Source/Programs/Horde/HordeStorage/Horde.Storage/Controllers/DDCRefController.cs
@@ -50,16 +50,18 @@ namespace Horde.Storage.Controllers
         private readonly IRefsStore _refsStore;
         private readonly IDiagnosticContext _diagnosticContext;
         private readonly IAuthorizationService _authorizationService;
+		private readonly RequestHelper _requestHelper;
 
         private readonly ILogger _logger = Log.ForContext<DDCRefController>();
         private readonly IDDCRefService _ddcRefService;
 
-        public DDCRefController(IDDCRefService ddcRefService, IRefsStore refsStore, IDiagnosticContext diagnosticContext, IAuthorizationService authorizationService)
+        public DDCRefController(IDDCRefService ddcRefService, IRefsStore refsStore, IDiagnosticContext diagnosticContext, IAuthorizationService authorizationService, RequestHelper requestHelper)
         {
             _refsStore = refsStore;
             _ddcRefService = ddcRefService;
             _diagnosticContext = diagnosticContext;
             _authorizationService = authorizationService;
+            _requestHelper = requestHelper;
         }
 
         /// <summary>
@@ -75,13 +77,18 @@ namespace Horde.Storage.Controllers
             NamespaceId[] namespaces = await _refsStore.GetNamespaces().ToArrayAsync();
 
             // filter namespaces down to only the namespaces the user has access to
-            namespaces = namespaces.Where(ns =>
+            List<NamespaceId> validNamespaces = new List<NamespaceId>();
+            foreach (NamespaceId ns in namespaces)
             {
-                Task<AuthorizationResult> authorizationResult = _authorizationService.AuthorizeAsync(User, ns, NamespaceAccessRequirement.Name);
-                return authorizationResult.Result.Succeeded;
-            }).ToArray();
+                ActionResult? result = await _requestHelper.HasAccessToNamespace(User, Request, ns);
+                if (result == null)
+                {
+                    validNamespaces.Add(ns);
+                }
+            }
+
 
-            return Ok(new GetNamespacesResponse(namespaces));
+            return Ok(new GetNamespacesResponse(validNamespaces.ToArray()));
         }
 
         /// <summary>
@@ -106,14 +113,10 @@ namespace Horde.Storage.Controllers
             [FromQuery] string[] fields,
             [FromRoute] string? format = null)
         {
+            ActionResult? result = await _requestHelper.HasAccessToNamespace(User, Request, ns);
+            if (result != null)
             {
-                using IScope _ = Tracer.Instance.StartActive("authorize");
-                AuthorizationResult authorizationResult = await _authorizationService.AuthorizeAsync(User, ns, NamespaceAccessRequirement.Name);
-
-                if (!authorizationResult.Succeeded)
-                {
-                    return Forbid();
-                }
+                return result;
             }
 
             try
@@ -198,13 +201,10 @@ namespace Horde.Storage.Controllers
             [FromRoute] [Required] BucketId bucket,
             [FromRoute] [Required] KeyId key)
         {
+            ActionResult? result = await _requestHelper.HasAccessToNamespace(User, Request, ns);
+            if (result != null)
             {
-                AuthorizationResult authorizationResult = await _authorizationService.AuthorizeAsync(User, ns, NamespaceAccessRequirement.Name);
-
-                if (!authorizationResult.Succeeded)
-                {
-                    return Forbid();
-                }
+                return result;
             }
 
             try
@@ -248,13 +248,10 @@ namespace Horde.Storage.Controllers
             [FromRoute] [Required] KeyId key,
             [FromBody] RefRequest refRequest)
         {
+            ActionResult? result = await _requestHelper.HasAccessToNamespace(User, Request, ns);
+            if (result != null)
             {
-                AuthorizationResult authorizationResult = await _authorizationService.AuthorizeAsync(User, ns, NamespaceAccessRequirement.Name);
-
-                if (!authorizationResult.Succeeded)
-                {
-                    return Forbid();
-                }
+                return result;
             }
 
             try
@@ -357,11 +354,10 @@ namespace Horde.Storage.Controllers
             [FromRoute] [Required] NamespaceId ns
         )
         {
-            AuthorizationResult authorizationResult = await _authorizationService.AuthorizeAsync(User, ns, NamespaceAccessRequirement.Name);
-
-            if (!authorizationResult.Succeeded)
+            ActionResult? result = await _requestHelper.HasAccessToNamespace(User, Request, ns);
+            if (result != null)
             {
-                return Forbid();
+                return result;
             }
 
             try
@@ -388,11 +384,10 @@ namespace Horde.Storage.Controllers
             [FromRoute] [Required] NamespaceId ns,
             [FromRoute] [Required] BucketId bucket)
         {
-            AuthorizationResult authorizationResult = await _authorizationService.AuthorizeAsync(User, ns, NamespaceAccessRequirement.Name);
-
-            if (!authorizationResult.Succeeded)
+            ActionResult? result = await _requestHelper.HasAccessToNamespace(User, Request, ns);
+            if (result != null)
             {
-                return Forbid();
+                return result;
             }
 
             try
@@ -422,11 +417,10 @@ namespace Horde.Storage.Controllers
             [FromRoute] [Required] BucketId bucket,
             [FromRoute] [Required] KeyId key)
         {
-            AuthorizationResult authorizationResult = await _authorizationService.AuthorizeAsync(User, ns, NamespaceAccessRequirement.Name);
-
-            if (!authorizationResult.Succeeded)
+            ActionResult? result = await _requestHelper.HasAccessToNamespace(User, Request, ns);
+            if (result != null)
             {
-                return Forbid();
+                return result;
             }
 
             try
@@ -512,8 +506,8 @@ namespace Horde.Storage.Controllers
 
                 AuthorizationResult authorizationResultNamespace = await _authorizationService.AuthorizeAsync(User, op.Namespace, NamespaceAccessRequirement.Name);
                 AuthorizationResult authorizationResultOp = await _authorizationService.AuthorizeAsync(User, op, OpToPolicy(op.Op));
-
-                bool authorized = authorizationResultNamespace.Succeeded && authorizationResultOp.Succeeded;
+                ActionResult? accessToNamespaceResult = await _requestHelper.HasAccessToNamespace(User, Request, op.Namespace!.Value);
+                bool authorized = authorizationResultNamespace.Succeeded && authorizationResultOp.Succeeded && accessToNamespaceResult == null;
 
                 if (!authorized)
                 {
@@ -624,13 +618,10 @@ namespace Horde.Storage.Controllers
                 throw new InvalidOperationException("A valid namespace must be specified");
             }
 
+            ActionResult? result = await _requestHelper.HasAccessToNamespace(User, Request, batch.Namespace.Value);
+            if (result != null)
             {
-                AuthorizationResult authorizationResult = await _authorizationService.AuthorizeAsync(User, batch.Namespace, NamespaceAccessRequirement.Name);
-
-                if (!authorizationResult.Succeeded)
-                {
-                    return Forbid();
-                }
+                return result;
             }
 
             BatchWriter writer = new BatchWriter();
diff --git a/Engine/Source/Programs/Horde/HordeStorage/Horde.Storage/Controllers/ObjectController.cs b/Engine/Source/Programs/Horde/HordeStorage/Horde.Storage/Controllers/ObjectController.cs
--- a/Engine/Source/Programs/Horde/HordeStorage/Horde.Storage/Controllers/ObjectController.cs
+++ b/Engine/Source/Programs/Horde/HordeStorage/Horde.Storage/Controllers/ObjectController.cs
@@ -30,17 +30,17 @@ namespace Horde.Storage.Controllers
     {
         private readonly IBlobService _storage;
         private readonly IDiagnosticContext _diagnosticContext;
-        private readonly IAuthorizationService _authorizationService;
+        private readonly RequestHelper _requestHelper;
         private readonly IReferenceResolver _referenceResolver;
         private readonly BufferedPayloadFactory _bufferedPayloadFactory;
 
         private readonly ILogger _logger = Log.ForContext<ObjectController>();
 
-        public ObjectController(IBlobService storage, IDiagnosticContext diagnosticContext, IAuthorizationService authorizationService, IReferenceResolver referenceResolver, BufferedPayloadFactory bufferedPayloadFactory)
+        public ObjectController(IBlobService storage, IDiagnosticContext diagnosticContext, RequestHelper requestHelper, IReferenceResolver referenceResolver, BufferedPayloadFactory bufferedPayloadFactory)
         {
             _storage = storage;
             _diagnosticContext = diagnosticContext;
-            _authorizationService = authorizationService;
+            _requestHelper = requestHelper;
             _referenceResolver = referenceResolver;
             _bufferedPayloadFactory = bufferedPayloadFactory;
         }
@@ -52,11 +52,10 @@ namespace Horde.Storage.Controllers
             [Required] NamespaceId ns,
             [Required] BlobIdentifier id)
         {
-            AuthorizationResult authorizationResult = await _authorizationService.AuthorizeAsync(User, ns, NamespaceAccessRequirement.Name);
-
-            if (!authorizationResult.Succeeded)
+            ActionResult? result = await _requestHelper.HasAccessToNamespace(User, Request, ns);
+            if (result != null)
             {
-                return Forbid();
+                return result;
             }
 
             try
@@ -78,11 +77,10 @@ namespace Horde.Storage.Controllers
             [Required] NamespaceId ns,
             [Required] BlobIdentifier id)
         {
-            AuthorizationResult authorizationResult = await _authorizationService.AuthorizeAsync(User, ns, NamespaceAccessRequirement.Name);
-
-            if (!authorizationResult.Succeeded)
+            ActionResult? result = await _requestHelper.HasAccessToNamespace(User, Request, ns);
+            if (result != null)
             {
-                return Forbid();
+                return result;
             }
 
             bool exists = await _storage.Exists(ns, id);
@@ -102,11 +100,10 @@ namespace Horde.Storage.Controllers
             [Required] NamespaceId ns,
             [Required] [FromQuery] List<BlobIdentifier> id)
         {
-            AuthorizationResult authorizationResult = await _authorizationService.AuthorizeAsync(User, ns, NamespaceAccessRequirement.Name);
-
-            if (!authorizationResult.Succeeded)
+            ActionResult? result = await _requestHelper.HasAccessToNamespace(User, Request, ns);
+            if (result != null)
             {
-                return Forbid();
+                return result;
             }
 
             ConcurrentBag<BlobIdentifier> missingBlobs = new ConcurrentBag<BlobIdentifier>();
@@ -130,11 +127,10 @@ namespace Horde.Storage.Controllers
             [Required] NamespaceId ns,
             [FromBody] BlobIdentifier[] bodyIds)
         {
-            AuthorizationResult authorizationResult = await _authorizationService.AuthorizeAsync(User, ns, NamespaceAccessRequirement.Name);
-
-            if (!authorizationResult.Succeeded)
+            ActionResult? result = await _requestHelper.HasAccessToNamespace(User, Request, ns);
+            if (result != null)
             {
-                return Forbid();
+                return result;
             }
 
             ConcurrentBag<BlobIdentifier> missingBlobs = new ConcurrentBag<BlobIdentifier>();
@@ -158,11 +154,10 @@ namespace Horde.Storage.Controllers
             [Required] NamespaceId ns,
             [Required] BlobIdentifier id)
         {
-            AuthorizationResult authorizationResult = await _authorizationService.AuthorizeAsync(User, ns, NamespaceAccessRequirement.Name);
-
-            if (!authorizationResult.Succeeded)
+            ActionResult? result = await _requestHelper.HasAccessToNamespace(User, Request, ns);
+            if (result != null)
             {
-                return Forbid();
+                return result;
             }
 
             _diagnosticContext.Set("Content-Length", Request.ContentLength ?? -1);
@@ -178,11 +173,10 @@ namespace Horde.Storage.Controllers
             [Required] NamespaceId ns,
             [Required] BlobIdentifier id)
         {
-            AuthorizationResult authorizationResult = await _authorizationService.AuthorizeAsync(User, ns, NamespaceAccessRequirement.Name);
-
-            if (!authorizationResult.Succeeded)
+            ActionResult? result = await _requestHelper.HasAccessToNamespace(User, Request, ns);
+            if (result != null)
             {
-                return Forbid();
+                return result;
             }
 
             BlobContents blob;
@@ -232,11 +226,10 @@ namespace Horde.Storage.Controllers
             [Required] NamespaceId ns,
             [Required] BlobIdentifier id)
         {
-            AuthorizationResult authorizationResult = await _authorizationService.AuthorizeAsync(User, ns, NamespaceAccessRequirement.Name);
-
-            if (!authorizationResult.Succeeded)
+            ActionResult? result = await _requestHelper.HasAccessToNamespace(User, Request, ns);
+            if (result != null)
             {
-                return Forbid();
+                return result;
             }
 
             await _storage.DeleteObject(ns, id);
@@ -252,11 +245,10 @@ namespace Horde.Storage.Controllers
         public async Task<IActionResult> DeleteNamespace(
             [Required] NamespaceId ns)
         {
-            AuthorizationResult authorizationResult = await _authorizationService.AuthorizeAsync(User, ns, NamespaceAccessRequirement.Name);
-
-            if (!authorizationResult.Succeeded)
+            ActionResult? result = await _requestHelper.HasAccessToNamespace(User, Request, ns);
+            if (result != null)
             {
-                return Forbid();
+                return result;
             }
 
             await _storage.DeleteNamespace(ns);
diff --git a/Engine/Source/Programs/Horde/HordeStorage/Horde.Storage/Controllers/ReferencesController.cs b/Engine/Source/Programs/Horde/HordeStorage/Horde.Storage/Controllers/ReferencesController.cs
--- a/Engine/Source/Programs/Horde/HordeStorage/Horde.Storage/Controllers/ReferencesController.cs
+++ b/Engine/Source/Programs/Horde/HordeStorage/Horde.Storage/Controllers/ReferencesController.cs
@@ -4,6 +4,7 @@ using System;
 using System.Collections.Concurrent;
 using System.Collections.Generic;
 using System.ComponentModel.DataAnnotations;
+using System.Diagnostics.CodeAnalysis;
 using System.Globalization;
 using System.IO;
 using System.Linq;
@@ -43,24 +44,24 @@ namespace Horde.Storage.Controllers
     public class ReferencesController : ControllerBase
     {
         private readonly IDiagnosticContext _diagnosticContext;
-        private readonly IAuthorizationService _authorizationService;
         private readonly FormatResolver _formatResolver;
         private readonly BufferedPayloadFactory _bufferedPayloadFactory;
         private readonly IReferenceResolver _referenceResolver;
+        private readonly RequestHelper _requestHelper;
 
         private readonly ILogger _logger = Log.ForContext<ReferencesController>();
         private readonly IObjectService _objectService;
         private readonly IBlobService _blobStore;
 
-        public ReferencesController(IObjectService objectService, IBlobService blobStore, IDiagnosticContext diagnosticContext, IAuthorizationService authorizationService, FormatResolver formatResolver, BufferedPayloadFactory bufferedPayloadFactory, IReferenceResolver referenceResolver)
+        public ReferencesController(IObjectService objectService, IBlobService blobStore, IDiagnosticContext diagnosticContext, FormatResolver formatResolver, BufferedPayloadFactory bufferedPayloadFactory, IReferenceResolver referenceResolver, RequestHelper requestHelper)
         {
             _objectService = objectService;
             _blobStore = blobStore;
             _diagnosticContext = diagnosticContext;
-            _authorizationService = authorizationService;
             _formatResolver = formatResolver;
             _bufferedPayloadFactory = bufferedPayloadFactory;
             _referenceResolver = referenceResolver;
+            _requestHelper = requestHelper;
         }
 
         /// <summary>
@@ -76,13 +77,17 @@ namespace Horde.Storage.Controllers
             NamespaceId[] namespaces = await _objectService.GetNamespaces().ToArrayAsync();
 
             // filter namespaces down to only the namespaces the user has access to
-            namespaces = namespaces.Where(ns =>
+            List<NamespaceId> namespacesWithAccess = new();
+            foreach (NamespaceId ns in namespaces)
             {
-                Task<AuthorizationResult> authorizationResult = _authorizationService.AuthorizeAsync(User, ns, NamespaceAccessRequirement.Name);
-                return authorizationResult.Result.Succeeded;
-            }).ToArray();
+                ActionResult? accessResult = await _requestHelper.HasAccessToNamespace(User, Request, ns);
+                if (accessResult == null)
+                {
+                    namespacesWithAccess.Add(ns);
+                }
+            }
 
-            return Ok(new GetNamespacesResponse(namespaces));
+            return Ok(new GetNamespacesResponse(namespacesWithAccess.ToArray()));
         }
 
         /// <summary>
@@ -101,14 +106,10 @@ namespace Horde.Storage.Controllers
             [FromRoute] [Required] IoHashKey key,
             [FromRoute] string? format = null)
         {
+            ActionResult? accessResult = await _requestHelper.HasAccessToNamespace(User, Request, ns);
+            if (accessResult != null)
             {
-                using IScope _ = Tracer.Instance.StartActive("authorize");
-                AuthorizationResult authorizationResult = await _authorizationService.AuthorizeAsync(User, ns, NamespaceAccessRequirement.Name);
-
-                if (!authorizationResult.Succeeded)
-                {
-                    return Forbid();
-                }
+                return accessResult;
             }
 
             try
@@ -323,14 +324,16 @@ namespace Horde.Storage.Controllers
             }
         }
 
-          /// <summary>
-        /// Returns the metadata about a ref key
-        /// </summary>
-        /// <param name="ns">Namespace. Each namespace is completely separated from each other. Use for different types of data that is never expected to be similar (between two different games for instance). Example: `uc4.ddc`</param>
-        /// <param name="bucket">The category/type of record you are caching. Is a clustered key together with the actual key, but all records in the same bucket can be dropped easily. Example: `terrainTexture` </param>
-        /// <param name="key">The unique name of this particular key. `iAmAVeryValidKey`</param>
-        /// <param name="fields">The fields to include in the response, omit this to include everything.</param>
-        [HttpGet("{ns}/{bucket}/{key}/metadata", Order = 500)]
+       
+
+    /// <summary>
+    /// Returns the metadata about a ref key
+    /// </summary>
+    /// <param name="ns">Namespace. Each namespace is completely separated from each other. Use for different types of data that is never expected to be similar (between two different games for instance). Example: `uc4.ddc`</param>
+    /// <param name="bucket">The category/type of record you are caching. Is a clustered key together with the actual key, but all records in the same bucket can be dropped easily. Example: `terrainTexture` </param>
+    /// <param name="key">The unique name of this particular key. `iAmAVeryValidKey`</param>
+    /// <param name="fields">The fields to include in the response, omit this to include everything.</param>
+    [HttpGet("{ns}/{bucket}/{key}/metadata", Order = 500)]
         [Authorize("Object.read")]
         public async Task<IActionResult> GetMetadata(
             [FromRoute] [Required] NamespaceId ns,
@@ -338,14 +341,10 @@ namespace Horde.Storage.Controllers
             [FromRoute] [Required] IoHashKey key,
             [FromQuery] string[] fields)
         {
+            ActionResult? accessResult = await _requestHelper.HasAccessToNamespace(User, Request, ns);
+            if (accessResult != null)
             {
-                using IScope _ = Tracer.Instance.StartActive("authorize");
-                AuthorizationResult authorizationResult = await _authorizationService.AuthorizeAsync(User, ns, NamespaceAccessRequirement.Name);
-
-                if (!authorizationResult.Succeeded)
-                {
-                    return Forbid();
-                }
+                return accessResult;
             }
 
             try
@@ -384,14 +383,10 @@ namespace Horde.Storage.Controllers
             [FromRoute] [Required] BucketId bucket,
             [FromRoute] [Required] IoHashKey key)
         {
+            ActionResult? accessResult = await _requestHelper.HasAccessToNamespace(User, Request, ns);
+            if (accessResult != null)
             {
-                using IScope _ = Tracer.Instance.StartActive("authorize");
-                AuthorizationResult authorizationResult = await _authorizationService.AuthorizeAsync(User, ns, NamespaceAccessRequirement.Name);
-
-                if (!authorizationResult.Succeeded)
-                {
-                    return Forbid();
-                }
+                return accessResult;
             }
 
             try
@@ -453,11 +448,10 @@ namespace Horde.Storage.Controllers
             [FromRoute] [Required] NamespaceId ns,
             [FromQuery] [Required] List<string> names)
         {
-            AuthorizationResult authorizationResult = await _authorizationService.AuthorizeAsync(User, ns, NamespaceAccessRequirement.Name);
-
-            if (!authorizationResult.Succeeded)
+            ActionResult? accessResult = await _requestHelper.HasAccessToNamespace(User, Request, ns);
+            if (accessResult != null)
             {
-                return Forbid();
+                return accessResult;
             }
 
             ConcurrentBag<(BucketId, IoHashKey)> missingObject = new ();
@@ -520,14 +514,10 @@ namespace Horde.Storage.Controllers
             [FromRoute] [Required] BucketId bucket,
             [FromRoute] [Required] IoHashKey key)
         {
+            ActionResult? accessResult = await _requestHelper.HasAccessToNamespace(User, Request, ns);
+            if (accessResult != null)
             {
-                using IScope _ = Tracer.Instance.StartActive("authorize");
-                AuthorizationResult authorizationResult = await _authorizationService.AuthorizeAsync(User, ns, NamespaceAccessRequirement.Name);
-
-                if (!authorizationResult.Succeeded)
-                {
-                    return Forbid();
-                }
+                return accessResult;
             }
 
             _diagnosticContext.Set("Content-Length", Request.ContentLength ?? -1);
@@ -620,14 +610,10 @@ namespace Horde.Storage.Controllers
             [FromRoute] [Required] IoHashKey key,
             [FromRoute] [Required] BlobIdentifier hash)
         {
+            ActionResult? accessResult = await _requestHelper.HasAccessToNamespace(User, Request, ns);
+            if (accessResult != null)
             {
-                using IScope _ = Tracer.Instance.StartActive("authorize");
-                AuthorizationResult authorizationResult = await _authorizationService.AuthorizeAsync(User, ns, NamespaceAccessRequirement.Name);
-
-                if (!authorizationResult.Succeeded)
-                {
-                    return Forbid();
-                }
+                return accessResult;
             }
 
             (ContentId[] missingReferences, BlobIdentifier[] missingBlobs) = await _objectService.Finalize(ns, bucket, key, hash);
@@ -644,15 +630,10 @@ namespace Horde.Storage.Controllers
             [FromRoute] [Required] NamespaceId ns,
             [FromBody] [Required] BatchOps ops)
         {
+            ActionResult? accessResult = await _requestHelper.HasAccessToNamespace(User, Request, ns);
+            if (accessResult != null)
             {
-                using IScope _ = Tracer.Instance.StartActive("authorize");
-                AuthorizationResult authorizationResult =
-                    await _authorizationService.AuthorizeAsync(User, ns, NamespaceAccessRequirement.Name);
-
-                if (!authorizationResult.Succeeded)
-                {
-                    return Forbid();
-                }
+                return accessResult;
             }
 
             ConcurrentDictionary<uint, (CbObject, HttpStatusCode)> results = new();
@@ -821,11 +802,10 @@ namespace Horde.Storage.Controllers
             [FromRoute] [Required] NamespaceId ns
         )
         {
-            AuthorizationResult authorizationResult = await _authorizationService.AuthorizeAsync(User, ns, NamespaceAccessRequirement.Name);
-
-            if (!authorizationResult.Succeeded)
+            ActionResult? accessResult = await _requestHelper.HasAccessToNamespace(User, Request, ns);
+            if (accessResult != null)
             {
-                return Forbid();
+                return accessResult;
             }
 
             try
@@ -852,11 +832,10 @@ namespace Horde.Storage.Controllers
             [FromRoute] [Required] NamespaceId ns,
             [FromRoute] [Required] BucketId bucket)
         {
-            AuthorizationResult authorizationResult = await _authorizationService.AuthorizeAsync(User, ns, NamespaceAccessRequirement.Name);
-
-            if (!authorizationResult.Succeeded)
+            ActionResult? accessResult = await _requestHelper.HasAccessToNamespace(User, Request, ns);
+            if (accessResult != null)
             {
-                return Forbid();
+                return accessResult;
             }
 
             long countOfDeletedRecords;
@@ -887,11 +866,10 @@ namespace Horde.Storage.Controllers
             [FromRoute] [Required] BucketId bucket,
             [FromRoute] [Required] IoHashKey key)
         {
-            AuthorizationResult authorizationResult = await _authorizationService.AuthorizeAsync(User, ns, NamespaceAccessRequirement.Name);
-
-            if (!authorizationResult.Succeeded)
+            ActionResult? accessResult = await _requestHelper.HasAccessToNamespace(User, Request, ns);
+            if (accessResult != null)
             {
-                return Forbid();
+                return accessResult;
             }
 
             try
diff --git a/Engine/Source/Programs/Horde/HordeStorage/Horde.Storage/Controllers/ReplicationController.cs b/Engine/Source/Programs/Horde/HordeStorage/Horde.Storage/Controllers/ReplicationController.cs
--- a/Engine/Source/Programs/Horde/HordeStorage/Horde.Storage/Controllers/ReplicationController.cs
+++ b/Engine/Source/Programs/Horde/HordeStorage/Horde.Storage/Controllers/ReplicationController.cs
@@ -22,12 +22,12 @@ namespace Horde.Storage.Controllers
     public class ReplicationController : ControllerBase
     {
         private readonly ReplicationService _replicationService;
-        private readonly IAuthorizationService _authorizationService;
+        private readonly RequestHelper _requestHelper;
 
-        public ReplicationController(ReplicationService replicationService, IAuthorizationService authorizationService)
+        public ReplicationController(ReplicationService replicationService, RequestHelper requestHelper)
         {
             _replicationService = replicationService;
-            _authorizationService = authorizationService;
+            _requestHelper = requestHelper;
         }
 
         [HttpGet("{ns}")]
@@ -36,11 +36,10 @@ namespace Horde.Storage.Controllers
             [Required] NamespaceId ns
         )
         {
-            AuthorizationResult authorizationResult = await _authorizationService.AuthorizeAsync(User, ns, NamespaceAccessRequirement.Name);
-
-            if (!authorizationResult.Succeeded)
+            ActionResult? result = await _requestHelper.HasAccessToNamespace(User, Request, ns);
+            if (result != null)
             {
-                return Forbid();
+                return result;
             }
 
             IEnumerable<IReplicator>? replicators = _replicationService.GetReplicators(ns);
@@ -80,11 +79,10 @@ namespace Horde.Storage.Controllers
             [Required] long offset
         )
         {
-            AuthorizationResult authorizationResult = await _authorizationService.AuthorizeAsync(User, ns, NamespaceAccessRequirement.Name);
-
-            if (!authorizationResult.Succeeded)
+            ActionResult? result = await _requestHelper.HasAccessToNamespace(User, Request, ns);
+            if (result != null)
             {
-                return Forbid();
+                return result;
             }
 
             IEnumerable<IReplicator>? replicators = _replicationService.GetReplicators(ns);
@@ -119,11 +117,10 @@ namespace Horde.Storage.Controllers
             [Required] [FromBody] NewReplicationState replicationState
         )
         {
-            AuthorizationResult authorizationResult = await _authorizationService.AuthorizeAsync(User, ns, NamespaceAccessRequirement.Name);
-
-            if (!authorizationResult.Succeeded)
+            ActionResult? result = await _requestHelper.HasAccessToNamespace(User, Request, ns);
+            if (result != null)
             {
-                return Forbid();
+                return result;
             }
 
             IEnumerable<IReplicator>? replicators = _replicationService.GetReplicators(ns);
@@ -163,11 +160,10 @@ namespace Horde.Storage.Controllers
             [Required] NamespaceId ns
         )
         {
-            AuthorizationResult authorizationResult = await _authorizationService.AuthorizeAsync(User, ns, NamespaceAccessRequirement.Name);
-
-            if (!authorizationResult.Succeeded)
+            ActionResult? result = await _requestHelper.HasAccessToNamespace(User, Request, ns);
+            if (result != null)
             {
-                return Forbid();
+                return result;
             }
 
             List<IReplicator> replicators = _replicationService.GetReplicators(ns).ToList();
diff --git a/Engine/Source/Programs/Horde/HordeStorage/Horde.Storage/Controllers/ReplicationLogController.cs b/Engine/Source/Programs/Horde/HordeStorage/Horde.Storage/Controllers/ReplicationLogController.cs
--- a/Engine/Source/Programs/Horde/HordeStorage/Horde.Storage/Controllers/ReplicationLogController.cs
+++ b/Engine/Source/Programs/Horde/HordeStorage/Horde.Storage/Controllers/ReplicationLogController.cs
@@ -26,14 +26,14 @@ namespace Horde.Storage.Controllers
     public class ReplicationLogController : ControllerBase
     {
         private readonly IServiceProvider _provider;
-        private readonly IAuthorizationService _authorizationService;
+        private readonly RequestHelper _requestHelper;
         private readonly IReplicationLog _replicationLog;
         private readonly IOptionsMonitor<SnapshotSettings> _snapshotSettings;
 
-        public ReplicationLogController(IServiceProvider provider, IAuthorizationService authorizationService, IReplicationLog replicationLog, IOptionsMonitor<SnapshotSettings> snapshotSettings)
+        public ReplicationLogController(IServiceProvider provider, RequestHelper requestHelper, IReplicationLog replicationLog, IOptionsMonitor<SnapshotSettings> snapshotSettings)
         {
             _provider = provider;
-            _authorizationService = authorizationService;
+            _requestHelper = requestHelper;
             _replicationLog = replicationLog;
             _snapshotSettings = snapshotSettings;
         }
@@ -46,11 +46,10 @@ namespace Horde.Storage.Controllers
             [Required] NamespaceId ns
         )
         {
-            AuthorizationResult authorizationResult = await _authorizationService.AuthorizeAsync(User, ns, NamespaceAccessRequirement.Name);
-
-            if (!authorizationResult.Succeeded)
+            ActionResult? result = await _requestHelper.HasAccessToNamespace(User, Request, ns);
+            if (result != null)
             {
-                return Forbid();
+                return result;
             }
 
             return Ok(new ReplicationLogSnapshots(await _replicationLog.GetSnapshots(ns).ToListAsync()));
@@ -65,11 +64,10 @@ namespace Horde.Storage.Controllers
             [Required] NamespaceId ns
         )
         {
-            AuthorizationResult authorizationResult = await _authorizationService.AuthorizeAsync(User, ns, NamespaceAccessRequirement.Name);
-
-            if (!authorizationResult.Succeeded)
+            ActionResult? result = await _requestHelper.HasAccessToNamespace(User, Request, ns);
+            if (result != null)
             {
-                return Forbid();
+                return result;
             }
 
             ReplicationLogSnapshotBuilder builder = ActivatorUtilities.CreateInstance<ReplicationLogSnapshotBuilder>(_provider);
@@ -88,11 +86,10 @@ namespace Horde.Storage.Controllers
             [FromQuery] int count = 100
         )
         {
-            AuthorizationResult authorizationResult = await _authorizationService.AuthorizeAsync(User, ns, NamespaceAccessRequirement.Name);
-
-            if (!authorizationResult.Succeeded)
+            ActionResult? result = await _requestHelper.HasAccessToNamespace(User, Request, ns);
+            if (result != null)
             {
-                return Forbid();
+                return result;
             }
 
             if (((lastBucket == null && lastEvent.HasValue) || (lastBucket != null && !lastEvent.HasValue)) && lastBucket != "now")
diff --git a/Engine/Source/Programs/Horde/HordeStorage/Horde.Storage/Controllers/RequestHelper.cs b/Engine/Source/Programs/Horde/HordeStorage/Horde.Storage/Controllers/RequestHelper.cs
new file mode 100644
--- /dev/null
+++ b/Engine/Source/Programs/Horde/HordeStorage/Horde.Storage/Controllers/RequestHelper.cs
@@ -0,0 +1,60 @@
+// Copyright Epic Games, Inc. All Rights Reserved.
+
+using System.Security.Claims;
+using System.Threading.Tasks;
+using Datadog.Trace;
+using EpicGames.Horde.Storage;
+using Jupiter;
+using Jupiter.Common;
+using Microsoft.AspNetCore.Authorization;
+using Microsoft.AspNetCore.Http;
+using Microsoft.AspNetCore.Mvc;
+using Microsoft.Extensions.Options;
+
+namespace Horde.Storage.Controllers;
+
+public class RequestHelper
+{
+    private readonly IAuthorizationService _authorizationService;
+    private readonly INamespacePolicyResolver _namespacePolicyResolver;
+    private readonly IOptionsMonitor<JupiterSettings> _settings;
+
+    public RequestHelper(IAuthorizationService authorizationService, INamespacePolicyResolver namespacePolicyResolver, IOptionsMonitor<JupiterSettings> settings)
+    {
+        _authorizationService = authorizationService;
+        _namespacePolicyResolver = namespacePolicyResolver;
+        _settings = settings;
+    }
+
+
+    public async Task<ActionResult?> HasAccessToNamespace(ClaimsPrincipal user, HttpRequest request, NamespaceId ns)
+    {
+        using IScope _ = Tracer.Instance.StartActive("authorize");
+        AuthorizationResult authorizationResult = await _authorizationService.AuthorizeAsync(user, ns, NamespaceAccessRequirement.Name);
+
+        if (!authorizationResult.Succeeded)
+        {
+            return new ForbidResult();
+        }
+
+        bool isPublicNamespace = _namespacePolicyResolver.GetPoliciesForNs(ns).IsPublicNamespace;
+
+        // public namespaces are always accessible
+        if (isPublicNamespace)
+            return null;
+
+        // namespace is a restricted namespace
+        HttpContext context = request.HttpContext;
+        bool isPublicPort = context.Connection.LocalPort == _settings!.CurrentValue.PublicApiPort ||
+            /* unit tests do not run on ports, we consider them always on the internal port */ !(context.Connection.LocalPort == 0 && context.Connection.LocalIpAddress == null);
+            
+        if (isPublicPort)
+        {
+            // trying to access restricted namespace on a public port, this is not allowed
+            return new ForbidResult();
+        }
+
+        // restricted namespace in corp or internal port, this is okay
+        return null;
+    }
+}
diff --git a/Engine/Source/Programs/Horde/HordeStorage/Horde.Storage/Controllers/StorageController.cs b/Engine/Source/Programs/Horde/HordeStorage/Horde.Storage/Controllers/StorageController.cs
--- a/Engine/Source/Programs/Horde/HordeStorage/Horde.Storage/Controllers/StorageController.cs
+++ b/Engine/Source/Programs/Horde/HordeStorage/Horde.Storage/Controllers/StorageController.cs
@@ -31,13 +31,15 @@ namespace Horde.Storage.Controllers
         private readonly IBlobService _storage;
         private readonly IDiagnosticContext _diagnosticContext;
         private readonly IAuthorizationService _authorizationService;
+        private readonly RequestHelper _requestHelper;
         private readonly BufferedPayloadFactory _bufferedPayloadFactory;
 
-        public StorageController(IBlobService storage, IDiagnosticContext diagnosticContext, IAuthorizationService authorizationService, BufferedPayloadFactory bufferedPayloadFactory)
+        public StorageController(IBlobService storage, IDiagnosticContext diagnosticContext, IAuthorizationService authorizationService, RequestHelper requestHelper, BufferedPayloadFactory bufferedPayloadFactory)
         {
             _storage = storage;
             _diagnosticContext = diagnosticContext;
             _authorizationService = authorizationService;
+            _requestHelper = requestHelper;
             _bufferedPayloadFactory = bufferedPayloadFactory;
         }
 
@@ -48,11 +50,10 @@ namespace Horde.Storage.Controllers
             [Required] NamespaceId ns,
             [Required] BlobIdentifier id)
         {
-            AuthorizationResult authorizationResult = await _authorizationService.AuthorizeAsync(User, ns, NamespaceAccessRequirement.Name);
-
-            if (!authorizationResult.Succeeded)
+            ActionResult? result = await _requestHelper.HasAccessToNamespace(User, Request, ns);
+            if (result != null)
             {
-                return Forbid();
+                return result;
             }
 
             try
@@ -74,11 +75,10 @@ namespace Horde.Storage.Controllers
             [Required] NamespaceId ns,
             [Required] BlobIdentifier id)
         {
-            AuthorizationResult authorizationResult = await _authorizationService.AuthorizeAsync(User, ns, NamespaceAccessRequirement.Name);
-
-            if (!authorizationResult.Succeeded)
+            ActionResult? result = await _requestHelper.HasAccessToNamespace(User, Request, ns);
+            if (result != null)
             {
-                return Forbid();
+                return result;
             }
             bool exists = await _storage.Exists(ns, id);
 
@@ -97,11 +97,10 @@ namespace Horde.Storage.Controllers
             [Required] NamespaceId ns,
              [Required] [FromQuery] List<BlobIdentifier> id)
         {
-            AuthorizationResult authorizationResult = await _authorizationService.AuthorizeAsync(User, ns, NamespaceAccessRequirement.Name);
-
-            if (!authorizationResult.Succeeded)
+            ActionResult? result = await _requestHelper.HasAccessToNamespace(User, Request, ns);
+            if (result != null)
             {
-                return Forbid();
+                return result;
             }
 
             ConcurrentBag<BlobIdentifier> missingBlobs = new ConcurrentBag<BlobIdentifier>();
@@ -125,11 +124,10 @@ namespace Horde.Storage.Controllers
             [Required] NamespaceId ns,
             [FromBody] BlobIdentifier[] bodyIds)
         {
-            AuthorizationResult authorizationResult = await _authorizationService.AuthorizeAsync(User, ns, NamespaceAccessRequirement.Name);
-
-            if (!authorizationResult.Succeeded)
+            ActionResult? result = await _requestHelper.HasAccessToNamespace(User, Request, ns);
+            if (result != null)
             {
-                return Forbid();
+                return result;
             }
 
             ConcurrentBag<BlobIdentifier> missingBlobs = new ConcurrentBag<BlobIdentifier>();
@@ -171,11 +169,10 @@ namespace Horde.Storage.Controllers
             [Required] NamespaceId ns,
             [Required] BlobIdentifier id)
         {
-            AuthorizationResult authorizationResult = await _authorizationService.AuthorizeAsync(User, ns, NamespaceAccessRequirement.Name);
-
-            if (!authorizationResult.Succeeded)
+            ActionResult? result = await _requestHelper.HasAccessToNamespace(User, Request, ns);
+            if (result != null)
             {
-                return Forbid();
+                return result;
             }
 
             _diagnosticContext.Set("Content-Length", Request.ContentLength ?? -1);
@@ -202,11 +199,10 @@ namespace Horde.Storage.Controllers
         public async Task<IActionResult> Post(
             [Required] NamespaceId ns)
         {
-            AuthorizationResult authorizationResult = await _authorizationService.AuthorizeAsync(User, ns, NamespaceAccessRequirement.Name);
-
-            if (!authorizationResult.Succeeded)
+            ActionResult? result = await _requestHelper.HasAccessToNamespace(User, Request, ns);
+            if (result != null)
             {
-                return Forbid();
+                return result;
             }
 
             _diagnosticContext.Set("Content-Length", Request.ContentLength ?? -1);
@@ -229,11 +225,10 @@ namespace Horde.Storage.Controllers
             [Required] NamespaceId ns,
             [Required] BlobIdentifier id)
         {
-            AuthorizationResult authorizationResult = await _authorizationService.AuthorizeAsync(User, ns, NamespaceAccessRequirement.Name);
-
-            if (!authorizationResult.Succeeded)
+            ActionResult? result = await _requestHelper.HasAccessToNamespace(User, Request, ns);
+            if (result != null)
             {
-                return Forbid();
+                return result;
             }
 
             await DeleteImpl(ns, id);
@@ -247,11 +242,10 @@ namespace Horde.Storage.Controllers
         public async Task<IActionResult> DeleteNamespace(
             [Required] NamespaceId ns)
         {
-            AuthorizationResult authorizationResult = await _authorizationService.AuthorizeAsync(User, ns, NamespaceAccessRequirement.Name);
-
-            if (!authorizationResult.Succeeded)
+            ActionResult? result = await _requestHelper.HasAccessToNamespace(User, Request, ns);
+            if (result != null)
             {
-                return Forbid();
+                return result;
             }
 
             await  _storage.DeleteNamespace(ns);
diff --git a/Engine/Source/Programs/Horde/HordeStorage/Horde.Storage/Horde.Storage.Startup.cs b/Engine/Source/Programs/Horde/HordeStorage/Horde.Storage/Horde.Storage.Startup.cs
--- a/Engine/Source/Programs/Horde/HordeStorage/Horde.Storage/Horde.Storage.Startup.cs
+++ b/Engine/Source/Programs/Horde/HordeStorage/Horde.Storage/Horde.Storage.Startup.cs
@@ -14,6 +14,7 @@ using Amazon.Runtime;
 using Amazon.S3;
 using Amazon.SecretsManager;
 using Cassandra;
+using EpicGames.AspNet;
 using Horde.Storage.Controllers;
 using Horde.Storage.Implementation;
 using Horde.Storage.Implementation.Blob;
@@ -177,6 +178,8 @@ namespace Horde.Storage
 
             services.AddHttpClient();
 
+            services.AddServerTiming();
+            
             services.AddOptions<HordeStorageSettings>().Bind(Configuration.GetSection("Horde_Storage")).ValidateDataAnnotations();
             services.AddOptions<MongoSettings>().Bind(Configuration.GetSection("Mongo")).ValidateDataAnnotations();
             services.AddOptions<CosmosSettings>().Bind(Configuration.GetSection("Cosmos")).ValidateDataAnnotations();
@@ -260,6 +263,8 @@ namespace Horde.Storage
 
             services.AddSingleton(serviceType: typeof(ILeaderElection), CreateLeaderElection);
 
+            services.AddTransient<RequestHelper>();
+
             services.AddSingleton(Configuration);
             services.AddSingleton<DynamoNamespaceStore>();
             services.AddSingleton(RefStoreFactory);
diff --git a/Engine/Source/Programs/Horde/HordeStorage/Horde.Storage/appsettings.Production.json b/Engine/Source/Programs/Horde/HordeStorage/Horde.Storage/appsettings.Production.json
--- a/Engine/Source/Programs/Horde/HordeStorage/Horde.Storage/appsettings.Production.json
+++ b/Engine/Source/Programs/Horde/HordeStorage/Horde.Storage/appsettings.Production.json
@@ -32,7 +32,8 @@
   },
 
   "Jupiter": {
-    "InternalApiPort": 8080
+    "InternalApiPort": 8080,
+    "CorpApiPort": 8008
   },
 
   "AllowedHosts": "*",
@@ -44,6 +45,9 @@
       },
       "InternalHttp": {
         "Url": "http://0.0.0.0:8080"
+      },
+      "CorpHttp": {
+        "Url": "http://0.0.0.0:8008"
       }
     }
   }
diff --git a/Engine/Source/Programs/Horde/HordeStorage/Horde.Storage/appsettings.Testing.json b/Engine/Source/Programs/Horde/HordeStorage/Horde.Storage/appsettings.Testing.json
--- a/Engine/Source/Programs/Horde/HordeStorage/Horde.Storage/appsettings.Testing.json
+++ b/Engine/Source/Programs/Horde/HordeStorage/Horde.Storage/appsettings.Testing.json
@@ -86,12 +86,18 @@
       "test-namespace": {
         "Claims": [ "*" ]
       },
+      "delete-namespace": {
+        "Claims": [ "*" ]
+      },
       "test-delete-namespace": {
         "Claims": [ "*" ]
       },
       "testbucket": {
         "Claims": [ "*" ]
       },
+      "test-namespace-last-access": {
+        "Claims": [ "*" ]
+      },
       "snapshot-namespace": {
         "Claims": [ "*" ]
       }
diff --git a/Engine/Source/Programs/Horde/HordeStorage/Jupiter.Common/BaseStartup.cs b/Engine/Source/Programs/Horde/HordeStorage/Jupiter.Common/BaseStartup.cs
--- a/Engine/Source/Programs/Horde/HordeStorage/Jupiter.Common/BaseStartup.cs
+++ b/Engine/Source/Programs/Horde/HordeStorage/Jupiter.Common/BaseStartup.cs
@@ -432,8 +432,6 @@ namespace Jupiter
 
     public class JupiterSettings
     {
-        public const int DefaultInternalPort = 8080;
-
         /// <summary>
         /// If the request is smaller then MemoryBufferSize we buffer it in memory rather then as a file
         /// </summary>
@@ -444,7 +442,21 @@ namespace Jupiter
         public bool DisableHealthChecks { get; set; } = false;
         public bool HostSwaggerDocumentation { get; set; } = true;
 
-        public int InternalApiPort { get; set; } = DefaultInternalPort;
+        /// <summary>
+        /// Port used to host the internally accessible api (as well as the public api).
+        /// This hosts both public and private namespaces
+        /// </summary>
+        public int InternalApiPort { get; set; } = 8080;
+
+        /// <summary>
+        /// Port that hosts public and private namespaces
+        /// </summary>
+        public int CorpApiPort { get; set; } = 8008;
+
+        /// <summary>
+        /// Port that only hosts the public namespaces
+        /// </summary>
+        public int PublicApiPort { get; set; } = 80;
 
         // Enable to echo every request to the log file, usually this is more efficiently done on the load balancer
         public bool UseRequestLogging { get; set; } = false;
@@ -473,6 +485,7 @@ namespace Jupiter
 	    public bool UseBlobIndexForExists { get; set; } = false;
 	    public bool UseBlobIndexForSlowExists { get; set; } = false;
 	    public bool? IsLegacyNamespace { get; set; } = null;
+        public bool IsPublicNamespace { get; set; } = true;
     }
 
     public class DatadogTraceMiddleware
