Horde Storage - We now return a Server-Timings header with some timers from horde storage. If your connection supports trailing headers and you specify a header "TE: trailers" we will use trailing headers for these allowing us to send more timings (related to the writing of the body)

diff --git a/Engine/Source/Programs/Horde/HordeStorage/Horde.Storage/Controllers/AuthController.cs b/Engine/Source/Programs/Horde/HordeStorage/Horde.Storage/Controllers/AuthController.cs
--- a/Engine/Source/Programs/Horde/HordeStorage/Horde.Storage/Controllers/AuthController.cs
+++ b/Engine/Source/Programs/Horde/HordeStorage/Horde.Storage/Controllers/AuthController.cs
@@ -3,7 +3,6 @@
 using Microsoft.AspNetCore.Mvc;
 using System.ComponentModel.DataAnnotations;
 using System.Threading.Tasks;
-using Jupiter;
 using Microsoft.AspNetCore.Authorization;
 using EpicGames.Horde.Storage;
 
diff --git a/Engine/Source/Programs/Horde/HordeStorage/Horde.Storage/Controllers/ContentIdController.cs b/Engine/Source/Programs/Horde/HordeStorage/Horde.Storage/Controllers/ContentIdController.cs
--- a/Engine/Source/Programs/Horde/HordeStorage/Horde.Storage/Controllers/ContentIdController.cs
+++ b/Engine/Source/Programs/Horde/HordeStorage/Horde.Storage/Controllers/ContentIdController.cs
@@ -2,7 +2,6 @@
 
 using System.Net.Mime;
 using System.Threading.Tasks;
-using Datadog.Trace;
 using EpicGames.Horde.Storage;
 using Horde.Storage.Implementation;
 using Jupiter;
diff --git a/Engine/Source/Programs/Horde/HordeStorage/Horde.Storage/Controllers/ReferencesController.cs b/Engine/Source/Programs/Horde/HordeStorage/Horde.Storage/Controllers/ReferencesController.cs
--- a/Engine/Source/Programs/Horde/HordeStorage/Horde.Storage/Controllers/ReferencesController.cs
+++ b/Engine/Source/Programs/Horde/HordeStorage/Horde.Storage/Controllers/ReferencesController.cs
@@ -4,7 +4,6 @@ using System;
 using System.Collections.Concurrent;
 using System.Collections.Generic;
 using System.ComponentModel.DataAnnotations;
-using System.Diagnostics.CodeAnalysis;
 using System.Globalization;
 using System.IO;
 using System.Linq;
@@ -15,12 +14,12 @@ using System.Threading;
 using System.Threading.Tasks;
 using Dasync.Collections;
 using Datadog.Trace;
+using EpicGames.AspNet;
 using EpicGames.Core;
 using EpicGames.Horde.Storage;
 using EpicGames.Serialization;
 using Horde.Storage.Implementation;
 using JetBrains.Annotations;
-using Jupiter;
 using Jupiter.Common.Implementation;
 using Jupiter.Implementation;
 using Jupiter.Utils;
@@ -28,9 +27,11 @@ using Microsoft.AspNetCore.Authorization;
 using Microsoft.AspNetCore.Http;
 using Microsoft.AspNetCore.Http.Extensions;
 using Microsoft.AspNetCore.Mvc;
+using Microsoft.Extensions.DependencyInjection;
 using Newtonsoft.Json;
 using Serilog;
 using ContentHash = Jupiter.Implementation.ContentHash;
+using CustomMediaTypeNames = Jupiter.CustomMediaTypeNames;
 using Log = Serilog.Log;
 
 namespace Horde.Storage.Controllers
@@ -132,6 +133,9 @@ namespace Horde.Storage.Controllers
 
                 async Task WriteBody(BlobContents blobContents, string contentType)
                 {
+                    IServerTiming? serverTiming = Request.HttpContext.RequestServices.GetService<IServerTiming>();
+                    using ServerTimingMetricScoped? serverTimingScope = serverTiming?.CreateServerTimingMetricScope("body.write", "Time spent writing body");
+
                     using IScope scope = Tracer.Instance.StartActive("body.write");
                     long contentLength = blobContents.Length;
                     scope.Span.SetTag("content-length", contentLength.ToString());
diff --git a/Engine/Source/Programs/Horde/HordeStorage/Horde.Storage/Controllers/ReplicationController.cs b/Engine/Source/Programs/Horde/HordeStorage/Horde.Storage/Controllers/ReplicationController.cs
--- a/Engine/Source/Programs/Horde/HordeStorage/Horde.Storage/Controllers/ReplicationController.cs
+++ b/Engine/Source/Programs/Horde/HordeStorage/Horde.Storage/Controllers/ReplicationController.cs
@@ -7,7 +7,6 @@ using System.Linq;
 using System.Threading.Tasks;
 using EpicGames.Horde.Storage;
 using Horde.Storage.Implementation;
-using Jupiter;
 using Jupiter.Implementation;
 using Microsoft.AspNetCore.Authorization;
 using Microsoft.AspNetCore.Mvc;
diff --git a/Engine/Source/Programs/Horde/HordeStorage/Horde.Storage/Controllers/ReplicationLogController.cs b/Engine/Source/Programs/Horde/HordeStorage/Horde.Storage/Controllers/ReplicationLogController.cs
--- a/Engine/Source/Programs/Horde/HordeStorage/Horde.Storage/Controllers/ReplicationLogController.cs
+++ b/Engine/Source/Programs/Horde/HordeStorage/Horde.Storage/Controllers/ReplicationLogController.cs
@@ -10,7 +10,6 @@ using EpicGames.Horde.Storage;
 using EpicGames.Serialization;
 using Horde.Storage.Implementation;
 using Horde.Storage.Implementation.TransactionLog;
-using Jupiter;
 using Jupiter.Implementation;
 using Microsoft.AspNetCore.Authorization;
 using Microsoft.AspNetCore.Mvc;
diff --git a/Engine/Source/Programs/Horde/HordeStorage/Horde.Storage/Controllers/RequestHelper.cs b/Engine/Source/Programs/Horde/HordeStorage/Horde.Storage/Controllers/RequestHelper.cs
--- a/Engine/Source/Programs/Horde/HordeStorage/Horde.Storage/Controllers/RequestHelper.cs
+++ b/Engine/Source/Programs/Horde/HordeStorage/Horde.Storage/Controllers/RequestHelper.cs
@@ -26,7 +26,6 @@ public class RequestHelper
         _settings = settings;
     }
 
-
     public async Task<ActionResult?> HasAccessToNamespace(ClaimsPrincipal user, HttpRequest request, NamespaceId ns)
     {
         using IScope _ = Tracer.Instance.StartActive("authorize");
@@ -41,7 +40,9 @@ public class RequestHelper
 
         // public namespaces are always accessible
         if (isPublicNamespace)
+        {
             return null;
+        }
 
         // namespace is a restricted namespace
         HttpContext context = request.HttpContext;
diff --git a/Engine/Source/Programs/Horde/HordeStorage/Horde.Storage/Horde.Storage.Startup.cs b/Engine/Source/Programs/Horde/HordeStorage/Horde.Storage/Horde.Storage.Startup.cs
--- a/Engine/Source/Programs/Horde/HordeStorage/Horde.Storage/Horde.Storage.Startup.cs
+++ b/Engine/Source/Programs/Horde/HordeStorage/Horde.Storage/Horde.Storage.Startup.cs
@@ -14,7 +14,6 @@ using Amazon.Runtime;
 using Amazon.S3;
 using Amazon.SecretsManager;
 using Cassandra;
-using EpicGames.AspNet;
 using Horde.Storage.Controllers;
 using Horde.Storage.Implementation;
 using Horde.Storage.Implementation.Blob;
@@ -178,8 +177,6 @@ namespace Horde.Storage
 
             services.AddHttpClient();
 
-            services.AddServerTiming();
-            
             services.AddOptions<HordeStorageSettings>().Bind(Configuration.GetSection("Horde_Storage")).ValidateDataAnnotations();
             services.AddOptions<MongoSettings>().Bind(Configuration.GetSection("Mongo")).ValidateDataAnnotations();
             services.AddOptions<CosmosSettings>().Bind(Configuration.GetSection("Cosmos")).ValidateDataAnnotations();
diff --git a/Engine/Source/Programs/Horde/HordeStorage/Horde.Storage/Implementation/Blob/IBlobService.cs b/Engine/Source/Programs/Horde/HordeStorage/Horde.Storage/Implementation/Blob/IBlobService.cs
--- a/Engine/Source/Programs/Horde/HordeStorage/Horde.Storage/Implementation/Blob/IBlobService.cs
+++ b/Engine/Source/Programs/Horde/HordeStorage/Horde.Storage/Implementation/Blob/IBlobService.cs
@@ -10,12 +10,14 @@ using System.Net.Http;
 using System.Threading.Tasks;
 using Dasync.Collections;
 using Datadog.Trace;
+using EpicGames.AspNet;
 using EpicGames.Horde.Storage;
 using Horde.Storage.Controllers;
 using Horde.Storage.Implementation.Blob;
 using Jupiter.Common;
 using Jupiter.Common.Implementation;
 using Jupiter.Implementation;
+using Microsoft.AspNetCore.Http;
 using Microsoft.Extensions.DependencyInjection;
 using Microsoft.Extensions.Options;
 using Serilog;
@@ -66,6 +68,7 @@ public class BlobService : IBlobService
     private readonly IHttpClientFactory _httpClientFactory;
     private readonly IServiceCredentials _serviceCredentials;
     private readonly INamespacePolicyResolver _namespacePolicyResolver;
+    private readonly IHttpContextAccessor _httpContextAccessor;
     private readonly ILogger _logger = Log.ForContext<BlobService>();
 
     internal IEnumerable<IBlobStore> BlobStore
@@ -74,7 +77,7 @@ public class BlobService : IBlobService
         set => _blobStores = value.ToList();
     }
 
-    public BlobService(IServiceProvider provider, IOptionsMonitor<HordeStorageSettings> settings, IBlobIndex blobIndex, IPeerStatusService peerStatusService, IHttpClientFactory httpClientFactory, IServiceCredentials serviceCredentials, INamespacePolicyResolver namespacePolicyResolver)
+    public BlobService(IServiceProvider provider, IOptionsMonitor<HordeStorageSettings> settings, IBlobIndex blobIndex, IPeerStatusService peerStatusService, IHttpClientFactory httpClientFactory, IServiceCredentials serviceCredentials, INamespacePolicyResolver namespacePolicyResolver, IHttpContextAccessor httpContextAccessor)
     {
         _blobStores = GetBlobStores(provider, settings).ToList();
         _settings = settings;
@@ -83,6 +86,7 @@ public class BlobService : IBlobService
         _httpClientFactory = httpClientFactory;
         _serviceCredentials = serviceCredentials;
         _namespacePolicyResolver = namespacePolicyResolver;
+        _httpContextAccessor = httpContextAccessor;
     }
 
     public static IEnumerable<IBlobStore> GetBlobStores(IServiceProvider provider, IOptionsMonitor<HordeStorageSettings> settings)
@@ -171,11 +175,16 @@ public class BlobService : IBlobService
 
     private async Task<BlobIdentifier> PutObjectToStores(NamespaceId ns, IBufferedPayload bufferedPayload, BlobIdentifier identifier)
     {
+        IServerTiming? serverTiming = _httpContextAccessor.HttpContext?.RequestServices.GetService<IServerTiming>();
+
         foreach (IBlobStore store in _blobStores)
         {
             using IScope scope = Tracer.Instance.StartActive("put_blob_to_store");
             scope.Span.ResourceName = identifier.ToString();
             scope.Span.SetTag("store", store.ToString());
+            string storeName = store.GetType().Name;
+
+            using ServerTimingMetricScoped? serverTimingScope = serverTiming?.CreateServerTimingMetricScope($"blob.put.{storeName}", $"PUT to store: '{storeName}'");
 
             await using Stream s = bufferedPayload.GetStream();
             await store.PutObject(ns, s, identifier);
@@ -200,12 +209,17 @@ public class BlobService : IBlobService
         bool seenNamespaceNotFound = false;
         int numStoreMisses = 0;
         BlobContents? blobContents = null;
+        IServerTiming? serverTiming = _httpContextAccessor.HttpContext?.RequestServices.GetService<IServerTiming>();
+
         foreach (IBlobStore store in _blobStores)
         {
             using IScope scope = Tracer.Instance.StartActive("HierarchicalStore.GetObject");
             scope.Span.ResourceName = blob.ToString();
             scope.Span.SetTag("BlobStore", store.GetType().Name);
             scope.Span.SetTag("ObjectFound", false.ToString());
+            string storeName = store.GetType().Name;
+            using ServerTimingMetricScoped? serverTimingScope = serverTiming?.CreateServerTimingMetricScope($"blob.get.{storeName}", $"Blob GET from: '{storeName}'");
+
             try
             {
                 blobContents = await store.GetObject(ns, blob);
@@ -242,6 +256,8 @@ public class BlobService : IBlobService
         if (numStoreMisses >= 1)
         {
             using IScope _ = Tracer.Instance.StartActive("HierarchicalStore.Populate");
+            using ServerTimingMetricScoped? serverTimingScope = serverTiming?.CreateServerTimingMetricScope($"blob.populate", "Populating caches with blob contents");
+
             await using MemoryStream tempStream = new MemoryStream();
             await blobContents.Stream.CopyToAsync(tempStream);
             byte[] data = tempStream.ToArray();
@@ -270,9 +286,11 @@ public class BlobService : IBlobService
         {
             throw new NotSupportedException($"Replication is not allowed in namespace {ns}");
         }
-
+        IServerTiming? serverTiming = _httpContextAccessor.HttpContext?.RequestServices.GetService<IServerTiming>();
         using IScope scope = Tracer.Instance.StartActive("HierarchicalStore.Replicate");
 
+        using ServerTimingMetricScoped? serverTimingScope = serverTiming?.CreateServerTimingMetricScope("blob.replicate", "Replicating blob from remote instances");
+
         BlobInfo? blobInfo = await _blobIndex.GetBlobInfo(ns, blob);
         if (blobInfo == null)
         {
diff --git a/Engine/Source/Programs/Horde/HordeStorage/Horde.Storage/Implementation/ObjectService.cs b/Engine/Source/Programs/Horde/HordeStorage/Horde.Storage/Implementation/ObjectService.cs
--- a/Engine/Source/Programs/Horde/HordeStorage/Horde.Storage/Implementation/ObjectService.cs
+++ b/Engine/Source/Programs/Horde/HordeStorage/Horde.Storage/Implementation/ObjectService.cs
@@ -6,11 +6,14 @@ using System.Linq;
 using System.Threading.Tasks;
 using async_enumerable_dotnet;
 using Datadog.Trace;
+using EpicGames.AspNet;
 using EpicGames.Horde.Storage;
 using EpicGames.Serialization;
 using Horde.Storage.Implementation.Blob;
 using Jupiter.Implementation;
 using Jupiter.Utils;
+using Microsoft.AspNetCore.Http;
+using Microsoft.Extensions.DependencyInjection;
 using Serilog;
 
 namespace Horde.Storage.Implementation
@@ -30,6 +33,7 @@ namespace Horde.Storage.Implementation
 
     public class ObjectService : IObjectService
     {
+        private readonly IHttpContextAccessor _httpContextAccessor;
         private readonly IReferencesStore _referencesStore;
         private readonly IBlobService _blobService;
         private readonly IReferenceResolver _referenceResolver;
@@ -38,8 +42,9 @@ namespace Horde.Storage.Implementation
         private readonly ILastAccessTracker<LastAccessRecord> _lastAccessTracker;
         private readonly ILogger _logger = Log.ForContext<ObjectService>();
 
-        public ObjectService(IReferencesStore referencesStore, IBlobService blobService, IReferenceResolver referenceResolver, IReplicationLog replicationLog, IBlobIndex blobIndex, ILastAccessTracker<LastAccessRecord> lastAccessTracker)
+        public ObjectService(IHttpContextAccessor httpContextAccessor, IReferencesStore referencesStore, IBlobService blobService, IReferenceResolver referenceResolver, IReplicationLog replicationLog, IBlobIndex blobIndex, ILastAccessTracker<LastAccessRecord> lastAccessTracker)
         {
+            _httpContextAccessor = httpContextAccessor;
             _referencesStore = referencesStore;
             _blobService = blobService;
             _referenceResolver = referenceResolver;
@@ -66,7 +71,15 @@ namespace Horde.Storage.Implementation
                         : IReferencesStore.FieldFlags.None;
                 }
             }
-            ObjectRecord o = await _referencesStore.Get(ns, bucket, key, flags);
+
+            IServerTiming? serverTiming = _httpContextAccessor.HttpContext?.RequestServices.GetService<IServerTiming>();
+
+            ObjectRecord o;
+            {
+                using ServerTimingMetricScoped? serverTimingScope = serverTiming?.CreateServerTimingMetricScope("ref.get", "Fetching Ref from DB");
+
+                o = await _referencesStore.Get(ns, bucket, key, flags);
+            }
 
             // we do not wait for the last access tracking as it does not matter when it completes
             Task lastAccessTask = _lastAccessTracker.TrackUsed(new LastAccessRecord(ns, bucket, key)).ContinueWith((task, _) =>
@@ -88,6 +101,8 @@ namespace Horde.Storage.Implementation
                 }
                 else
                 {
+                    using ServerTimingMetricScoped? serverTimingScope = serverTiming?.CreateServerTimingMetricScope("blob.get", "Downloading blob from store");
+
                     blobContents = await _blobService.GetObject(ns, o.BlobIdentifier);
                 }
             }
diff --git a/Engine/Source/Programs/Horde/HordeStorage/Horde.Storage/Tests/UnitTests/HierarchicalBlobStorageTest.cs b/Engine/Source/Programs/Horde/HordeStorage/Horde.Storage/Tests/UnitTests/HierarchicalBlobStorageTest.cs
--- a/Engine/Source/Programs/Horde/HordeStorage/Horde.Storage/Tests/UnitTests/HierarchicalBlobStorageTest.cs
+++ b/Engine/Source/Programs/Horde/HordeStorage/Horde.Storage/Tests/UnitTests/HierarchicalBlobStorageTest.cs
@@ -13,6 +13,7 @@ using Horde.Storage.Implementation.Blob;
 using Jupiter;
 using Jupiter.Common;
 using Jupiter.Implementation;
+using Microsoft.AspNetCore.Http;
 using Microsoft.Extensions.Options;
 using Microsoft.VisualStudio.TestTools.UnitTesting;
 using Moq;
@@ -56,7 +57,7 @@ namespace Horde.Storage.UnitTests
             Mock<INamespacePolicyResolver> mockPolicyResolver = new Mock<INamespacePolicyResolver>();
             mockPolicyResolver.Setup(x => x.GetPoliciesForNs(It.IsAny<NamespaceId>())).Returns(new NamespacePolicy());
 
-            _chained = new BlobService(serviceProviderMock.Object, settingsMonitor, Mock.Of<IBlobIndex>(), Mock.Of<IPeerStatusService>(), Mock.Of<IHttpClientFactory>(), Mock.Of<IServiceCredentials>(), mockPolicyResolver.Object);
+            _chained = new BlobService(serviceProviderMock.Object, settingsMonitor, Mock.Of<IBlobIndex>(), Mock.Of<IPeerStatusService>(), Mock.Of<IHttpClientFactory>(), Mock.Of<IServiceCredentials>(), mockPolicyResolver.Object, Mock.Of<IHttpContextAccessor>());
             _chained.BlobStore = new List<IBlobStore> { _first, _second, _third };
 
             await _first.PutObject(Ns, Encoding.ASCII.GetBytes("onlyFirstContent"), _onlyFirstId);
diff --git a/Engine/Source/Programs/Horde/HordeStorage/Jupiter.Common/BaseStartup.cs b/Engine/Source/Programs/Horde/HordeStorage/Jupiter.Common/BaseStartup.cs
--- a/Engine/Source/Programs/Horde/HordeStorage/Jupiter.Common/BaseStartup.cs
+++ b/Engine/Source/Programs/Horde/HordeStorage/Jupiter.Common/BaseStartup.cs
@@ -39,7 +39,7 @@ namespace Jupiter
 {
     public abstract class BaseStartup
     {
-	    protected ILogger Logger { get; } = Log.ForContext<BaseStartup>();
+        protected ILogger Logger { get; } = Log.ForContext<BaseStartup>();
 
         protected BaseStartup(IConfiguration configuration)
         {
@@ -56,6 +56,8 @@ namespace Jupiter
         {
             CbConvertersAspNet.AddAspnetConverters();
 
+            services.AddServerTiming();
+
             // aws specific settings
             services.AddOptions<AWSCredentialsSettings>().Bind(Configuration.GetSection("AWSCredentials")).ValidateDataAnnotations();
             services.AddDefaultAWSOptions(Configuration.GetAWSOptions());
@@ -282,7 +284,9 @@ namespace Jupiter
             app.UseAuthentication();
             //app.UseMiddleware<DatadogTraceMiddleware>("Authorization");
             app.UseAuthorization();
-            
+
+            app.UseMiddleware<ServerTimingMiddleware>();
+
             //app.UseMiddleware<DatadogTraceMiddleware>("Endpoints");
             app.UseEndpoints(endpoints =>
             {
@@ -477,14 +481,14 @@ namespace Jupiter
 
     public class NamespacePolicy
     {
-	    public string[] Claims { get; set; } = Array.Empty<string>();
-	    public string StoragePool { get; set; } = "";
-
-	    public bool LastAccessTracking { get; set; } = true;
-	    public bool OnDemandReplication { get; set; } = false;
-	    public bool UseBlobIndexForExists { get; set; } = false;
-	    public bool UseBlobIndexForSlowExists { get; set; } = false;
-	    public bool? IsLegacyNamespace { get; set; } = null;
+        public string[] Claims { get; set; } = Array.Empty<string>();
+        public string StoragePool { get; set; } = "";
+
+        public bool LastAccessTracking { get; set; } = true;
+        public bool OnDemandReplication { get; set; } = false;
+        public bool UseBlobIndexForExists { get; set; } = false;
+        public bool UseBlobIndexForSlowExists { get; set; } = false;
+        public bool? IsLegacyNamespace { get; set; } = null;
         public bool IsPublicNamespace { get; set; } = true;
     }
 
