Horde: Fix warnings for log message templates.

diff --git a/Engine/Source/Programs/Horde/Horde.Build/Projects/Impl/StreamCollection.cs b/Engine/Source/Programs/Horde/Horde.Build/Projects/Impl/StreamCollection.cs
--- a/Engine/Source/Programs/Horde/Horde.Build/Projects/Impl/StreamCollection.cs
+++ b/Engine/Source/Programs/Horde/Horde.Build/Projects/Impl/StreamCollection.cs
@@ -1,4 +1,4 @@
-﻿// Copyright Epic Games, Inc. All Rights Reserved.
+// Copyright Epic Games, Inc. All Rights Reserved.
 
 using System;
 using System.Collections.Generic;
@@ -139,7 +139,7 @@ namespace Horde.Build.Collections.Impl
 			}
 			catch (Exception)
 			{
-				_logger.LogError("Unable to get stream config for {streamId}/{streamName} at {revision}", stream.Id, stream.Name, stream.ConfigRevision);
+				_logger.LogError("Unable to get stream config for {StreamId} at {Revision}", stream.Id, stream.ConfigRevision);
 				throw;
 			}
 		}
