Horde: Roll with the initial link event fix, the new one grouped sequential link errors together
#rnx

diff --git a/Engine/Source/Programs/Horde/Horde.Agent.Tests/LogParserTests.cs b/Engine/Source/Programs/Horde/Horde.Agent.Tests/LogParserTests.cs
--- a/Engine/Source/Programs/Horde/Horde.Agent.Tests/LogParserTests.cs
+++ b/Engine/Source/Programs/Horde/Horde.Agent.Tests/LogParserTests.cs
@@ -645,8 +645,7 @@ namespace Horde.Agent.Tests
 			Assert.AreEqual("Foo::Bar", symbolProperty.Properties!["identifier"].ToString());
 		}
 
-		[TestMethod]
-		[Ignore]
+		[TestMethod]		
 		public void LinuxLinkErrorMatcher()
 		{
 			string[] lines =
diff --git a/Engine/Source/Programs/Horde/Horde.Agent/Parser/Matchers/LinkEventMatcher.cs b/Engine/Source/Programs/Horde/Horde.Agent/Parser/Matchers/LinkEventMatcher.cs
--- a/Engine/Source/Programs/Horde/Horde.Agent/Parser/Matchers/LinkEventMatcher.cs
+++ b/Engine/Source/Programs/Horde/Horde.Agent/Parser/Matchers/LinkEventMatcher.cs
@@ -18,7 +18,7 @@ namespace Horde.Agent.Parser.Matchers
 			int lineCount = 0;
 			bool isError = false;
 			bool isWarning = false;
-			while (cursor.IsMatch(lineCount, @": undefined reference to |undefined symbol|linker command failed|^\s*(ld|ld.lld):|^[^:]*[^a-zA-Z][a-z]?ld: |^\s*>>>"))
+			while (cursor.IsMatch(lineCount, @": undefined reference to |undefined symbol|^\s*(ld|ld.lld|(.*clang)):|^[^:]*[^a-zA-Z][a-z]?ld: |^\s*>>>"))
 			{
 				isError |= cursor.IsMatch("error:");
 				isWarning |= cursor.IsMatch("warning:");
