Horde Storage Benchmarker - Reduced number of requests sent in parallell to blobs, to avoid to much load on certain s3 indexes (as we have a limited set of blobs we frequently hit the same indexes).
Some other fixes up needed to cathc up to new ubuntu image.

diff --git a/Engine/Source/Programs/Horde/HordeStorage/Benchmarks/benchmark-suite/tests.py b/Engine/Source/Programs/Horde/HordeStorage/Benchmarks/benchmark-suite/tests.py
--- a/Engine/Source/Programs/Horde/HordeStorage/Benchmarks/benchmark-suite/tests.py
+++ b/Engine/Source/Programs/Horde/HordeStorage/Benchmarks/benchmark-suite/tests.py
@@ -62,11 +62,11 @@ def seed_small_blobs(args):
 
 
 def test_small_blob_uploads(args):
-    helpers.run_test(args, 'small_blob_uploads', duration_seconds=60, rate_per_second=2500)
+    helpers.run_test(args, 'small_blob_uploads', duration_seconds=60, rate_per_second=100)
 
 
 def test_small_blob_downloads(args):
-    helpers.run_test(args, 'small_blob_downloads', duration_seconds=60, rate_per_second=2500)
+    helpers.run_test(args, 'small_blob_downloads', duration_seconds=60, rate_per_second=100)
 
 def seed_large_blobs(args):
     ns = args['namespace']
