Horde Storage - Allow hinting for last access tracking to the blob store (used by make sure the filesystem cache doesnt update its internal last access tracking)

diff --git a/Engine/Source/Programs/Horde/HordeStorage/Horde.Storage/Implementation/Blob/AmazonS3Store.cs b/Engine/Source/Programs/Horde/HordeStorage/Horde.Storage/Implementation/Blob/AmazonS3Store.cs
--- a/Engine/Source/Programs/Horde/HordeStorage/Horde.Storage/Implementation/Blob/AmazonS3Store.cs
+++ b/Engine/Source/Programs/Horde/HordeStorage/Horde.Storage/Implementation/Blob/AmazonS3Store.cs
@@ -142,7 +142,7 @@ namespace Horde.Storage.Implementation
             }
         }
 
-        public async Task<BlobContents> GetObject(NamespaceId ns, BlobIdentifier blob)
+        public async Task<BlobContents> GetObject(NamespaceId ns, BlobIdentifier blob, LastAccessTrackingFlags flags = LastAccessTrackingFlags.DoTracking)
         {
             string bucketName = GetBucketName(ns);
             GetObjectResponse response;
diff --git a/Engine/Source/Programs/Horde/HordeStorage/Horde.Storage/Implementation/Blob/AzureBlobStore.cs b/Engine/Source/Programs/Horde/HordeStorage/Horde.Storage/Implementation/Blob/AzureBlobStore.cs
--- a/Engine/Source/Programs/Horde/HordeStorage/Horde.Storage/Implementation/Blob/AzureBlobStore.cs
+++ b/Engine/Source/Programs/Horde/HordeStorage/Horde.Storage/Implementation/Blob/AzureBlobStore.cs
@@ -108,7 +108,8 @@ namespace Horde.Storage.Implementation
             await blob.SetMetadataAsync(metadata);
         }
 
-        public async Task<BlobContents> GetObject(NamespaceId ns, BlobIdentifier blobIdentifier)
+        public async Task<BlobContents> GetObject(NamespaceId ns, BlobIdentifier blobIdentifier,
+            LastAccessTrackingFlags flags)
         {
             string fixedNamespace = SanitizeNamespace(ns);
             BlobContainerClient container = new BlobContainerClient(_connectionString, fixedNamespace);
diff --git a/Engine/Source/Programs/Horde/HordeStorage/Horde.Storage/Implementation/Blob/BlobStoreConsistencyCheckService.cs b/Engine/Source/Programs/Horde/HordeStorage/Horde.Storage/Implementation/Blob/BlobStoreConsistencyCheckService.cs
--- a/Engine/Source/Programs/Horde/HordeStorage/Horde.Storage/Implementation/Blob/BlobStoreConsistencyCheckService.cs
+++ b/Engine/Source/Programs/Horde/HordeStorage/Horde.Storage/Implementation/Blob/BlobStoreConsistencyCheckService.cs
@@ -104,7 +104,7 @@ namespace Horde.Storage.Implementation
 
                         Interlocked.Increment(ref countOfBlobsChecked);
                         
-                        BlobContents contents = await blobStore.GetObject(ns, blob);
+                        BlobContents contents = await blobStore.GetObject(ns, blob, LastAccessTrackingFlags.SkipTracking);
                         await using Stream s = contents.Stream;
 
                         bool inconsistencyFound = false;
diff --git a/Engine/Source/Programs/Horde/HordeStorage/Horde.Storage/Implementation/Blob/FileSystemStore.cs b/Engine/Source/Programs/Horde/HordeStorage/Horde.Storage/Implementation/Blob/FileSystemStore.cs
--- a/Engine/Source/Programs/Horde/HordeStorage/Horde.Storage/Implementation/Blob/FileSystemStore.cs
+++ b/Engine/Source/Programs/Horde/HordeStorage/Horde.Storage/Implementation/Blob/FileSystemStore.cs
@@ -104,7 +104,7 @@ namespace Horde.Storage.Implementation
 			return await PutObject(ns, stream, blobIdentifier);
         }
 
-        public Task<BlobContents> GetObject(NamespaceId ns, BlobIdentifier blob)
+        public Task<BlobContents> GetObject(NamespaceId ns, BlobIdentifier blob, LastAccessTrackingFlags flags)
         {
             FileInfo filePath = GetFilesystemPath(ns, blob);
 
@@ -113,7 +113,10 @@ namespace Horde.Storage.Implementation
                 throw new BlobNotFoundException(ns, blob);
             }
 
-            UpdateLastWriteTime(filePath.FullName, DateTime.UtcNow);
+            if (flags == LastAccessTrackingFlags.DoTracking)
+            {
+                UpdateLastWriteTime(filePath.FullName, DateTime.UtcNow);
+            }
             FileStream fs = new FileStream(filePath.FullName, FileMode.Open, FileAccess.Read, FileShare.Read, DefaultBufferSize, FileOptions.Asynchronous | FileOptions.SequentialScan);
 
             return Task.FromResult(new BlobContents(fs, fs.Length));
diff --git a/Engine/Source/Programs/Horde/HordeStorage/Horde.Storage/Implementation/Blob/IBlobStore.cs b/Engine/Source/Programs/Horde/HordeStorage/Horde.Storage/Implementation/Blob/IBlobStore.cs
--- a/Engine/Source/Programs/Horde/HordeStorage/Horde.Storage/Implementation/Blob/IBlobStore.cs
+++ b/Engine/Source/Programs/Horde/HordeStorage/Horde.Storage/Implementation/Blob/IBlobStore.cs
@@ -9,13 +9,19 @@ using Jupiter.Implementation;
 
 namespace Horde.Storage.Implementation
 {
+    public enum LastAccessTrackingFlags
+    {
+        DoTracking = 0,
+        SkipTracking = 1,
+    }
+
     public interface IBlobStore
     {
         Task<BlobIdentifier> PutObject(NamespaceId ns, byte[] blob, BlobIdentifier identifier);
         Task<BlobIdentifier> PutObject(NamespaceId ns, ReadOnlyMemory<byte> blob, BlobIdentifier identifier);
         Task<BlobIdentifier> PutObject(NamespaceId ns, Stream content, BlobIdentifier identifier);
 
-        Task<BlobContents> GetObject(NamespaceId ns, BlobIdentifier blob);
+        Task<BlobContents> GetObject(NamespaceId ns, BlobIdentifier blob, LastAccessTrackingFlags flags = LastAccessTrackingFlags.DoTracking);
         Task<bool> Exists(NamespaceId ns, BlobIdentifier blob);
 
         // Delete a object
diff --git a/Engine/Source/Programs/Horde/HordeStorage/Horde.Storage/Implementation/Blob/MemoryBlobStore.cs b/Engine/Source/Programs/Horde/HordeStorage/Horde.Storage/Implementation/Blob/MemoryBlobStore.cs
--- a/Engine/Source/Programs/Horde/HordeStorage/Horde.Storage/Implementation/Blob/MemoryBlobStore.cs
+++ b/Engine/Source/Programs/Horde/HordeStorage/Horde.Storage/Implementation/Blob/MemoryBlobStore.cs
@@ -72,7 +72,7 @@ namespace Horde.Storage.Implementation
             return await PutObject(ns, blob: await blob.ToByteArray(), identifier);
         }
 
-        public Task<BlobContents> GetObject(NamespaceId ns, BlobIdentifier blob)
+        public Task<BlobContents> GetObject(NamespaceId ns, BlobIdentifier blob, LastAccessTrackingFlags flags = LastAccessTrackingFlags.DoTracking)
         {
             if (!_blobs.TryGetValue(ns, value: out ConcurrentDictionary<BlobIdentifier, BlobContainer>? namespaceContainer))
             {
diff --git a/Engine/Source/Programs/Horde/HordeStorage/Horde.Storage/Implementation/Blob/MemoryCacheBlobStore.cs b/Engine/Source/Programs/Horde/HordeStorage/Horde.Storage/Implementation/Blob/MemoryCacheBlobStore.cs
--- a/Engine/Source/Programs/Horde/HordeStorage/Horde.Storage/Implementation/Blob/MemoryCacheBlobStore.cs
+++ b/Engine/Source/Programs/Horde/HordeStorage/Horde.Storage/Implementation/Blob/MemoryCacheBlobStore.cs
@@ -76,7 +76,7 @@ namespace Horde.Storage.Implementation
             return await PutObject(ns, blob: await blob.ToByteArray(), identifier);
         }
 
-        public Task<BlobContents> GetObject(NamespaceId ns, BlobIdentifier identifier)
+        public Task<BlobContents> GetObject(NamespaceId ns, BlobIdentifier identifier, LastAccessTrackingFlags flags = LastAccessTrackingFlags.DoTracking)
         {
             bool exists = _memoryCache.TryGetValue(BuildKey(ns, identifier), out object? dataObj);
             if (!exists)
diff --git a/Engine/Source/Programs/Horde/HordeStorage/Horde.Storage/Implementation/Blob/RelayBlobStore.cs b/Engine/Source/Programs/Horde/HordeStorage/Horde.Storage/Implementation/Blob/RelayBlobStore.cs
--- a/Engine/Source/Programs/Horde/HordeStorage/Horde.Storage/Implementation/Blob/RelayBlobStore.cs
+++ b/Engine/Source/Programs/Horde/HordeStorage/Horde.Storage/Implementation/Blob/RelayBlobStore.cs
@@ -53,7 +53,7 @@ namespace Horde.Storage.Implementation
             return identifier;
         }
 
-        public async Task<BlobContents> GetObject(NamespaceId ns, BlobIdentifier blob)
+        public async Task<BlobContents> GetObject(NamespaceId ns, BlobIdentifier blob, LastAccessTrackingFlags flags)
         {
             using HttpRequestMessage getObjectRequest = BuildHttpRequest(HttpMethod.Get, new Uri($"api/v1/blobs/{ns}/{blob}", UriKind.Relative));
             getObjectRequest.Headers.Add("Accept", MediaTypeNames.Application.Octet);
