Horde - Automated Test Session - Feedback fixes and revamp summary display

diff --git a/Engine/Source/Programs/Horde/HordeDashboard/src/Plugins/AutomatedTestSession/models/AutomatedTestSession.ts b/Engine/Source/Programs/Horde/HordeDashboard/src/Plugins/AutomatedTestSession/models/AutomatedTestSession.ts
--- a/Engine/Source/Programs/Horde/HordeDashboard/src/Plugins/AutomatedTestSession/models/AutomatedTestSession.ts
+++ b/Engine/Source/Programs/Horde/HordeDashboard/src/Plugins/AutomatedTestSession/models/AutomatedTestSession.ts
@@ -15,9 +15,9 @@ export enum TestState {
 
 const TestStateRank: {[Key in string]: number}  = {
     "Fail": 0,
-    "SuccessWithWarnings": 1,
-    "InProcess": 2,
-    "NotRun": 3,
+    "InProcess": 1,
+    "NotRun": 2,
+    "SuccessWithWarnings": 3,
     "Success": 4,
     "Skipped": 5,
     "Unknown": 6,
@@ -56,7 +56,7 @@ export type TestResult = {
     TestUID: string;
     Suite: string;
     State: string;
-    DeviceAppInstanceName: string;
+    DeviceAppInstanceName: string | string[];
     ErrorCount: number;
     WarningCount: number;
     ErrorHashAggregate: string;
@@ -74,6 +74,7 @@ export type TestResultDetails = {
 export type TestDevice = MetadataHolder& {
     Name: string;
     AppInstanceName: string;
+    AppInstanceLog: string;
 }
 
 export type TestSessionInfo = {
@@ -140,6 +141,12 @@ export class MetaWrapper {
         return [];
     }
 
+    forEach(predicate: (key: string, value?: string) => void) {
+        if (this.collection) {
+            Object.keys(this.collection).forEach((key) => predicate(key, this.collection![key]));
+        }
+    }
+
     maskByKeys(keyMask: string[]) {
         return this.map((key, value) => keyMask.includes(key)?undefined:value).filter((info) => info)
     }
@@ -256,9 +263,19 @@ export class TestSessionWrapper implements TestSession, DataWrapper {
         return undefined
     }
 
-    getDevice(AppInstanceName: string) {
+    getDevices(AppInstanceName: string | string[]) {
         if(this.Devices) {
-            return this.Devices.find((item) => item.AppInstanceName === AppInstanceName);
+            if (typeof(AppInstanceName) === 'string') {
+                const asString = AppInstanceName as string;
+                const device = this.Devices.find((item) => item.AppInstanceName === asString);
+                if(device) {
+                    return [device];
+                }
+            }
+            else {
+                const asArray = AppInstanceName as string[];
+                return this.Devices.filter((item) => asArray.includes(item.AppInstanceName))
+            }
         }
 
         return undefined
@@ -285,20 +302,25 @@ export class TestStats {
     Skipped: number = 0;
     Unexecuted: number = 0;
     Incomplete: number = 0;
+    TotalRun: number = 0;
 
     update(state: string, increment: boolean = true) {
         switch(state) {
         case TestState.Failed:
             increment? this.Failed++ : this.Failed--;
+            increment? this.TotalRun++ : this.TotalRun--
             break;
         case TestState.NotRun:
             increment? this.Unexecuted++ : this.Unexecuted--;
             break;
         case TestState.Success:
+        case TestState.SuccessWithWarnings:
             increment? this.Passed++ : this.Passed--;
+            increment? this.TotalRun++ : this.TotalRun--
             break;
         case TestState.InProcess:
             increment? this.Incomplete++ : this.Incomplete--;
+            increment? this.TotalRun++ : this.TotalRun--
             break;
         case TestState.Skipped:
             increment? this.Skipped++ : this.Skipped--;
@@ -312,6 +334,7 @@ export class TestStats {
         this.Skipped += stats.Skipped;
         this.Unexecuted += stats.Unexecuted;
         this.Incomplete += stats.Incomplete;
+        this.TotalRun += stats.TotalRun;
     }
 }
 
@@ -322,6 +345,7 @@ export class TestCase implements Loader {
     SessionName: string;
     Suite: string;
     State: string;
+    StateMetaUID: string | undefined;
     Metas: Map<string, TestResult>;
     Stats: TestStats;
 
@@ -366,6 +390,7 @@ export class TestCase implements Loader {
         this.DisplayName = displayName??this.toDislayName(name);
         this.TestUID = testuid;
         this.State = TestState.Unknown;
+        this.StateMetaUID = undefined;
         this.Suite = suite;
         this.SessionName = sessionName;
         this.Metas = new Map();
@@ -399,7 +424,9 @@ export class TestCase implements Loader {
             state = TestState.SuccessWithWarnings;
             item.State = state;
         }
-        this.setState(state);
+        if(this.setState(state)) {
+            this.StateMetaUID = metauid;
+        }
         this.Stats.update(state);
         this.updateVersion();
         return true;
@@ -417,9 +444,25 @@ export class TestCase implements Loader {
         return results;
     }
 
-    getFilteredTestCase(metafilter: Metadata) {
+    getFilteredByMeta(metafilter: MetadataFilter) {
         const metas = Array.from(this.Metas.values()).filter(
-            (test) => Object.keys(metafilter).every((key) => test.Session.Metadata[key] === metafilter[key])
+            (test) => Object.keys(metafilter).every((key) => {
+                const sessionKey = test.Session.Metadata[key];
+                return sessionKey && metafilter[key].includes(sessionKey);
+            })
+        );
+        if(metas.length === 0) {
+            return undefined;
+        }
+        const testcase = new TestCase(this.Name, this.TestUID, this.Suite, this.SessionName, this.DisplayName);
+        metas.forEach((test) => testcase.addMeta(test));
+        return testcase;
+    }
+
+    getFilteredByState(state: TestState) {
+        const targetRank = TestStateRank[state];
+        const metas = Array.from(this.Metas.values()).filter(
+            (test) => targetRank <= TestStateRank[test.State]
         );
         if(metas.length === 0) {
             return undefined;
@@ -443,8 +486,10 @@ export class TestCase implements Loader {
             const newRank = TestStateRank[state];
             if(newRank < currentRank) {
                 this.State = state;
+                return true;
             }
         }
+        return false;
     }
 
     private sortResults(a: TestResult, b: TestResult) {
@@ -496,6 +541,20 @@ export enum FilterType {
     testchange = 'testchange',
     view = 'view',
     namecontains = 'namecontains',
+    meta = 'meta',
+}
+
+export type MetadataFilter = {[Key in string]: string[]}
+export class MetaFilterTools {
+    static Match = (key: string, target: string) : boolean => target.startsWith(`${key}.`)
+    static Encode = (key: string, items: string[]) : string => `${key}.${items.join('~')}`
+    static Decode = (items: string[]) : MetadataFilter | undefined => {
+        return items.length === 0 ? undefined :
+            Object.fromEntries(items.map((item) => {
+                const [key, value] = item.split('.', 2);
+                return [key, value.split('~')]
+            }))
+    }    
 }
 
 export type SectionFilter = {
@@ -503,6 +562,7 @@ export type SectionFilter = {
     name?: string;
     project?: string;
     namecontains?: string;
+    meta?: MetadataFilter;
 }
 
 export type ReportFilter = {
@@ -512,7 +572,22 @@ export type ReportFilter = {
     view?: string;
 }
 
+export const generateHashFromSectionFilter = (obj: SectionFilter) => {
+    const baseFilter : Metadata = Object.fromEntries(
+        Object.keys(obj).filter((key) => key !== 'meta').map((key) => [
+            key, Object.getOwnPropertyDescriptor(obj, key)?.value
+        ])
+    );
+    let aggregate = Object.keys(baseFilter).sort().map((key) => key + (baseFilter[key]??'')).join();
+    const subObj = obj.meta;
+    if(subObj !== undefined) {
+        aggregate += 'meta'+Object.keys(subObj).sort().map((key) => key + (subObj[key]?.join('~')??'')).join();
+    }
+    return generatetHashFromString(aggregate);
+}
+
 type GetKeyFunction = (testcase: TestCase) => [string, string];
+const regexNameFilter = /"([^"]+)"|([^\s"]+)/g;
 
 export class SectionCollection {
 
@@ -525,7 +600,7 @@ export class SectionCollection {
 
     extractSections(filter: SectionFilter, sessions: TestSessionCollection) {
         // Get the sections, first look for in the cache, otherwise generate it
-        const filterCacheKey = generateHashFromMetadata(filter);
+        const filterCacheKey = generateHashFromSectionFilter(filter);
         const cache = this.filterCache.get(filterCacheKey);
         const isSameHash = cache && cache.hash === sessions.hash;
         const sections = cache && isSameHash? cache.sections : this.filterSections(filter, sessions.getTestCases());
@@ -540,7 +615,7 @@ export class SectionCollection {
         let sections = this.extractSections({...filter, namecontains: undefined}, sessions);
         const onTheFlyFilter = {...filter, namecontains: testNameFilter.toLowerCase()}
         sections = this.filterSections(onTheFlyFilter, sections.flatMap((section) => section.TestCases));
-        this.onTheFlyCache = {filter: generateHashFromMetadata(onTheFlyFilter), hash: sessions.hash, sections};
+        this.onTheFlyCache = {filter: generateHashFromSectionFilter(onTheFlyFilter), hash: sessions.hash, sections};
         return sections;
     }
 
@@ -559,8 +634,8 @@ export class SectionCollection {
                 }
             }
             if(filter.namecontains) {
-                const substrings = filter.namecontains.split(' ');
-                const testName = test.DisplayName.toLowerCase();
+                const substrings = Array.from(filter.namecontains.matchAll(regexNameFilter), match => match[1] || match[2]);
+                const testName = test.Name.toLowerCase();
                 if(!substrings.some((value) => testName.indexOf(value) >= 0)) {
                     return false;
                 }
@@ -572,7 +647,19 @@ export class SectionCollection {
             }
             return true;
         }
-        const filteredTests = tests.filter(filterFunc);
+        const filteredTests = tests.reduce((filtered: TestCase[], item: TestCase) => {
+            if(filterFunc(item)) {
+                if(filter.meta === undefined) {
+                    filtered.push(item);
+                } else {
+                    const filteredCase = item.getFilteredByMeta(filter.meta);
+                    if (filteredCase) {
+                        filtered.push(filteredCase);
+                    }
+                }
+            }
+            return filtered
+        }, []);
 
         const commonSection = filter.name? this.getCommonSection(filteredTests.map((item) => item.DisplayName)) : "";
         const sectionKeyFunc = (key: FilterType): GetKeyFunction | undefined => {
@@ -678,7 +765,11 @@ export class SectionCollection {
 
     // sort in that order Failed, Incomplete, Unexecuted, Passed, Skipped
     private sortSectionResults(section: Section) {
-        section.TestCases.sort((a, b) => {
+        SectionCollection.sortTestCases(section.TestCases);
+    }
+
+    static sortTestCases(testcases: TestCase[]) {
+        testcases.sort((a, b) => {
             const aRank = TestStateRank[a.State];
             const bRank = TestStateRank[b.State];
             if(aRank > bRank) {
@@ -747,6 +838,7 @@ export class TestSessionCollection implements Loader {
     private _testcases: Map<string, TestCase>;
     private _datahandler: TestDataCollection;
     private _includePreflight: boolean = false;
+    private _cachedMetaValueByKey: {hash: string, value: Map<string, Set<string>>} | undefined;
 
     constructor(handler: TestDataCollection, includePreflight: boolean = false) {
         this._collection = new Map();
@@ -798,6 +890,32 @@ export class TestSessionCollection implements Loader {
         this._collection.forEach(func);
     }
 
+    getMetaValueByKey() {
+        if(this._cachedMetaValueByKey !== undefined) {
+            if(this._cachedMetaValueByKey.hash === this.hash) {
+                return this._cachedMetaValueByKey.value;
+            }
+        }
+        const keys: Map<string, Set<string>> = new Map();
+        this.forEachSession((item) => {
+            const meta = item.MetaHandler;
+            if(meta) {
+                meta.forEach((key, value) => {
+                    if (value) {
+                        let valueSet = keys.get(key);
+                        if (!valueSet) {
+                            valueSet = new Set();
+                            keys.set(key, valueSet);
+                        }
+                        valueSet.add(value);
+                    }
+                });
+            }
+        });
+        this._cachedMetaValueByKey = {hash:this.hash, value:keys};
+        return keys;
+    }
+
     *iterSessions() {
         const iterator = this._collection.values();
         let result = iterator.next();
