Fixed AutomationTool not being copied for UE5 Projects

AutomationTool.exe moved for UE5, so no UE5 projects would copy UAT files down at the start of their batch.

#rnx

diff --git a/Engine/Source/Programs/Horde/Horde.Agent/Execution/BuildGraphExecutor.cs b/Engine/Source/Programs/Horde/Horde.Agent/Execution/BuildGraphExecutor.cs
--- a/Engine/Source/Programs/Horde/Horde.Agent/Execution/BuildGraphExecutor.cs
+++ b/Engine/Source/Programs/Horde/Horde.Agent/Execution/BuildGraphExecutor.cs
@@ -286,8 +286,13 @@ namespace Horde.Agent.Execution
 		{
 			DirectoryReference buildDir = GetAutomationToolDir(sharedStorageDir);
 
-			FileReference automationTool = FileReference.Combine(buildDir, "Engine", "Binaries", "DotNET", "AutomationTool.exe");
-			if(FileReference.Exists(automationTool))
+			FileReference[] automationToolPaths = new FileReference[]
+			{
+				FileReference.Combine(buildDir, "Engine", "Binaries", "DotNET", "AutomationTool.exe"),
+				FileReference.Combine(buildDir, "Engine", "Binaries", "DotNET", "AutomationTool", "AutomationTool.exe")
+			};
+
+			if (automationToolPaths.Any(automationTool => FileReference.Exists(automationTool)))
 			{
 				logger.LogInformation("Copying AutomationTool binaries from '{BuildDir}' to '{WorkspaceDir}", buildDir, workspaceDir);
 				foreach (FileReference sourceFile in DirectoryReference.EnumerateFiles(buildDir, "*", SearchOption.AllDirectories))
