Horde: AWS Lambda execution fixes

- Ensure the executable inside a compute task workload is marked as an executable on Linux
- Proper lookup of AWS Param Store lookup (missing decryption parameter)

diff --git a/Engine/Source/Programs/Horde/Horde.Agent/Commands/Compute/AwsLambdaCommand.cs b/Engine/Source/Programs/Horde/Horde.Agent/Commands/Compute/AwsLambdaCommand.cs
--- a/Engine/Source/Programs/Horde/Horde.Agent/Commands/Compute/AwsLambdaCommand.cs
+++ b/Engine/Source/Programs/Horde/Horde.Agent/Commands/Compute/AwsLambdaCommand.cs
@@ -154,6 +154,13 @@ namespace Horde.Agent.Commands
 	[Command("Compute", "AwsLambda", "Listen for AWS Lambda invocations")]
 	class AwsLambdaCommand : Command
 	{
+		private const string EnvVarHordeStorageUrl = "UE_HORDE_STORAGE_URL";
+		private const string EnvVarHordeStorageOAuthUrl = "UE_HORDE_STORAGE_OAUTH_URL";
+		private const string EnvVarHordeStorageOAuthGrantType = "UE_HORDE_STORAGE_OAUTH_GRANT_TYPE";
+		private const string EnvVarHordeStorageOAuthClientId = "UE_HORDE_STORAGE_OAUTH_CLIENT_ID";
+		private const string EnvVarHordeStorageOAuthClientSecretArn = "UE_HORDE_STORAGE_OAUTH_CLIENT_SECRET_ARN";
+		private const string EnvVarHordeStorageOAuthScope = "UE_HORDE_STORAGE_OAUTH_SCOPE";
+		
 		/// <inheritdoc/>
 		public override async Task<int> ExecuteAsync(ILogger logger)
 		{
@@ -220,10 +227,11 @@ namespace Horde.Agent.Commands
 
 		private static async Task<string> GetOAuthClientSecretAsync(CancellationToken cancellationToken)
 		{
-			string oAuthClientSecret = GetEnvVar("UE_HORDE_STORAGE_OAUTH_CLIENT_SECRET_ARN");
+			string oAuthClientSecret = GetEnvVar(EnvVarHordeStorageOAuthClientSecretArn);
 			AmazonSimpleSystemsManagementClient ssmClient = new AmazonSimpleSystemsManagementClient();
-			GetParameterResponse paramRes = await ssmClient.GetParameterAsync(new GetParameterRequest { Name = oAuthClientSecret }, cancellationToken);
-			return paramRes.Parameter.Value;
+			GetParameterRequest request = new GetParameterRequest { Name = oAuthClientSecret, WithDecryption = true };
+			GetParameterResponse response = await ssmClient.GetParameterAsync(request, cancellationToken);
+			return response.Parameter.Value;
 		}
 
 		private static HttpServiceClientOptions CreateHttpServiceClientOptionsFromEnv(string clientSecret)
diff --git a/Engine/Source/Programs/Horde/Horde.Agent/Execution/ComputeTaskExecutor.cs b/Engine/Source/Programs/Horde/Horde.Agent/Execution/ComputeTaskExecutor.cs
--- a/Engine/Source/Programs/Horde/Horde.Agent/Execution/ComputeTaskExecutor.cs
+++ b/Engine/Source/Programs/Horde/Horde.Agent/Execution/ComputeTaskExecutor.cs
@@ -5,6 +5,7 @@ using System.Collections.Generic;
 using System.Diagnostics;
 using System.IO;
 using System.Linq;
+using System.Runtime.InteropServices;
 using System.Text;
 using System.Threading;
 using System.Threading.Tasks;
@@ -115,6 +116,7 @@ namespace Horde.Agent.Execution
 			{
 				string fileName = FileReference.Combine(sandboxDir, task.WorkingDirectory.ToString(), task.Executable.ToString()).FullName;
 				string workingDirectory = DirectoryReference.Combine(sandboxDir, task.WorkingDirectory.ToString()).FullName;
+				EnsureFileIsExecutable(fileName);
 
 				Dictionary<string, string> newEnvironment = new Dictionary<string, string>();
 				foreach (System.Collections.DictionaryEntry? entry in Environment.GetEnvironmentVariables())
@@ -255,6 +257,15 @@ namespace Horde.Agent.Execution
 			return (tree, hash);
 		}
 
+		private void EnsureFileIsExecutable(string filePath)
+		{
+			if (RuntimeInformation.IsOSPlatform(OSPlatform.Linux))
+			{
+				ushort mode = 493; // 0755 octal Posix permission
+				FileUtils.SetFileMode_Linux(filePath, mode);
+			}
+		}
+
 		async Task<DirectoryNode> CreateDirectoryNode(NamespaceId namespaceId, string name, int baseDirLen, List<FileReference> sortedFiles, int minIdx, int maxIdx)
 		{
 			(_, IoHash hash) = await PutDirectoryTree(namespaceId, baseDirLen, sortedFiles, minIdx, maxIdx);
