Horde Storage - Force http2 on corp and public ports

diff --git a/Engine/Source/Programs/Horde/HordeStorage/Horde.Storage/appsettings.Production.json b/Engine/Source/Programs/Horde/HordeStorage/Horde.Storage/appsettings.Production.json
--- a/Engine/Source/Programs/Horde/HordeStorage/Horde.Storage/appsettings.Production.json
+++ b/Engine/Source/Programs/Horde/HordeStorage/Horde.Storage/appsettings.Production.json
@@ -41,13 +41,15 @@
   "Kestrel": {
     "Endpoints": {
       "Http": {
-        "Url": "http://0.0.0.0:80"
+        "Url": "http://0.0.0.0:80",
+        "Protocols": "Http2"
       },
       "InternalHttp": {
         "Url": "http://0.0.0.0:8080"
       },
       "CorpHttp": {
-        "Url": "http://0.0.0.0:8008"
+        "Url": "http://0.0.0.0:8008",
+        "Protocols": "Http2"
       }
     }
   }
