Horde Storage - Support running our different access controlled ports (public, corp and internal) on more then one port. Allows us to host http1 and http2 on different ports.

diff --git a/Engine/Source/Programs/Horde/HordeStorage/Horde.Storage/Controllers/RequestHelper.cs b/Engine/Source/Programs/Horde/HordeStorage/Horde.Storage/Controllers/RequestHelper.cs
--- a/Engine/Source/Programs/Horde/HordeStorage/Horde.Storage/Controllers/RequestHelper.cs
+++ b/Engine/Source/Programs/Horde/HordeStorage/Horde.Storage/Controllers/RequestHelper.cs
@@ -46,7 +46,7 @@ public class RequestHelper
 
         // namespace is a restricted namespace
         HttpContext context = request.HttpContext;
-        bool isPublicPort = context.Connection.LocalPort == _settings!.CurrentValue.PublicApiPort ||
+        bool isPublicPort = _settings!.CurrentValue.PublicApiPorts.Contains(context.Connection.LocalPort) ||
             /* unit tests do not run on ports, we consider them always on the internal port */ !(context.Connection.LocalPort == 0 && context.Connection.LocalIpAddress == null);
             
         if (isPublicPort)
diff --git a/Engine/Source/Programs/Horde/HordeStorage/Horde.Storage/appsettings.Production.json b/Engine/Source/Programs/Horde/HordeStorage/Horde.Storage/appsettings.Production.json
--- a/Engine/Source/Programs/Horde/HordeStorage/Horde.Storage/appsettings.Production.json
+++ b/Engine/Source/Programs/Horde/HordeStorage/Horde.Storage/appsettings.Production.json
@@ -31,26 +31,37 @@
     "ReplicationPollFrequencySeconds": 15
   },
 
-  "Jupiter": {
-    "InternalApiPort": 8080,
-    "CorpApiPort": 8008
-  },
+	"Jupiter": {
+		"PublicApiPorts": [ 80, 8081 ],
+		"CorpApiPorts": [ 8008, 8082 ],
+		"InternalApiPorts": [ 8080, 8083 ]
+	},
 
   "AllowedHosts": "*",
 
   "Kestrel": {
-    "Endpoints": {
-      "Http": {
-        "Url": "http://0.0.0.0:80",
-        "Protocols": "Http2"
-      },
-      "InternalHttp": {
-        "Url": "http://0.0.0.0:8080"
-      },
-      "CorpHttp": {
-        "Url": "http://0.0.0.0:8008",
-        "Protocols": "Http2"
-      }
-    }
+	"Endpoints": {
+		"Http": {
+			"Url": "http://0.0.0.0:80"
+		},
+		"Http2": {
+			"Url": "http://0.0.0.0:8081",
+			"Protocols": "Http2"
+		},
+		"CorpHttp": {
+			"Url": "http://0.0.0.0:8008"
+		},
+		"CorpHttp2": {
+			"Url": "http://0.0.0.0:8082",
+			"Protocols": "Http2"
+		},
+		"InternalHttp": {
+			"Url": "http://0.0.0.0:8080"
+		},
+		"InternalHttp2": {
+			"Url": "http://0.0.0.0:8083",
+			"Protocols": "Http2"
+		}
+	}
   }
 }
diff --git a/Engine/Source/Programs/Horde/HordeStorage/Jupiter.Common/BaseStartup.cs b/Engine/Source/Programs/Horde/HordeStorage/Jupiter.Common/BaseStartup.cs
--- a/Engine/Source/Programs/Horde/HordeStorage/Jupiter.Common/BaseStartup.cs
+++ b/Engine/Source/Programs/Horde/HordeStorage/Jupiter.Common/BaseStartup.cs
@@ -450,17 +450,23 @@ namespace Jupiter
         /// Port used to host the internally accessible api (as well as the public api).
         /// This hosts both public and private namespaces
         /// </summary>
-        public int InternalApiPort { get; set; } = 8080;
+        [System.Diagnostics.CodeAnalysis.SuppressMessage("Usage", "CA2227:Collection properties should be read only", Justification = "Used by the configuration system")]
+
+        public List<int> InternalApiPorts { get; set; } = new List<int>() { 8080 };
 
         /// <summary>
         /// Port that hosts public and private namespaces
         /// </summary>
-        public int CorpApiPort { get; set; } = 8008;
+        [System.Diagnostics.CodeAnalysis.SuppressMessage("Usage", "CA2227:Collection properties should be read only", Justification = "Used by the configuration system")]
+
+        public List<int> CorpApiPorts { get; set; } = new List<int>() { 8008 };
 
         /// <summary>
         /// Port that only hosts the public namespaces
         /// </summary>
-        public int PublicApiPort { get; set; } = 80;
+        [System.Diagnostics.CodeAnalysis.SuppressMessage("Usage", "CA2227:Collection properties should be read only", Justification = "Used by the configuration system")]
+
+        public List<int> PublicApiPorts { get; set; } = new List<int>() { 80 };
 
         // Enable to echo every request to the log file, usually this is more efficiently done on the load balancer
         public bool UseRequestLogging { get; set; } = false;
diff --git a/Engine/Source/Programs/Horde/HordeStorage/Jupiter.Common/InternalApiFilterAttribute.cs b/Engine/Source/Programs/Horde/HordeStorage/Jupiter.Common/InternalApiFilterAttribute.cs
--- a/Engine/Source/Programs/Horde/HordeStorage/Jupiter.Common/InternalApiFilterAttribute.cs
+++ b/Engine/Source/Programs/Horde/HordeStorage/Jupiter.Common/InternalApiFilterAttribute.cs
@@ -15,7 +15,7 @@ namespace Jupiter.Implementation
 		{
 			IOptionsMonitor<JupiterSettings>? settings = context.HttpContext.RequestServices.GetService<IOptionsMonitor<JupiterSettings>>();
 
-			bool isInternalPort = context.HttpContext.Connection.LocalPort == settings!.CurrentValue.InternalApiPort || 
+			bool isInternalPort = settings!.CurrentValue.InternalApiPorts.Contains(context.HttpContext.Connection.LocalPort) || 
 				/* unit tests do not run on ports, we consider them always on the internal port */ (context.HttpContext.Connection.LocalPort == 0 && context.HttpContext.Connection.LocalIpAddress == null);
 			if (!isInternalPort)
 			{
