Horde: Jira integration as an external issue service and a number of other cleanups
#rnx

diff --git a/Engine/Source/Programs/Horde/Horde.Build/Api/Dashboard.cs b/Engine/Source/Programs/Horde/Horde.Build/Api/Dashboard.cs
new file mode 100644
--- /dev/null
+++ b/Engine/Source/Programs/Horde/Horde.Build/Api/Dashboard.cs
@@ -0,0 +1,34 @@
+// Copyright Epic Games, Inc. All Rights Reserved.
+
+namespace Horde.Build.Api
+{
+	/// <summary>
+	/// Setting information required by dashboard
+	/// </summary>
+	public class GetDashboardConfigResponse
+	{
+		/// <summary>
+		/// The name of the external issue service
+		/// </summary>
+		public string? ExternalIssueServiceName { get; set; }
+
+		/// <summary>
+		/// The url of the external issue service
+		/// </summary>
+		public string? ExternalIssueServiceUrl { get; set; }
+
+		/// <summary>
+		/// The url of the perforce swarm installation
+		/// </summary>
+		public string? PerforceSwarmUrl { get; set; }
+
+		/// <summary>
+		/// Response constructor
+		/// </summary>
+		public GetDashboardConfigResponse()
+		{
+		}
+
+	}
+}
+
diff --git a/Engine/Source/Programs/Horde/Horde.Build/Api/Issues.cs b/Engine/Source/Programs/Horde/Horde.Build/Api/Issues.cs
--- a/Engine/Source/Programs/Horde/Horde.Build/Api/Issues.cs
+++ b/Engine/Source/Programs/Horde/Horde.Build/Api/Issues.cs
@@ -513,6 +513,11 @@ namespace Horde.Build.Api
 		/// </summary>
 		public bool ShowDesktopAlerts { get; set; }
 
+		/// <summary>
+		/// Key for this issue in external issue tracker
+		/// </summary>
+		public string? ExternalIssueKey { get; set; }
+
 		/// <summary>
 		/// Constructs a new issue
 		/// </summary>
@@ -570,6 +575,7 @@ namespace Horde.Build.Api
 			PrimarySuspectIds= details.SuspectUsers.Select(x => x.Id.ToString()).ToList();
 			PrimarySuspectsInfo = details.SuspectUsers.ConvertAll(x => new GetThinUserInfoResponse(x));
 			ShowDesktopAlerts = showDesktopAlerts;
+			ExternalIssueKey = details.ExternalIssueKey;
 		}
 	}
 
@@ -725,6 +731,11 @@ namespace Horde.Build.Api
 		/// </summary>
 		public List<FindIssueSpanResponse> Spans { get; set; }
 
+		/// <summary>
+		/// Key for this issue in external issue tracker
+		/// </summary>
+		public string? ExternalIssueKey { get; set; }
+
 		/// <summary>
 		/// Constructs a new issue
 		/// </summary>
@@ -762,6 +773,7 @@ namespace Horde.Build.Api
 			VerifiedAt = issue.VerifiedAt;
 			LastSeenAt = issue.LastSeenAt;
 			Spans = spans;
+			ExternalIssueKey = issue.ExternalIssueKey;
 		}
 	}
 
@@ -824,5 +836,200 @@ namespace Horde.Build.Api
 		/// List of spans to remove from this issue
 		/// </summary>
 		public List<string>? RemoveSpans { get; set; }
+
+		/// <summary>
+		/// A key to issue in external tracker
+		/// </summary>
+		public string? ExternalIssueKey { get; set; }
+
 	}
+
+	
+	/// <summary>
+	/// External issue project information
+	/// </summary>
+	public class GetExternalIssueProjectResponse
+	{
+		/// <summary>
+		/// The project key
+		/// </summary>
+		public string ProjectKey { get; set; }
+
+		/// <summary>
+		/// The name of the project
+		/// </summary>
+		public string Name { get; set; }
+
+		/// <summary>
+		/// The id of the project
+		/// </summary>
+		public string Id { get; set; }
+
+		/// <summary>
+		/// component id => name
+		/// </summary>
+		public Dictionary<string, string> Components { get; set; }
+
+		/// <summary>
+		/// IssueType id => name
+		/// </summary>
+		public Dictionary<string, string> IssueTypes { get; set; }
+
+		/// <summary>
+		/// Constructor
+		/// </summary>
+		/// <param name="key"></param>
+		/// <param name="name"></param>
+		/// <param name="id"></param>
+		/// <param name="components"></param>
+		/// <param name="issueTypes"></param>
+		public GetExternalIssueProjectResponse(string key, string name, string id, Dictionary<string, string> components, Dictionary<string, string> issueTypes)
+		{
+			ProjectKey = key;
+			Name = name; 
+			Id = id; 
+			Components = new Dictionary<string, string>(components); 
+			IssueTypes = new Dictionary<string, string>(issueTypes);
+		}
+	}
+
+
+	/// <summary>
+	/// Request an issue to be created on external issue tracking system
+	/// </summary>
+	public class CreateExternalIssueRequest
+	{
+		/// <summary>
+		/// Horde issue which is linked to external issue
+		/// </summary>
+		public int IssueId { get; set; } = 0;
+
+		/// <summary>
+		/// StreamId of a stream with this issue
+		/// </summary>
+		public string StreamId { get; set; } = String.Empty;
+
+		/// <summary>
+		/// Summary text for external issue
+		/// </summary>
+		public string Summary { get; set; } = String.Empty;
+
+		/// <summary>
+		/// External issue project id
+		/// </summary>
+		public string ProjectId { get; set; } = String.Empty;
+
+		/// <summary>
+		/// External issue component id
+		/// </summary>
+		public string ComponentId { get; set; } = String.Empty;
+
+		/// <summary>
+		/// External issue type id
+		/// </summary>
+		public string IssueTypeId { get; set; } = String.Empty;
+
+		/// <summary>
+		/// Optional description text for external issue
+		/// </summary>
+		public string? Description { get; set; }
+
+		/// <summary>
+		/// Optional link to issue on Horde
+		/// </summary>
+		public string? HordeIssueLink { get; set; }
+
+	}
+
+	/// <summary>
+	/// Response for externally created issue
+	/// </summary>
+	public class CreateExternalIssueResponse
+	{
+		/// <summary>
+		/// External issue key 
+		/// </summary>
+		public string Key { get; set; }
+
+		/// <summary>
+		/// Link to issue on external tracking site
+		/// </summary>
+		public string? Link { get; set; }
+
+		/// <summary>
+		/// Constructor
+		/// </summary>
+		/// <param name="key"></param>
+		/// <param name="link"></param>
+		public CreateExternalIssueResponse(string key, string? link)
+		{
+			Key = key;
+			Link = link;
+		}
+
+	}
+
+	/// <summary>
+	/// External issue response object
+	/// </summary>
+	public class GetExternalIssueResponse : IExternalIssue
+	{
+		/// <summary>
+		/// The external issue key
+		/// </summary>
+		public string Key { get; set; }
+
+		/// <summary>
+		/// The issue link on external tracking site
+		/// </summary>
+		public string? Link { get; set; }
+
+		/// <summary>
+		/// The issue status name, "To Do", "In Progress", etc
+		/// </summary>
+		public string? StatusName { get; set; }
+
+		/// <summary>
+		/// The issue resolution name, "Fixed", "Closed", etc
+		/// </summary>
+		public string? ResolutionName { get; set; }
+
+		/// <summary>
+		/// The issue priority name, "1 - Critical", "2 - Major", etc
+		/// </summary>
+		public string? PriorityName { get; set; }
+
+		/// <summary>
+		/// The current assignee's user name
+		/// </summary>
+		public string? AssigneeName { get; set; }
+
+		/// <summary>
+		/// The current assignee's display name
+		/// </summary>
+		public string? AssigneeDisplayName { get; set; }
+
+		/// <summary>
+		/// The current assignee's email address
+		/// </summary>
+		public string? AssigneeEmailAddress { get; set; }
+
+		/// <summary>
+		/// Response constructor
+		/// </summary>
+		public GetExternalIssueResponse(IExternalIssue issue)
+		{
+			Key = issue.Key;
+			Link = issue.Link;
+			StatusName = issue.StatusName;
+			ResolutionName = issue.ResolutionName;
+			PriorityName = issue.PriorityName;
+			AssigneeName = issue.AssigneeName;
+			AssigneeDisplayName = issue.AssigneeDisplayName;
+			AssigneeEmailAddress = issue.AssigneeEmailAddress;
+		}
+
+
+	}
+
 }
diff --git a/Engine/Source/Programs/Horde/Horde.Build/Api/Jira.cs b/Engine/Source/Programs/Horde/Horde.Build/Api/Jira.cs
deleted file mode 100644
diff --git a/Engine/Source/Programs/Horde/Horde.Build/Api/Streams.cs b/Engine/Source/Programs/Horde/Horde.Build/Api/Streams.cs
--- a/Engine/Source/Programs/Horde/Horde.Build/Api/Streams.cs
+++ b/Engine/Source/Programs/Horde/Horde.Build/Api/Streams.cs
@@ -435,6 +435,46 @@ namespace Horde.Build.Api
 		}
 	}
 
+	/// <summary>
+	/// State information for a step in the stream
+	/// </summary>
+	public class GetTemplateStepStateResponse
+	{
+		/// <summary>
+		/// Name of the step
+		/// </summary>
+		public string Name { get; set; } = String.Empty;
+
+		/// <summary>
+		/// User who paused the step
+		/// </summary>
+		public string? PausedByUserId { get; set; }
+
+		/// <summary>
+		/// User who quarantined the step
+		/// </summary>
+		public string? QuarantinedByUserId { get; set; }
+
+		/// <summary>
+		/// Default constructor for serialization
+		/// </summary>
+		private GetTemplateStepStateResponse()
+		{
+		}
+
+		/// <summary>
+		/// Constructor
+		/// </summary>
+		public GetTemplateStepStateResponse(TemplateStepState state)
+		{
+			Name = state.Name;
+			PausedByUserId = state.PausedByUserId?.ToString();
+			QuarantinedByUserId = state.QuarantinedByUserId?.ToString();
+		}
+
+	}
+
+
 	/// <summary>
 	/// Information about a template in this stream
 	/// </summary>
@@ -485,6 +525,11 @@ namespace Horde.Build.Api
 		/// </summary>
 		public List<GetChainedJobTemplateResponse>? ChainedJobs { get; set; }
 
+		/// <summary>
+		/// List of step states
+		/// </summary>
+		public List<GetTemplateStepStateResponse>? StepStates { get; set; }
+
 		/// <summary>
 		/// ACL for this template
 		/// </summary>
@@ -508,6 +553,7 @@ namespace Horde.Build.Api
 			NotificationChannelFilter = templateRef.NotificationChannelFilter;
 			Schedule = (templateRef.Schedule != null) ? new GetScheduleResponse(templateRef.Schedule) : null;
 			ChainedJobs = (templateRef.ChainedJobs != null && templateRef.ChainedJobs.Count > 0) ? templateRef.ChainedJobs.ConvertAll(x => new GetChainedJobTemplateResponse(x)) : null;
+			StepStates = (templateRef.StepStates != null && templateRef.StepStates.Count > 0) ? templateRef.StepStates.ConvertAll(x => new GetTemplateStepStateResponse(x)) : null;
 			Acl = (bIncludeAcl && templateRef.Acl != null)? new GetAclResponse(templateRef.Acl) : null;
 		}
 	}
diff --git a/Engine/Source/Programs/Horde/Horde.Build/Controllers/DashboardController.cs b/Engine/Source/Programs/Horde/Horde.Build/Controllers/DashboardController.cs
--- a/Engine/Source/Programs/Horde/Horde.Build/Controllers/DashboardController.cs
+++ b/Engine/Source/Programs/Horde/Horde.Build/Controllers/DashboardController.cs
@@ -5,6 +5,7 @@
 using System;
 using System.Text;
 using System.Threading.Tasks;
+using Horde.Build.Api;
 using Horde.Build.Controllers;
 using Microsoft.AspNetCore.Authentication;
 using Microsoft.AspNetCore.Authentication.Cookies;
@@ -19,20 +20,27 @@ namespace Horde.Build
 	/// </summary>	
 	[ApiController]
 	[Route("[controller]")]
-	public class DashboardChallengeController : Controller
+	public class DashboardController : Controller
 	{
 		/// <summary>
 		/// Authentication scheme in use
 		/// </summary>
 		readonly string _authenticationScheme;
 
+		/// <summary>
+		/// Server settings
+		/// </summary>
+		private readonly ServerSettings _settings;
+
+
 		/// <summary>
 		/// Constructor
 		/// </summary>
 		/// <param name="serverSettings">Server settings</param>
-		public DashboardChallengeController(IOptionsMonitor<ServerSettings> serverSettings)
+		public DashboardController(IOptionsMonitor<ServerSettings> serverSettings)
 		{
 			_authenticationScheme = AccountController.GetAuthScheme(serverSettings.CurrentValue.AuthMethod);
+			_settings = serverSettings.CurrentValue;
 		}
 
 		/// <summary>	
@@ -98,5 +106,31 @@ namespace Horde.Build
 
 			return Ok();
 		}
+
+		/// <summary>
+		/// Query all the projects
+		/// </summary>
+		/// <returns>Config information needed by the dashboard</returns>
+		[HttpGet]
+		[Authorize]
+		[Route("/api/v1/dashboard/config")]
+		public ActionResult<GetDashboardConfigResponse> GetDashbordConfig()
+		{
+			GetDashboardConfigResponse dashboardConfigResponse = new GetDashboardConfigResponse();
+
+			if (_settings.JiraUrl != null)
+			{
+				dashboardConfigResponse.ExternalIssueServiceName = "Jira";
+				dashboardConfigResponse.ExternalIssueServiceUrl = _settings.JiraUrl.ToString().TrimEnd('/');
+			}
+
+			if (_settings.P4SwarmUrl != null)
+			{
+				dashboardConfigResponse.PerforceSwarmUrl = _settings.P4SwarmUrl.ToString().TrimEnd('/');
+			}
+
+			return dashboardConfigResponse;
+		}
+
 	}
 }
diff --git a/Engine/Source/Programs/Horde/Horde.Build/Controllers/IssuesController.cs b/Engine/Source/Programs/Horde/Horde.Build/Controllers/IssuesController.cs
--- a/Engine/Source/Programs/Horde/Horde.Build/Controllers/IssuesController.cs
+++ b/Engine/Source/Programs/Horde/Horde.Build/Controllers/IssuesController.cs
@@ -7,7 +7,6 @@ using System.Threading.Tasks;
 using Horde.Build.Acls;
 using Horde.Build.Api;
 using Horde.Build.Collections;
-using Horde.Build.Config;
 using Horde.Build.Models;
 using Horde.Build.Server;
 using Horde.Build.Services;
@@ -36,6 +35,7 @@ namespace Horde.Build.Controllers
 	{
 		private readonly IIssueCollection _issueCollection;
 		private readonly IIssueService _issueService;
+		private readonly IExternalIssueService _externalIssueService;
 		private readonly JobService _jobService;
 		private readonly StreamService _streamService;
 		private readonly IUserCollection _userCollection;
@@ -45,7 +45,7 @@ namespace Horde.Build.Controllers
 		/// <summary>
 		/// Constructor
 		/// </summary>
-		public IssuesController(ILogger<IssuesController> logger, IIssueCollection issueCollection, IIssueService issueService, JobService jobService, StreamService streamService, IUserCollection userCollection, ILogFileService logFileService)
+		public IssuesController(ILogger<IssuesController> logger, IIssueCollection issueCollection, IIssueService issueService, JobService jobService, StreamService streamService, IUserCollection userCollection, ILogFileService logFileService, IExternalIssueService externalIssueService)
 		{			
 			_issueCollection = issueCollection;
 			_issueService = issueService;
@@ -53,6 +53,7 @@ namespace Horde.Build.Controllers
 			_streamService = streamService;
 			_userCollection = userCollection;
 			_logFileService = logFileService;
+			_externalIssueService = externalIssueService;
 			_logger = logger;
 		}
 
@@ -566,11 +567,125 @@ namespace Horde.Build.Controllers
 				removeSpans = request.RemoveSpans.ConvertAll(x => ObjectId.Parse(x));
 			}
 
-			if (!await _issueService.UpdateIssueAsync(issueId, request.Summary, request.Description, request.Promoted, newOwnerId, newNominatedById, request.Acknowledged, newDeclinedById, request.FixChange, newResolvedById, addSpans, removeSpans))
+			if (!await _issueService.UpdateIssueAsync(issueId, request.Summary, request.Description, request.Promoted, newOwnerId, newNominatedById, request.Acknowledged, newDeclinedById, request.FixChange, newResolvedById, addSpans, removeSpans, request.ExternalIssueKey))
 			{
 				return NotFound();
 			}
 			return Ok();
 		}
+
+		// External issue tracking
+
+		/// <summary>
+		/// Get external issue information for provided keys
+		/// </summary>
+		[HttpGet]
+		[Authorize]
+		[Route("/api/v1/issues/external")]
+		public async Task<ActionResult<List<GetExternalIssueResponse>>> GetExternalIssuesAsync([FromQuery] string streamId, [FromQuery] string[] keys)
+		{
+			if (_externalIssueService == null)
+			{
+				return NotFound();
+			}
+
+			StreamId streamIdValue = new StreamId(streamId);
+			StreamPermissionsCache cache = new StreamPermissionsCache();
+
+			if (!await _streamService.AuthorizeAsync(streamIdValue, AclAction.ViewStream, User, cache))
+			{
+				return Forbid();
+			}
+
+			List<GetExternalIssueResponse> response = new List<GetExternalIssueResponse>();
+
+			if (keys.Length != 0)
+			{
+				List<IExternalIssue> issues = await _externalIssueService.GetIssuesAsync(keys);
+
+				for (int i = 0; i < issues.Count; i++)
+				{
+					response.Add(new GetExternalIssueResponse(issues[i]));
+				}
+			}
+
+			return response;
+		}
+
+		/// <summary>
+		/// Create a new external issue
+		/// </summary>
+		[HttpPost]
+		[Authorize]
+		[Route("/api/v1/issues/external")]
+		public async Task<ActionResult<CreateExternalIssueResponse>> CreateExternalIssueAsync([FromBody] CreateExternalIssueRequest issueRequest)
+		{
+			if (_externalIssueService == null)
+			{
+				return NotFound();
+			}
+
+			StreamId streamIdValue = new StreamId(issueRequest.StreamId);
+			StreamPermissionsCache cache = new StreamPermissionsCache();
+
+			if (!await _streamService.AuthorizeAsync(streamIdValue, AclAction.ViewStream, User, cache))
+			{
+				return Forbid();
+			}
+
+			IUser? user = await _userCollection.GetUserAsync(User);
+
+			if (user == null)
+			{
+				return BadRequest($"Missing user for {User.GetUserName()}");
+			}
+
+			(string? key, string? url) = await _externalIssueService.CreateIssueAsync(user, issueRequest.IssueId, issueRequest.Summary, issueRequest.ProjectId, issueRequest.ComponentId, issueRequest.IssueTypeId, issueRequest.Description, issueRequest.HordeIssueLink);
+
+			if (key == null)
+			{
+				throw new Exception($"Unable to create external issue");
+			}
+
+			return new CreateExternalIssueResponse(key, url);
+		}
+
+		/// <summary>
+		/// Gets external issue tracking projects associated with the provided stream
+		/// </summary>
+		/// <param name="streamId"></param>
+		/// <returns></returns>
+		[HttpGet]
+		[Authorize]
+		[Route("/api/v1/issues/external/projects")]
+		public async Task<ActionResult<List<GetExternalIssueProjectResponse>>> GetExternalIssueProjectsAsync([FromQuery] string streamId)
+		{
+			if (_externalIssueService == null)
+			{
+				return NotFound();
+			}
+
+			StreamId streamIdValue = new StreamId(streamId);
+			StreamPermissionsCache cache = new StreamPermissionsCache();
+
+			if (!await _streamService.AuthorizeAsync(streamIdValue, AclAction.ViewStream, User, cache))
+			{
+				return Forbid();
+			}
+
+			IStream? stream = await _streamService.GetStreamAsync(streamIdValue);
+
+			List<IExternalIssueProject> projects = await _externalIssueService.GetProjects(stream!);
+			List<GetExternalIssueProjectResponse> response = new List<GetExternalIssueProjectResponse>();
+
+			projects.ForEach(project =>
+			{
+				response.Add(new GetExternalIssueProjectResponse(project.Key, project.Name, project.Id, project.Components, project.IssueTypes));
+			});
+
+			return response;
+		}
+
+
 	}
 }
diff --git a/Engine/Source/Programs/Horde/Horde.Build/Controllers/JiraController.cs b/Engine/Source/Programs/Horde/Horde.Build/Controllers/JiraController.cs
deleted file mode 100644
diff --git a/Engine/Source/Programs/Horde/Horde.Build/Issues/External/IExternalIssue.cs b/Engine/Source/Programs/Horde/Horde.Build/Issues/External/IExternalIssue.cs
new file mode 100644
--- /dev/null
+++ b/Engine/Source/Programs/Horde/Horde.Build/Issues/External/IExternalIssue.cs
@@ -0,0 +1,66 @@
+// Copyright Epic Games, Inc. All Rights Reserved.	
+
+using System;
+using System.Collections.Concurrent;
+using System.Collections.Generic;
+using System.Net.Http;
+using System.Net.Http.Headers;
+using System.Text;
+using System.Text.Json;
+using System.Threading.Tasks;
+using Microsoft.Extensions.Logging;
+using Microsoft.Extensions.Options;
+using Polly;
+using Polly.Extensions.Http;
+
+namespace Horde.Build.Models
+{
+	/// <summary>
+	/// 
+	/// </summary>
+	public interface IExternalIssue
+	{
+		/// <summary>
+		/// The issue key or tracking id
+		/// </summary>
+		public string Key { get; set; }
+
+		/// <summary>
+		/// The issue link on external tracking site
+		/// </summary>
+		public string? Link { get; set; }
+
+		/// <summary>
+		/// The issue status name, "To Do", "In Progress", etc
+		/// </summary>
+		public string? StatusName { get; set; }
+
+		/// <summary>
+		/// The issue resolution name, "Fixed", "Closed", etc
+		/// </summary>
+		public string? ResolutionName { get; set; }
+
+		/// <summary>
+		/// The issue priority name, "1 - Critical", "2 - Major", etc
+		/// </summary>
+		public string? PriorityName { get; set; }
+
+		/// <summary>
+		/// The current assignee's user name
+		/// </summary>
+		public string? AssigneeName { get; set; }
+
+		/// <summary>
+		/// The current assignee's display name
+		/// </summary>
+		public string? AssigneeDisplayName { get; set; }
+
+		/// <summary>
+		/// The current assignee's email address
+		/// </summary>
+		public string? AssigneeEmailAddress { get; set; }
+
+	}
+
+}
+
diff --git a/Engine/Source/Programs/Horde/Horde.Build/Issues/External/IExternalIssueService.cs b/Engine/Source/Programs/Horde/Horde.Build/Issues/External/IExternalIssueService.cs
new file mode 100644
--- /dev/null
+++ b/Engine/Source/Programs/Horde/Horde.Build/Issues/External/IExternalIssueService.cs
@@ -0,0 +1,114 @@
+// Copyright Epic Games, Inc. All Rights Reserved.	
+
+using System;
+using System.Collections.Generic;
+using System.Threading;
+using System.Threading.Tasks;
+using Horde.Build.Models;
+using Horde.Build.Utilities;
+using Microsoft.Extensions.Hosting;
+
+namespace Horde.Build.Services
+{
+	using StreamId = StringId<IStream>;
+
+	/// <summary>
+	/// External issue tracking project
+	/// </summary>
+	public interface IExternalIssueProject
+	{
+		/// <summary>
+		/// Project key
+		/// </summary>
+		string Key { get; }
+
+		/// <summary>
+		/// Project name
+		/// </summary>
+		string Name { get; }
+
+		/// <summary>
+		/// Project id
+		/// </summary>
+		string Id { get; }
+
+		/// <summary>
+		/// component id => name
+		/// </summary>
+		Dictionary<string, string> Components { get; }
+
+		/// <summary>
+		/// // IssueType id => name
+		/// </summary>
+		Dictionary<string, string> IssueTypes { get; }
+	}
+
+	/// <summary>
+	/// Interface for an external issue tracking service
+	/// </summary>
+	public interface IExternalIssueService : IHostedService
+	{
+		/// <summary>
+		/// Get issues associated with provided keys
+		/// </summary>
+		/// <param name="keys"></param>
+		/// <returns></returns>
+		Task<List<IExternalIssue>> GetIssuesAsync(string[] keys);
+
+		/// <summary>
+		/// Create and link an external issue
+		/// </summary>		
+		/// <param name="User"></param>
+		/// <param name="issueId"></param>
+		/// <param name="summary"></param>
+		/// <param name="projectId"></param>
+		/// <param name="componentId"></param>
+		/// <param name="issueType"></param>
+		/// <param name="description"></param>
+		/// <param name="hordeIssueLink"></param>
+		Task<(string? key, string? url)> CreateIssueAsync(IUser User, int issueId, string summary, string projectId, string componentId, string issueType, string? description, string? hordeIssueLink);
+
+
+		/// <summary>
+		///  Get projects for provided keys
+		/// </summary>
+		/// <param name="stream"></param>
+		/// <returns></returns>;
+		Task<List<IExternalIssueProject>> GetProjects(IStream stream);
+
+	}
+
+	/// <summary>
+	/// Default external issue service
+	/// </summary>
+	public class DefaultExternalIssueService : IExternalIssueService
+	{
+		/// <inheritdoc/>
+		public Task<List<IExternalIssue>> GetIssuesAsync(string[] keys)
+		{
+			return Task.FromResult(new List<IExternalIssue>());
+		}
+
+		/// <inheritdoc/>
+		public Task<(string? key, string? url)> CreateIssueAsync(IUser User, int issueId, string summary, string projectId, string componentId, string issueType, string? description, string? hordeIssueLink)
+		{
+			throw new NotImplementedException();
+		}
+
+		/// <inheritdoc/>
+		public Task<List<IExternalIssueProject>> GetProjects(IStream stream)
+		{
+			return Task.FromResult(new List<IExternalIssueProject>());
+		}
+
+		/// <inheritdoc/>
+		public void Dispose() { }
+
+		/// <inheritdoc/>
+		public Task StartAsync(CancellationToken token) => Task.CompletedTask;
+
+		/// <inheritdoc/>
+		public Task StopAsync(CancellationToken token) => Task.CompletedTask;
+	}
+
+}
\ No newline at end of file
diff --git a/Engine/Source/Programs/Horde/Horde.Build/Issues/External/Impl/JiraService.cs b/Engine/Source/Programs/Horde/Horde.Build/Issues/External/Impl/JiraService.cs
new file mode 100644
--- /dev/null
+++ b/Engine/Source/Programs/Horde/Horde.Build/Issues/External/Impl/JiraService.cs
@@ -0,0 +1,570 @@
+// Copyright Epic Games, Inc. All Rights Reserved.	
+
+using System;
+using System.Collections.Concurrent;
+using System.Collections.Generic;
+using System.Net.Http;
+using System.Net.Http.Headers;
+using System.Text;
+using System.Text.Json;
+using System.Threading.Tasks;
+using Microsoft.Extensions.Logging;
+using Microsoft.Extensions.Options;
+using Horde.Build.Models;
+using Polly;
+using Polly.Extensions.Http;
+using Horde.Build.Utilities;
+using Horde.Build.Config;
+using Microsoft.Extensions.Hosting;
+using HordeCommon;
+using System.Threading;
+using System.Linq;
+using Horde.Build.Server;
+
+namespace Horde.Build.Services
+{	
+	/// <summary>
+	/// Jira issue information
+	/// </summary>
+	internal class JiraIssue : IExternalIssue
+	{
+		/// <summary>
+		/// The Jira issue key
+		/// </summary>
+		public string Key { get; set; } = String.Empty;
+
+		/// <summary>
+		/// The Jira issue link
+		/// </summary>
+		public string? Link { get; set; }
+
+		/// <summary>
+		/// The issue status name, "To Do", "In Progress", etc
+		/// </summary>
+		public string? StatusName { get; set; }
+
+		/// <summary>
+		/// The issue resolution name, "Fixed", "Closed", etc
+		/// </summary>
+		public string? ResolutionName { get; set; }
+
+		/// <summary>
+		/// The issue priority name, "1 - Critical", "2 - Major", etc
+		/// </summary>
+		public string? PriorityName { get; set; }
+
+		/// <summary>
+		/// The current assignee's user name
+		/// </summary>
+		public string? AssigneeName { get; set; }
+
+		/// <summary>
+		/// The current assignee's display name
+		/// </summary>
+		public string? AssigneeDisplayName { get; set; }
+
+		/// <summary>
+		/// The current assignee's email address
+		/// </summary>
+		public string? AssigneeEmailAddress { get; set; }
+	}
+
+	internal class JiraCreateResponse
+	{
+		/// <summary>
+		/// Key of external issue
+		/// </summary>
+		public string Key { get; set; } = String.Empty;
+
+	}
+
+	internal class JiraProject : IExternalIssueProject
+	{
+		public string Key { get; set; } = String.Empty;
+
+		public string Name { get; set; } = String.Empty;
+
+		public string Id { get; set; } = String.Empty;
+
+		// component id => name
+		public Dictionary<string, string> Components { get; set; } = new Dictionary<string, string>();
+
+		// IssueType id => name
+		public Dictionary<string, string> IssueTypes { get; set; } = new Dictionary<string, string>();
+
+	}
+
+	/// <summary>
+	/// Jira service functionality
+	/// </summary>
+	internal sealed class JiraService : IExternalIssueService, IHostedService, IDisposable
+	{
+		readonly ILogger _logger;
+
+		/// <summary>
+		/// The server settings
+		/// </summary>
+		readonly ServerSettings _settings;
+		readonly HttpClient _client;
+		readonly AsyncPolicy<HttpResponseMessage> _retryPolicy;
+		readonly ConcurrentDictionary<string, JiraCacheValue> _issueCache = new ConcurrentDictionary<string, JiraCacheValue>();
+
+		/// <summary>
+		/// Singleton instance of the stream service
+		/// </summary>
+		readonly StreamService _streamService;
+
+		/// <summary>
+		/// Singleton instance of the project service
+		/// </summary>
+		readonly IIssueService _issueService;
+
+		/// <summary>
+		/// Config collection
+		/// </summary>
+		readonly ConfigCollection _configCollection;
+
+		readonly IClock _clock;
+		readonly ITicker _ticker;
+
+		readonly Uri _jiraUrl;
+
+		Dictionary<string, JiraProject> _cachedJiraProjects = new Dictionary<string, JiraProject>();
+
+
+		/// <summary>
+		/// Jira service constructor
+		/// </summary>
+		/// <param name="settings"></param>
+		/// <param name="streamService"></param>
+		/// <param name="issueService"></param>
+		/// <param name="configCollection"></param>
+		/// <param name="clock"></param>
+		/// <param name="logger"></param>
+		public JiraService(IOptions<ServerSettings> settings, StreamService streamService, IIssueService issueService, ConfigCollection configCollection, IClock clock, ILogger<JiraService> logger)
+		{
+
+			_settings = settings.Value;
+			_logger = logger;
+			_streamService = streamService;
+			_issueService = issueService;
+			_configCollection = configCollection;
+			_clock = clock;
+			_ticker = clock.AddTicker<JiraService>(TimeSpan.FromMinutes(2.0), TickAsync, logger);
+			_jiraUrl = _settings.JiraUrl!;
+
+			// setup http client for Jira rest api queries
+			_client = new HttpClient();
+			byte[] authBytes = Encoding.ASCII.GetBytes($"{_settings.JiraUsername}:{_settings.JiraApiToken}");
+			_client.DefaultRequestHeaders.Authorization = new AuthenticationHeaderValue("Basic", Convert.ToBase64String(authBytes));
+			_client.Timeout = TimeSpan.FromSeconds(15.0);
+			_retryPolicy = HttpPolicyExtensions.HandleTransientHttpError().WaitAndRetryAsync(3, attempt => TimeSpan.FromSeconds(Math.Pow(2.0, attempt)));
+		}
+
+		/// <inheritdoc/>
+		public Task StartAsync(CancellationToken cancellationToken) => _ticker.StartAsync();
+
+		/// <inheritdoc/>
+		public Task StopAsync(CancellationToken cancellationToken) => _ticker.StopAsync();
+
+
+		/// <inheritdoc/>
+		public void Dispose()
+		{
+			_ticker.Dispose();
+			_client.Dispose();
+		}
+
+		public async Task<List<IExternalIssueProject>> GetProjects(IStream stream)
+		{
+			HashSet<string> projectKeys = new HashSet<string>();
+			List<IExternalIssueProject> result = new List<IExternalIssueProject>();
+
+			StreamConfig? config = null;
+
+			try
+			{
+				config = await _configCollection.GetConfigAsync<StreamConfig>(stream.ConfigRevision);
+			}
+			catch (Exception ex)
+			{
+				_logger.LogError(ex, "Unable to get stream config for {StreamId}", stream.Id);
+
+			}
+
+			if (config == null || config.Workflows == null)
+			{
+				return result;
+			}
+
+			config.Workflows.ForEach(workflow =>
+			{
+				if (workflow.ExternalIssues != null)
+				{
+					projectKeys.Add(workflow.ExternalIssues.ProjectKey);
+				}				
+
+			});
+
+			foreach (string key in projectKeys)
+			{
+				if (_cachedJiraProjects.TryGetValue(key, out JiraProject? project))
+				{
+					result.Add(project);
+				}
+			}
+
+			return result;
+		}
+
+		async Task UpdateJiraProjects(string[] jiraProjectKeys)
+		{
+			if (jiraProjectKeys.Length == 0)
+			{
+				return;
+			}
+
+			for (int i = 0; i < jiraProjectKeys.Length; i++)
+			{
+				string projectKey = jiraProjectKeys[i];
+
+				Uri uri = new Uri(_jiraUrl, $"/rest/api/2/project/{projectKey}");
+
+				HttpResponseMessage response;
+
+				try
+				{
+					response = await _retryPolicy.ExecuteAsync(() => _client.GetAsync(uri));
+					response.EnsureSuccessStatusCode();
+				}
+				catch (Exception)
+				{
+					_logger.LogError("Unable to get project info for {ProjectKey", projectKey);
+					continue;
+				}
+
+				string responseBody = await response.Content.ReadAsStringAsync();
+
+				JsonElement jsonProject = JsonSerializer.Deserialize<JsonElement>(responseBody);
+
+				JiraProject jiraProject = new JiraProject() { Key = projectKey, Id = jsonProject.GetProperty("id")!.GetString()!, Name = jsonProject.GetProperty("name")!.GetString()! };
+
+				// Issue Types
+				if (jsonProject.TryGetProperty("issueTypes", out JsonElement issueTypes))
+				{
+					foreach (JsonElement issueType in issueTypes.EnumerateArray())
+					{
+						// filter out subtasks issue types
+						if (issueType.TryGetProperty("subtask", out JsonElement subtask))
+						{
+							if (subtask.GetBoolean())
+							{
+								continue;
+							}
+						}
+
+						// filter out certain issue types
+						string name = issueType.GetProperty("name").GetString()!;
+						if (name == "Epic")
+						{
+							continue;
+						}
+
+						jiraProject.IssueTypes[issueType.GetProperty("id").GetString()!] = name;
+
+					}
+
+				}
+
+				// Components
+				// We do get components from the projects endpoint, though these don't contain the archived property :/
+				try
+				{
+					response = await _retryPolicy.ExecuteAsync(() => _client.GetAsync(new Uri(_jiraUrl, $"/rest/api/2/project/{projectKey}/components")));
+					response.EnsureSuccessStatusCode();
+				}
+				catch (Exception)
+				{
+					_logger.LogError("Unable to get components info for {ProjectKey", projectKey);
+					continue;
+				}
+
+				responseBody = await response.Content.ReadAsStringAsync();
+
+				JsonElement jiraComponents = JsonSerializer.Deserialize<JsonElement>(responseBody);
+
+				for (int j = 0; j < jiraComponents.GetArrayLength(); j++)
+				{
+					JsonElement component = jiraComponents[j];
+
+					// skip archived components
+					if (component.TryGetProperty("archived", out JsonElement archived))
+					{
+						if (archived.GetBoolean())
+						{
+							continue;
+						}
+
+						if (component.TryGetStringProperty("id", out string? id))
+						{
+							if (component.TryGetStringProperty("name", out string? name))
+							{
+								jiraProject.Components[id] = name;
+							}
+						}
+					}
+				}
+
+				_cachedJiraProjects[projectKey] = jiraProject;
+			}
+		}
+
+		async ValueTask TickAsync(CancellationToken stoppingToken)
+		{
+
+			HashSet<string> jiraProjectKeys = new HashSet<string>();
+
+			List<IStream> cachedStreams = await _streamService.StreamCollection.FindAllAsync();
+			StreamConfig? config = null;
+			foreach (IStream cachedStream in cachedStreams)
+			{
+				config = null;
+				try
+				{
+					config = await _configCollection.GetConfigAsync<StreamConfig>(cachedStream.ConfigRevision);
+				}
+				catch (Exception ex)
+				{
+					_logger.LogError(ex, "Unable to get stream config for {StreamId}", cachedStream.Id);
+					continue;
+				}
+
+				foreach (WorkflowConfig workflow in config.Workflows)
+				{
+					if (workflow.ExternalIssues == null || String.IsNullOrEmpty(workflow.ExternalIssues.ProjectKey))
+					{
+						continue;
+					}
+
+					jiraProjectKeys.Add(workflow.ExternalIssues.ProjectKey);
+				}
+			}
+
+			// update projects
+			await UpdateJiraProjects(jiraProjectKeys.ToArray());
+
+			HashSet<string> jiraKeys = new HashSet<string>();
+
+			// Refresh issues
+			List<IIssue> openIssues = await _issueService.FindIssuesAsync(resolved: false);
+			for (int idx = 0; idx < openIssues.Count; idx++)
+			{
+				IIssue openIssue = openIssues[idx];
+				if (!String.IsNullOrEmpty(openIssue.ExternalIssueKey))
+				{
+					jiraKeys.Add(openIssue.ExternalIssueKey);
+				}
+			}
+
+			await GetIssuesAsync(jiraKeys.ToArray());
+
+		}
+
+		public async Task<(string? key, string? url)> CreateIssueAsync(IUser User, int issueId, string summary, string projectId, string componentId, string issueType, string? description, string? hordeIssueLink)
+		{
+
+			IIssue? issue = await _issueService.GetIssueAsync(issueId);
+			if (issue == null)
+			{
+				throw new Exception($"Issue not found: {issueId}");
+			}
+
+			string issueDesc = description ?? "";
+
+			if (hordeIssueLink != null)
+			{
+				issueDesc = $"\n{issueDesc}\n\nJira issue created by {User.Name} for *[Horde Issue {issueId}|{hordeIssueLink}]*\n\n";
+			}
+			else
+			{
+				issueDesc = $"\n{issueDesc}\n\nJira Issue created by {User.Name} from Horde\n\n";
+			}
+
+
+			Uri uri = new Uri(_jiraUrl, $"/rest/api/2/issue");
+
+			string bodyJson = JsonSerializer.Serialize(new
+			{
+				fields = new
+				{
+					project = new
+					{
+						id = projectId
+					},
+					issuetype = new
+					{
+						id = issueType
+					},
+					summary = summary,
+					description = issueDesc,
+					components = new object[]
+					{
+						new
+						{
+							id = componentId
+						}
+					}
+				}
+			});
+
+			HttpResponseMessage response = await _retryPolicy.ExecuteAsync(() => _client.PostAsync(uri, new StringContent(bodyJson, Encoding.UTF8, "application/json")));
+
+			response.EnsureSuccessStatusCode();
+
+			string responseBody = await response.Content.ReadAsStringAsync();
+
+			JiraCreateResponse? jiraResponse = JsonSerializer.Deserialize<JiraCreateResponse>(responseBody, new JsonSerializerOptions { PropertyNameCaseInsensitive = true });
+
+			if (jiraResponse == null || String.IsNullOrEmpty(jiraResponse.Key))
+			{
+				throw new Exception($"Unable to parse returned jira json: {responseBody}");
+			}
+
+			await _issueService.UpdateIssueAsync(issueId, externalIssueKey: jiraResponse.Key);
+
+			return (jiraResponse.Key, $"{_jiraUrl}browse/{jiraResponse.Key}");
+		}
+
+		/// <inheritdoc/>
+		public async Task<List<IExternalIssue>> GetIssuesAsync(string[] jiraKeys)
+		{
+			List<IExternalIssue> result = new List<IExternalIssue>();
+
+			if (jiraKeys.Length == 0)
+			{
+				return result;
+			}
+
+			List<string> queryJiras = new List<string>();
+			for (int i = 0; i < jiraKeys.Length; i++)
+			{
+				JiraCacheValue value;
+				if (_issueCache.TryGetValue(jiraKeys[i], out value))
+				{
+					if (DateTime.UtcNow.Subtract(value._cacheTime).TotalMinutes >= 2)
+					{
+						queryJiras.Add(jiraKeys[i]);
+					}
+				}
+				else
+				{
+					queryJiras.Add(jiraKeys[i]);
+				}
+			}
+
+			if (queryJiras.Count > 0)
+			{
+				Uri uri = new Uri(_jiraUrl, $"/rest/api/2/search?jql=issueKey%20in%20({String.Join(",", queryJiras)})&fields=assignee,status,resolution,priority&maxResults={queryJiras.Count}");
+
+				HttpResponseMessage response = await _retryPolicy.ExecuteAsync(() => _client.GetAsync(uri));
+				if (!response.IsSuccessStatusCode)
+				{
+					_logger.LogError("GET to {Uri} returned {Code} ({Response})", uri, response.StatusCode, await response.Content.ReadAsStringAsync());
+					return result;
+				}
+
+				byte[] data = await response.Content.ReadAsByteArrayAsync();
+				IssueQueryResponse? jiras = JsonSerializer.Deserialize<IssueQueryResponse>(data.AsSpan(), new JsonSerializerOptions() { PropertyNameCaseInsensitive = true });
+
+				if (jiras != null)
+				{
+					for (int i = 0; i < jiras.Issues.Count; i++)
+					{
+						IssueResponse issue = jiras.Issues[i];
+
+						JiraIssue jiraIssue = new JiraIssue();
+						jiraIssue.Key = issue.Key;
+						jiraIssue.Link = $"{_jiraUrl}browse/{issue.Key}";
+						jiraIssue.AssigneeName = issue.Fields?.Assignee?.Name;
+						jiraIssue.AssigneeDisplayName = issue.Fields?.Assignee?.DisplayName;
+						jiraIssue.AssigneeEmailAddress = issue.Fields?.Assignee?.EmailAddress;
+						jiraIssue.PriorityName = issue.Fields?.Priority?.Name;
+						jiraIssue.ResolutionName = issue.Fields?.Resolution?.Name;
+						jiraIssue.StatusName = issue.Fields?.Status?.Name;
+
+						_issueCache[issue.Key] = new JiraCacheValue() { _issue = jiraIssue, _cacheTime = DateTime.UtcNow };
+					}
+				}
+			}
+
+			for (int i = 0; i < jiraKeys.Length; i++)
+			{
+				JiraCacheValue value;
+				if (_issueCache.TryGetValue(jiraKeys[i], out value))
+				{
+					result.Add(value._issue);
+				}
+			}
+
+			return result;
+
+		}
+
+		// Jira rest api mapping and caching
+
+		struct JiraCacheValue
+		{
+			public JiraIssue _issue;
+			public DateTime _cacheTime;
+		};
+
+		class IssueQueryResponse
+		{
+			public int Total { get; set; } = 0;
+
+			public List<IssueResponse> Issues { get; set; } = new List<IssueResponse>();
+		}
+
+		class IssueResponse
+		{
+			public string Key { get; set; } = String.Empty;
+
+			public IssueFields? Fields { get; set; }
+		}
+
+		class IssueFields
+		{
+			public AssigneeField? Assignee { get; set; }
+
+			public StatusField? Status { get; set; }
+
+			public ResolutionField? Resolution { get; set; }
+
+			public PriorityField? Priority { get; set; }
+		}
+
+		class AssigneeField
+		{
+			public string? Name { get; set; }
+
+			public string? DisplayName { get; set; }
+
+			public string? EmailAddress { get; set; }
+		}
+
+		class StatusField
+		{
+			public string? Name { get; set; }
+		}
+
+		class PriorityField
+		{
+			public string? Name { get; set; }
+		}
+
+		class ResolutionField
+		{
+			public string? Name { get; set; }
+		}
+	}
+}
diff --git a/Engine/Source/Programs/Horde/Horde.Build/Issues/IIssue.cs b/Engine/Source/Programs/Horde/Horde.Build/Issues/IIssue.cs
--- a/Engine/Source/Programs/Horde/Horde.Build/Issues/IIssue.cs
+++ b/Engine/Source/Programs/Horde/Horde.Build/Issues/IIssue.cs
@@ -146,6 +146,12 @@ namespace Horde.Build.Models
 		/// Update index for this instance
 		/// </summary>
 		public int UpdateIndex { get; }
+
+		/// <summary>
+		/// External issue linked to this issue
+		/// </summary>
+		public string? ExternalIssueKey { get; }
+
 	}
 
 	/// <summary>
diff --git a/Engine/Source/Programs/Horde/Horde.Build/Issues/IIssueCollection.cs b/Engine/Source/Programs/Horde/Horde.Build/Issues/IIssueCollection.cs
--- a/Engine/Source/Programs/Horde/Horde.Build/Issues/IIssueCollection.cs
+++ b/Engine/Source/Programs/Horde/Horde.Build/Issues/IIssueCollection.cs
@@ -438,8 +438,9 @@ namespace Horde.Build.Collections
 		/// <param name="newResolvedById">User that resolved the issue (may be ObjectId.Empty to clear)</param>
 		/// <param name="newExcludeSpanIds">List of span ids to exclude from this issue</param>
 		/// <param name="newLastSeenAt"></param>
+		/// <param name="externalIssueKey">Issue key for external issue tracking</param>
 		/// <returns>True if the issue was updated</returns>
-		Task<IIssue?> TryUpdateIssueAsync(IIssue issue, IssueSeverity? newSeverity = null, string? newSummary = null, string? newUserSummary = null, string? newDescription = null, bool? newPromoted = null, UserId? newOwnerId = null, UserId? newNominatedById = null, bool? newAcknowledged = null, UserId? newDeclinedById = null, int? newFixChange = null, UserId? newResolvedById = null, List<ObjectId>? newExcludeSpanIds = null, DateTime? newLastSeenAt = null);
+		Task<IIssue?> TryUpdateIssueAsync(IIssue issue, IssueSeverity? newSeverity = null, string? newSummary = null, string? newUserSummary = null, string? newDescription = null, bool? newPromoted = null, UserId? newOwnerId = null, UserId? newNominatedById = null, bool? newAcknowledged = null, UserId? newDeclinedById = null, int? newFixChange = null, UserId? newResolvedById = null, List<ObjectId>? newExcludeSpanIds = null, DateTime? newLastSeenAt = null, string? externalIssueKey = null);
 
 		/// <summary>
 		/// Updates derived data for an issue (ie. data computed from the spans attached to it). Also clears the issue's 'modified' state.
diff --git a/Engine/Source/Programs/Horde/Horde.Build/Issues/IIssueDetails.cs b/Engine/Source/Programs/Horde/Horde.Build/Issues/IIssueDetails.cs
--- a/Engine/Source/Programs/Horde/Horde.Build/Issues/IIssueDetails.cs
+++ b/Engine/Source/Programs/Horde/Horde.Build/Issues/IIssueDetails.cs
@@ -54,6 +54,11 @@ namespace Horde.Build.Models
 		/// </summary>
 		IReadOnlyList<IUser> SuspectUsers { get; }
 
+		/// <summary>
+		/// Issue key in external issue tracker
+		/// </summary>
+		string? ExternalIssueKey { get; }
+
 		/// <summary>
 		/// Determines whether the given user should be notified about the given issue
 		/// </summary>
diff --git a/Engine/Source/Programs/Horde/Horde.Build/Issues/IIssueService.cs b/Engine/Source/Programs/Horde/Horde.Build/Issues/IIssueService.cs
--- a/Engine/Source/Programs/Horde/Horde.Build/Issues/IIssueService.cs
+++ b/Engine/Source/Programs/Horde/Horde.Build/Issues/IIssueService.cs
@@ -141,8 +141,9 @@ namespace Horde.Build.Services
 		/// <param name="resolvedById">Whether the issue has been resolved</param>
 		/// <param name="addSpanIds">Add spans to this issue</param>
 		/// <param name="removeSpanIds">Remove spans from this issue</param>
+		/// <param name="externalIssueKey">Key for external issue tracking</param>
 		/// <returns>True if the issue was updated</returns>
-		Task<bool> UpdateIssueAsync(int id, string? summary = null, string? description = null, bool? promoted = null, UserId? ownerId = null, UserId? nominatedById = null, bool? acknowledged = null, UserId? declinedById = null, int? fixChange = null, UserId? resolvedById = null, List<ObjectId>? addSpanIds = null, List<ObjectId>? removeSpanIds = null);
+		Task<bool> UpdateIssueAsync(int id, string? summary = null, string? description = null, bool? promoted = null, UserId? ownerId = null, UserId? nominatedById = null, bool? acknowledged = null, UserId? declinedById = null, int? fixChange = null, UserId? resolvedById = null, List<ObjectId>? addSpanIds = null, List<ObjectId>? removeSpanIds = null, string? externalIssueKey = null);
 
 		/// <summary>
 		/// Marks a step as complete
diff --git a/Engine/Source/Programs/Horde/Horde.Build/Issues/Impl/IssueCollection.cs b/Engine/Source/Programs/Horde/Horde.Build/Issues/Impl/IssueCollection.cs
--- a/Engine/Source/Programs/Horde/Horde.Build/Issues/Impl/IssueCollection.cs
+++ b/Engine/Source/Programs/Horde/Horde.Build/Issues/Impl/IssueCollection.cs
@@ -108,6 +108,9 @@ namespace Horde.Build.Collections.Impl
 			IReadOnlyList<IIssueStream> IIssue.Streams => Streams;
 			DateTime IIssue.LastSeenAt => (LastSeenAt == default) ? DateTime.UtcNow : LastSeenAt;
 
+			[BsonIgnoreIfNull]
+			public string? ExternalIssueKey { get; set; }
+			
 			[BsonConstructor]
 			private Issue()
 			{
@@ -799,7 +802,7 @@ namespace Horde.Build.Collections.Impl
 		}
 
 		/// <inheritdoc/>
-		public async Task<IIssue?> TryUpdateIssueAsync(IIssue issue, IssueSeverity? newSeverity = null, string? newSummary = null, string? newUserSummary = null, string? newDescription = null, bool? newManuallyPromoted = null, UserId? newOwnerId = null, UserId? newNominatedById = null, bool? newAcknowledged = null, UserId? newDeclinedById = null, int? newFixChange = null, UserId? newResolvedById = null, List<ObjectId>? newExcludeSpanIds = null, DateTime? newLastSeenAt = null)
+		public async Task<IIssue?> TryUpdateIssueAsync(IIssue issue, IssueSeverity? newSeverity = null, string? newSummary = null, string? newUserSummary = null, string? newDescription = null, bool? newManuallyPromoted = null, UserId? newOwnerId = null, UserId? newNominatedById = null, bool? newAcknowledged = null, UserId? newDeclinedById = null, int? newFixChange = null, UserId? newResolvedById = null, List<ObjectId>? newExcludeSpanIds = null, DateTime? newLastSeenAt = null, string? externaIssueKey = null)
 		{
 			Issue issueDocument = (Issue)issue;
 
@@ -932,6 +935,12 @@ namespace Horde.Build.Collections.Impl
 			{
 				await _issueSuspects.UpdateManyAsync(x => x.IssueId == issue.Id && x.AuthorId == newDeclinedById.Value, Builders<IssueSuspect>.Update.Set(x => x.DeclinedAt, DateTime.UtcNow));
 			}
+
+			if (externaIssueKey != null)
+			{
+				updates.Add(Builders<Issue>.Update.Set(x => x.ExternalIssueKey, externaIssueKey == String.Empty ? null : externaIssueKey));
+			}
+
 			if (updates.Count == 0)
 			{
 				return issueDocument;
diff --git a/Engine/Source/Programs/Horde/Horde.Build/Issues/Impl/IssueService.cs b/Engine/Source/Programs/Horde/Horde.Build/Issues/Impl/IssueService.cs
--- a/Engine/Source/Programs/Horde/Horde.Build/Issues/Impl/IssueService.cs
+++ b/Engine/Source/Programs/Horde/Horde.Build/Issues/Impl/IssueService.cs
@@ -51,6 +51,9 @@ namespace Horde.Build.Services.Impl
 		/// <inheritdoc/>
 		public IUser? ResolvedBy { get; }
 
+		/// <inheritdoc/>
+		public string? ExternalIssueKey { get; }
+
 		/// <inheritdoc/>
 		public IReadOnlyList<IIssueSpan> Spans { get; }
 
@@ -81,7 +84,8 @@ namespace Horde.Build.Services.Impl
 		/// <param name="suspects"></param>
 		/// <param name="suspectUsers"></param>
 		/// <param name="showDesktopAlerts"></param>
-		public IssueDetails(IIssue issue, IUser? owner, IUser? nominatedBy, IUser? resolvedBy, IReadOnlyList<IIssueSpan> spans, IReadOnlyList<IIssueStep> steps, IReadOnlyList<IIssueSuspect> suspects, IReadOnlyList<IUser> suspectUsers, bool showDesktopAlerts)
+		/// <param name="externalIssueKey"></param>
+		public IssueDetails(IIssue issue, IUser? owner, IUser? nominatedBy, IUser? resolvedBy, IReadOnlyList<IIssueSpan> spans, IReadOnlyList<IIssueStep> steps, IReadOnlyList<IIssueSuspect> suspects, IReadOnlyList<IUser> suspectUsers, bool showDesktopAlerts, string? externalIssueKey)
 		{
 			Issue = issue;
 			Owner = owner;
@@ -91,8 +95,10 @@ namespace Horde.Build.Services.Impl
 			Steps = steps;
 			Suspects = suspects;
 			SuspectUsers = suspectUsers;
+			ExternalIssueKey = externalIssueKey; 
 			_showDesktopAlerts = showDesktopAlerts;
 
+
 			if (issue.OwnerId == null)
 			{
 				_notifyUsers = new HashSet<UserId>(suspects.Select(x => x.AuthorId));
@@ -303,7 +309,7 @@ namespace Horde.Build.Services.Impl
 				}
 			}
 
-			return new IssueDetails(issue, owner, nominatedBy, resolvedBy, spans, steps, suspects, suspectUsers, ShowDesktopAlertsForIssue(issue, spans));
+			return new IssueDetails(issue, owner, nominatedBy, resolvedBy, spans, steps, suspects, suspectUsers, ShowDesktopAlertsForIssue(issue, spans), issue.ExternalIssueKey);
 		}
 
 		/// <inheritdoc/>
@@ -384,7 +390,7 @@ namespace Horde.Build.Services.Impl
 		}
 
 		/// <inheritdoc/>
-		public async Task<bool> UpdateIssueAsync(int id, string? userSummary = null, string? description = null, bool? promoted = null, UserId? ownerId = null, UserId? nominatedById = null, bool? acknowledged = null, UserId? declinedById = null, int? fixChange = null, UserId? resolvedById = null, List<ObjectId>? addSpanIds = null, List<ObjectId>? removeSpanIds = null)
+		public async Task<bool> UpdateIssueAsync(int id, string? userSummary = null, string? description = null, bool? promoted = null, UserId? ownerId = null, UserId? nominatedById = null, bool? acknowledged = null, UserId? declinedById = null, int? fixChange = null, UserId? resolvedById = null, List<ObjectId>? addSpanIds = null, List<ObjectId>? removeSpanIds = null, string? externalIssueKey = null)
 		{
 			IIssue? issue;
 			for (; ; )
@@ -395,7 +401,7 @@ namespace Horde.Build.Services.Impl
 					return false;
 				}
 
-				issue = await _issueCollection.TryUpdateIssueAsync(issue, newUserSummary: userSummary, newDescription: description, newPromoted: promoted, newOwnerId: ownerId ?? resolvedById, newNominatedById: nominatedById, newDeclinedById: declinedById, newAcknowledged: acknowledged, newFixChange: fixChange, newResolvedById: resolvedById, newExcludeSpanIds: removeSpanIds);
+				issue = await _issueCollection.TryUpdateIssueAsync(issue, newUserSummary: userSummary, newDescription: description, newPromoted: promoted, newOwnerId: ownerId ?? resolvedById, newNominatedById: nominatedById, newDeclinedById: declinedById, newAcknowledged: acknowledged, newFixChange: fixChange, newResolvedById: resolvedById, newExcludeSpanIds: removeSpanIds, externalIssueKey: externalIssueKey);
 				if (issue != null)
 				{
 					break;
diff --git a/Engine/Source/Programs/Horde/Horde.Build/Jira/IJiraService.cs b/Engine/Source/Programs/Horde/Horde.Build/Jira/IJiraService.cs
deleted file mode 100644
diff --git a/Engine/Source/Programs/Horde/Horde.Build/Jira/Impl/JiraService.cs b/Engine/Source/Programs/Horde/Horde.Build/Jira/Impl/JiraService.cs
deleted file mode 100644
diff --git a/Engine/Source/Programs/Horde/Horde.Build/Server/Config.cs b/Engine/Source/Programs/Horde/Horde.Build/Server/Config.cs
--- a/Engine/Source/Programs/Horde/Horde.Build/Server/Config.cs
+++ b/Engine/Source/Programs/Horde/Horde.Build/Server/Config.cs
@@ -297,6 +297,31 @@ namespace Horde.Build.Server
 		}
 	}
 
+	/// <summary>
+	/// External issue tracking configuration for a workflow
+	/// </summary>
+	public class ExternalIssueConfig
+	{
+		/// <summary>
+		/// Project key in external issue tracker
+		/// </summary>
+		[Required]
+		public string ProjectKey { get; set; } = String.Empty;
+
+		/// <summary>
+		/// Default component id for issues using workflow
+		/// </summary>
+		[Required]
+		public string DefaultComponentId { get; set; } = String.Empty;
+
+		/// <summary>
+		/// Default issue type id for issues using workflow
+		/// </summary>
+		[Required]
+		public string DefaultIssueTypeId { get; set; } = String.Empty;
+
+	}
+
 	/// <summary>
 	/// Configuration for an issue workflow
 	/// </summary>
@@ -346,6 +371,11 @@ namespace Horde.Build.Server
 		/// Additional node annotations implicit in this workflow
 		/// </summary>
 		public NodeAnnotations Annotations { get; set; } = new NodeAnnotations();
+
+		/// <summary>
+		/// External issue tracking configuration for this workflow
+		/// </summary>
+		public ExternalIssueConfig? ExternalIssues { get; set; }
 	}
 
 	/// <summary>
diff --git a/Engine/Source/Programs/Horde/Horde.Build/ServerSettings.cs b/Engine/Source/Programs/Horde/Horde.Build/ServerSettings.cs
--- a/Engine/Source/Programs/Horde/Horde.Build/ServerSettings.cs
+++ b/Engine/Source/Programs/Horde/Horde.Build/ServerSettings.cs
@@ -459,6 +459,11 @@ namespace Horde.Build
 		/// </summary>
 		public bool P4BridgeCanImpersonate { get; set; } = false;
 
+		/// <summary>
+		/// Url of P4 Swarm installation
+		/// </summary>
+		public Uri? P4SwarmUrl { get; set; }
+
 		/// <summary>
 		/// The Jira service account user name
 		/// </summary>
diff --git a/Engine/Source/Programs/Horde/Horde.Build/Startup.cs b/Engine/Source/Programs/Horde/Horde.Build/Startup.cs
--- a/Engine/Source/Programs/Horde/Horde.Build/Startup.cs
+++ b/Engine/Source/Programs/Horde/Horde.Build/Startup.cs
@@ -403,10 +403,19 @@ namespace Horde.Build
 			services.AddSingleton<INotificationSink, SlackNotificationSink>(sp => sp.GetRequiredService<SlackNotificationSink>());
 			services.AddSingleton<StreamService>();
 			services.AddSingleton<UpgradeService>();
-			services.AddSingleton<DeviceService>();			
-			services.AddSingleton<JiraService>();
+			services.AddSingleton<DeviceService>();						
 			services.AddSingleton<NoticeService>();
 
+			if (settings.JiraUrl != null)
+			{
+				services.AddSingleton<IExternalIssueService, JiraService>();
+			}
+			else
+			{
+				services.AddSingleton<IExternalIssueService, DefaultExternalIssueService>();
+			}
+
+
 			AWSOptions awsOptions = Configuration.GetAWSOptions();
 			if (settings.S3CredentialType == "AssumeRole" && settings.S3AssumeArn != null)
 			{
@@ -598,6 +607,8 @@ namespace Horde.Build
 				services.AddHostedService(provider => provider.GetRequiredService<DeviceService>());
 			}
 
+			services.AddHostedService(provider => provider.GetRequiredService<IExternalIssueService>());
+				
 			// Task sources. Order of registration is important here; it dictates the priority in which sources are served.
 			services.AddSingleton<JobTaskSource>();
 			services.AddHostedService<JobTaskSource>(provider => provider.GetRequiredService<JobTaskSource>());
diff --git a/Engine/Source/Programs/Horde/HordeDashboard/src/backend/Api.ts b/Engine/Source/Programs/Horde/HordeDashboard/src/backend/Api.ts
--- a/Engine/Source/Programs/Horde/HordeDashboard/src/backend/Api.ts
+++ b/Engine/Source/Programs/Horde/HordeDashboard/src/backend/Api.ts
@@ -374,6 +374,20 @@ export enum IssueSeverity {
 	Error = "Error"
 }
 
+/** Setting information required by dashboard */
+export type GetDashboardConfigResponse = {
+	
+	/** The name of the external issue service */
+	externalIssueServiceName?: string;
+
+	/** The url of the external issue service */
+	externalIssueServiceUrl?: string;
+
+	/** The url of the perforce swarm installation */
+	perforceSwarmUrl?: string;
+}
+
+
 /**Parameters to register a new agent */
 export type CreateAgentRequest = {
 
@@ -1763,8 +1777,26 @@ export type GetWorkspaceTypeResponse = {
 }
 
 
+/** State information for a step in the stream */
+export type GetTemplateStepStateResponse = {
+
+	/**Name of the step */
+	name: string;
+
+	/**User who paused the step */
+	pausedByUserId?: string;
+
+	/**User who quarantined the step */
+	quarantinedByUserId?: string; 
+
+}
+
 /**Information about a template in this stream */
 export type GetTemplateRefResponse = {
+
+	/** The name of the template */
+	name: string;
+
 	/**Unique id of this template ref */
 	id: string;
 
@@ -1774,8 +1806,8 @@ export type GetTemplateRefResponse = {
 	/// The schedule for this ref
 	schedule?: GetScheduleResponse;
 
-	name: string;
-
+	/** Step state for template in stream */
+	stepStates?: GetTemplateStepStateResponse[];
 }
 
 /** Specifies defaults for running a preflight */
@@ -2718,7 +2750,6 @@ export type GetIssueResponse = {
 	/**Whether the issue is promoted */
 	promoted: boolean;
 
-
 	/** Owner of the issue */
 	ownerInfo: GetThinUserInfoResponse;
 
@@ -2754,6 +2785,9 @@ export type GetIssueResponse = {
 	/** Whether to show alerts for this issue */
 	showDesktopAlerts: boolean;
 
+	/** External issue tracking */
+	externalIssueKey?: string;
+
 }
 
 /**Request an issue to be updated */
@@ -2792,6 +2826,9 @@ export type UpdateIssueRequest = {
 	/** List of spans to remove from this issue */
 	removeSpans?: string[];
 
+	/** An external issue key*/
+	externalIssueKey?: string;
+
 }
 
 export type GetUtilizationTelemetryStream = {
@@ -3847,6 +3884,9 @@ export type FindIssueResponse = {
 
 	/// Spans for this issue
 	spans: FindIssueSpanResponse[];
+
+	/** External issue tracking */
+	externalIssueKey?: string;
 }
 
 export type GetAgentSoftwareChannelResponse = {
@@ -3856,33 +3896,89 @@ export type GetAgentSoftwareChannelResponse = {
 	version?: string;
 }
 
-/// Jira issue information
-export type GetJiraIssueResponse = {
+/// External Issue Response
+export type GetExternalIssueResponse = {
 
-	/** The jira key */
-	key: string;
+	/** The external issue key */
+	key: string; 
 
-	/** Link to the issue on jira */
-	jiraLink: string;
+	/** The issue link on external tracking site */
+	link?: string; 
 
-	/** The current status of the issue */
+	/** The issue status name, "To Do", "In Progress", etc*/
 	statusName?: string;
 
-	/** The current priority of the issue */
-	priorityName?: string;
-
-	/** The resolution of the issue */
+	/** The issue resolution name, "Fixed", "Closed", etc*/
 	resolutionName?: string;
 
-	/** Jira username of assignee */
+	/** The issue priority name, "1 - Critical", "2 - Major", etc*/
+	priorityName?: string;
+
+	/** The current assignee's user name*/
 	assigneeName?: string;
 
-	/** Jira display name of assignee */
+	/** The current assignee's display name*/
 	assigneeDisplayName?: string;
 
-	/** Email address of assignee */
+	/** The current assignee's email address*/
 	assigneeEmailAddress?: string;
+}
+
+/** Request an issue to be created on external issue tracking system */
+export type CreateExternalIssueRequest = {
+		
+	/** Horde issue which is linked to external issue */
+	issueId: number;
+
+	/** Summary text for external issue */
+	summary: string;
+
+	/** A stream this this issue */
+	streamId: string;
+
+	/** External issue project id */
+	projectId: string;
+
+	/** External issue component id */
+	componentId: string;
+
+	/** External issue type id */
+	issueTypeId: string;
+
+	/** Optional description text for external issue */
+	description?: string;
+
+	/** Optional link to Horde issue */
+	hordeIssueLink?: string;
+}
+
+/** Response for externally created issue */
+export type CreateExternalIssueResponse = {
+	
+	/** External issue tracking key	 */
+	key: string;
+
+	/**  Link to issue on external tracking site */
+	link?: string;
+}
+
+/// External issue project information
+export type GetExternalIssueProjectResponse = {
+		
+	/// The project key	
+	projectKey: string;
+
+	/// The name of the project
+	name: string;
+
+	/// The id of the project
+	id: string;
+
+	/// component id => name
+	components: Record<string, string>;
 
+	/// IssueType id => name
+	issueTypes: Record<string, string>;
 }
 
 /// Create a notice which will display on the dashboard
diff --git a/Engine/Source/Programs/Horde/HordeDashboard/src/backend/Backend.ts b/Engine/Source/Programs/Horde/HordeDashboard/src/backend/Backend.ts
--- a/Engine/Source/Programs/Horde/HordeDashboard/src/backend/Backend.ts
+++ b/Engine/Source/Programs/Horde/HordeDashboard/src/backend/Backend.ts
@@ -3,7 +3,7 @@
 import { configure } from "mobx";
 import { isNumber } from 'util';
 import templateCache from '../backend/TemplateCache';
-import { AgentData, ArtifactData, AuditLogEntry, AuditLogQuery, BatchUpdatePoolRequest, ChangeSummaryData, CreateDeviceRequest, CreateDeviceResponse, CreateJobRequest, CreateJobResponse, CreateNoticeRequest, CreatePoolRequest, CreateSoftwareResponse, CreateSubscriptionRequest, CreateSubscriptionResponse, DashboardPreference, EventData, FindIssueResponse, FindJobTimingsResponse, GetAgentSoftwareChannelResponse, GetArtifactZipRequest, GetDevicePlatformResponse, GetDevicePoolResponse, GetDeviceReservationResponse, GetDeviceResponse, GetGraphResponse, GetIssueStreamResponse, GetJiraIssueResponse, GetJobsTabResponse, GetJobStepRefResponse, GetJobStepTraceResponse, GetJobTimingResponse, GetLogEventResponse, GetNoticeResponse, GetNotificationResponse, GetPerforceServerStatusResponse, GetServerInfoResponse, GetServerSettingsResponse, GetSoftwareResponse, GetSubscriptionResponse, GetUserResponse, GetUtilizationTelemetryResponse, GlobalConfig, IssueData, IssueQuery, IssueQueryV2, JobData, JobQuery, JobsTabColumnType, JobStepOutcome, JobStreamQuery, JobTimingsQuery, LeaseData, LogData, LogLineData, PoolData, ProjectData, ScheduleData, ScheduleQuery, SearchLogFileResponse, ServerUpdateResponse, SessionData, StreamData, TabType, TemplateData, TestData, UpdateAgentRequest, UpdateDeviceRequest, UpdateGlobalConfigRequest, UpdateIssueRequest, UpdateJobRequest, UpdateLeaseRequest, UpdateNoticeRequest, UpdateNotificationsRequest, UpdatePoolRequest, UpdateServerSettingsRequest, UpdateStepRequest, UpdateStepResponse, UpdateUserRequest, UsersQuery } from './Api';
+import { AgentData, ArtifactData, AuditLogEntry, AuditLogQuery, BatchUpdatePoolRequest, ChangeSummaryData, CreateDeviceRequest, CreateDeviceResponse, CreateExternalIssueRequest, CreateExternalIssueResponse, CreateJobRequest, CreateJobResponse, CreateNoticeRequest, CreatePoolRequest, CreateSoftwareResponse, CreateSubscriptionRequest, CreateSubscriptionResponse, DashboardPreference, EventData, FindIssueResponse, FindJobTimingsResponse, GetAgentSoftwareChannelResponse, GetArtifactZipRequest, GetDashboardConfigResponse, GetDevicePlatformResponse, GetDevicePoolResponse, GetDeviceReservationResponse, GetDeviceResponse, GetExternalIssueProjectResponse, GetExternalIssueResponse, GetGraphResponse, GetIssueStreamResponse, GetJobsTabResponse, GetJobStepRefResponse, GetJobStepTraceResponse, GetJobTimingResponse, GetLogEventResponse, GetNoticeResponse, GetNotificationResponse, GetPerforceServerStatusResponse, GetServerInfoResponse, GetServerSettingsResponse, GetSoftwareResponse, GetSubscriptionResponse, GetUserResponse, GetUtilizationTelemetryResponse, GlobalConfig, IssueData, IssueQuery, IssueQueryV2, JobData, JobQuery, JobsTabColumnType, JobStepOutcome, JobStreamQuery, JobTimingsQuery, LeaseData, LogData, LogLineData, PoolData, ProjectData, ScheduleData, ScheduleQuery, SearchLogFileResponse, ServerUpdateResponse, SessionData, StreamData, TabType, TemplateData, TestData, UpdateAgentRequest, UpdateDeviceRequest, UpdateGlobalConfigRequest, UpdateIssueRequest, UpdateJobRequest, UpdateLeaseRequest, UpdateNoticeRequest, UpdateNotificationsRequest, UpdatePoolRequest, UpdateServerSettingsRequest, UpdateStepRequest, UpdateStepResponse, UpdateUserRequest, UsersQuery } from './Api';
 import dashboard from './Dashboard';
 import { ChallengeStatus, Fetch } from './Fetch';
 import graphCache, { GraphQuery } from './GraphCache';
@@ -21,6 +21,17 @@ const updateInterval = 120 * 1000;
 
 export class Backend {
 
+
+    getDashboardConfig(): Promise<GetDashboardConfigResponse> {
+        return new Promise<GetDashboardConfigResponse>((resolve, reject) => {
+            this.backend.get("/api/v1/dashboard/config").then((response) => {
+                resolve(response.data as GetDashboardConfigResponse);
+            }).catch(reason => {
+                reject(reason);
+            });
+        });
+    }
+
     getProjects(): Promise<ProjectData[]> {
 
         return new Promise<ProjectData[]>((resolve, reject) => {
@@ -437,20 +448,20 @@ export class Backend {
     getBatchJobTiming(query: JobTimingsQuery): Promise<FindJobTimingsResponse> {
 
         return new Promise<FindJobTimingsResponse>((resolve, reject) => {
-            this.backend.get(`/api/v1/jobs/timing`, { 
+            this.backend.get(`/api/v1/jobs/timing`, {
                 params: query
             }).then((response) => {
                 let timingsWrapper = response.data as FindJobTimingsResponse;
                 let graphCalls: { jobResponse: JobData, call: any }[] = [];
                 Object.values(timingsWrapper.timings).forEach(timing => {
-                    if(!timing.jobResponse.graphHash) {
+                    if (!timing.jobResponse.graphHash) {
                         return reject(`Job ${timing.job.id} has undefined graph hash`);
                     }
-                    graphCalls.push( { jobResponse: timing.jobResponse, call: graphCache.get({ graphHash: timing.jobResponse.graphHash, jobId: timing.jobResponse.id }) });
+                    graphCalls.push({ jobResponse: timing.jobResponse, call: graphCache.get({ graphHash: timing.jobResponse.graphHash, jobId: timing.jobResponse.id }) });
                 });
                 let allPromises = graphCalls.map(item => item.call);
                 Promise.all(allPromises).then(responses => {
-                    for(let idx = 0; idx < responses.length; idx++) {
+                    for (let idx = 0; idx < responses.length; idx++) {
                         graphCalls[idx].jobResponse.graphRef = responses[idx] as GetGraphResponse;
                     }
                     resolve(timingsWrapper);
@@ -1043,27 +1054,52 @@ export class Backend {
         });
     }
 
-    getJiraIssues(streamId: string, jiraKeys: string[]): Promise<GetJiraIssueResponse[]> {
+    getExternalIssues(streamId: string, keys: string[]): Promise<GetExternalIssueResponse[]> {
 
         const params = {
             streamId: streamId,
-            jiraKeys: jiraKeys
+            keys: keys
         };
 
-        return new Promise<GetJiraIssueResponse[]>((resolve, reject) => {
-            this.backend.get(`/api/v1/jira`, { params: params }).then((value) => {
-                resolve(value.data as GetJiraIssueResponse[]);
+        return new Promise<GetExternalIssueResponse[]>((resolve, reject) => {
+            this.backend.get(`/api/v1/issues/external`, { params: params }).then((value) => {
+                resolve(value.data as GetExternalIssueResponse[]);
             }).catch((reason) => {
                 reject(reason);
             });
         });
+    }
+
+    getExternalIssueProjects(streamId: string): Promise<GetExternalIssueProjectResponse[]> {
+
+        const params = {
+            streamId: streamId,
+        };
 
+        return new Promise<GetExternalIssueProjectResponse[]>((resolve, reject) => {
+            this.backend.get(`/api/v1/issues/external/projects`, { params: params }).then((value) => {
+                resolve(value.data as GetExternalIssueProjectResponse[]);
+            }).catch((reason) => {
+                reject(reason);
+            });
+        });
+    }
+
+
+    createExternalIssue(request: CreateExternalIssueRequest): Promise<CreateExternalIssueResponse> {
+        return new Promise<CreateExternalIssueResponse>((resolve, reject) => {
+            this.backend.post(`/api/v1/issues/external`, request).then((value) => {
+                resolve(value.data as CreateExternalIssueResponse);
+            }).catch(reason => {
+                reject(reason);
+            });
+        });
     }
 
     getTestData(id: string, filter?: string): Promise<TestData> {
         return new Promise<TestData>((resolve, reject) => {
             const params: any = {
-                filter: 'id,key,change,jobId,stepId,streamId,' + (filter? filter : 'data')
+                filter: 'id,key,change,jobId,stepId,streamId,' + (filter ? filter : 'data')
             };
 
             this.backend.get(`/api/v1/testdata/${id}`, {
@@ -1104,7 +1140,7 @@ export class Backend {
                 streamId: streamId,
                 key: key,
                 count: count,
-                filter: 'id,key,change,jobId,stepId,streamId,' + (filter? filter : 'data')
+                filter: 'id,key,change,jobId,stepId,streamId,' + (filter ? filter : 'data')
             };
 
             if (maxChange) {
@@ -1419,7 +1455,7 @@ export class Backend {
     }
 
     // create a new notice
-    createNotice(request:CreateNoticeRequest): Promise<void> {
+    createNotice(request: CreateNoticeRequest): Promise<void> {
         return new Promise<void>((resolve, reject) => {
             this.backend.post(`/api/v1/notices`, request).then(() => {
                 resolve();
@@ -1445,7 +1481,7 @@ export class Backend {
     }
 
     // update a notice
-    updateNotice(request: UpdateNoticeRequest):Promise<void> {
+    updateNotice(request: UpdateNoticeRequest): Promise<void> {
         return new Promise<void>((resolve, reject) => {
             this.backend.put(`/api/v1/notices`, request).then(() => {
                 resolve();
@@ -1456,7 +1492,7 @@ export class Backend {
     }
 
     // delete a notice
-    deleteNotice(id:string):Promise<void> {
+    deleteNotice(id: string): Promise<void> {
         return new Promise<void>((resolve, reject) => {
             this.backend.delete(`/api/v1/notices/${id}`).then(() => {
                 resolve();
@@ -1534,7 +1570,7 @@ export class Backend {
                 this.backend.login(window.location.toString());
                 return;
             }
-
+            
             await dashboard.update();
 
             if (dashboard.localCache) {
diff --git a/Engine/Source/Programs/Horde/HordeDashboard/src/backend/Dashboard.ts b/Engine/Source/Programs/Horde/HordeDashboard/src/backend/Dashboard.ts
--- a/Engine/Source/Programs/Horde/HordeDashboard/src/backend/Dashboard.ts
+++ b/Engine/Source/Programs/Horde/HordeDashboard/src/backend/Dashboard.ts
@@ -3,7 +3,7 @@
 import { getTheme } from "@fluentui/react";
 import { action, observable } from 'mobx';
 import backend from '.';
-import { DashboardPreference, GetUserResponse, UserClaim } from './Api';
+import { DashboardPreference, GetDashboardConfigResponse, GetUserResponse, UserClaim } from './Api';
 
 const theme = getTheme();
 
@@ -75,7 +75,7 @@ export class Dashboard {
 
         this.data.pinnedJobIds = [];
 
-        backend.updateUser({ removePinnedJobIds:jobs });
+        backend.updateUser({ removePinnedJobIds: jobs });
 
         this.setUpdated();
     }
@@ -138,11 +138,22 @@ export class Dashboard {
 
     }
 
+
+    get swarmUrl(): string | undefined {
+        return this.config?.perforceSwarmUrl;
+    }
+
+    get externalIssueService(): { name: string, url: string } | undefined {
+        if (!this.config?.externalIssueServiceUrl) {
+            return undefined;
+        }
+        return { name: this.config.externalIssueServiceName ?? "???", url: this.config.externalIssueServiceUrl };
+    }
+
     get userId(): string {
         return this.data.id;
     }
 
-
     get pinnedJobsIds(): string[] {
 
         return this.data.pinnedJobIds ?? [];
@@ -299,7 +310,7 @@ export class Dashboard {
             [StatusColor.Running, dark ? "#146579" : theme.palette.blueLight],
             [StatusColor.Waiting, dark ? "#474542" : "#A19F9D"],
             [StatusColor.Ready, dark ? "#474542" : "#A19F9D"],
-            [StatusColor.Skipped, dark ? "#63625c" : "#F3F2F1"],            
+            [StatusColor.Skipped, dark ? "#63625c" : "#F3F2F1"],
             [StatusColor.Unspecified, "#637087"]
         ]);
 
@@ -372,6 +383,10 @@ export class Dashboard {
 
             this.updating = true;
 
+            if (!this.config) {
+                this.config = await backend.getDashboardConfig();                
+            }
+
             if (this.polling || !this.available) {
 
                 const cancelId = this.cancelId++;
@@ -462,7 +477,7 @@ export class Dashboard {
         return this.data;
     }
 
-    private async postPreferences(reload?:boolean): Promise<boolean> {
+    private async postPreferences(reload?: boolean): Promise<boolean> {
 
         // cancel any pending        
         for (let i = 0; i < this.cancelId; i++) {
@@ -485,13 +500,13 @@ export class Dashboard {
             if (reload) {
                 window.location.reload();
             }
-                
+
         }
 
         return success;
 
     }
-    
+
     requestLogout = false;
 
     serverSettingsChanged: boolean = false;
@@ -509,6 +524,8 @@ export class Dashboard {
 
     private pollMS = 4 * 1000;
 
+    private config?: GetDashboardConfigResponse;
+
 }
 
 const dashboard = new Dashboard();
diff --git a/Engine/Source/Programs/Horde/HordeDashboard/src/base/utilities/issueUtils.ts b/Engine/Source/Programs/Horde/HordeDashboard/src/base/utilities/issueUtils.ts
--- a/Engine/Source/Programs/Horde/HordeDashboard/src/base/utilities/issueUtils.ts
+++ b/Engine/Source/Programs/Horde/HordeDashboard/src/base/utilities/issueUtils.ts
@@ -4,46 +4,6 @@ import dashboard from "../../backend/Dashboard";
 import { projectStore } from "../../backend/ProjectStore";
 import { displayTimeZone } from "./timeUtils";
 
-export type IssueJira = {
-	description: string;
-	link: string;
-}
-
-// This regex could be optimized, though it works
-// eslint-disable-next-line
-const urlRegex = /([\w+]+\:\/\/)?([\w\d-]+\.)*[\w-]+[\.\:]\w+([\/\?\=\&\#\.]?[\w-]+)*\/?/gm
-
-export function getIssueJiras(description?: string): IssueJira[] {
-
-	if (!description) {
-		return [];
-	}
-
-	const jiras: IssueJira[] = [];
-
-	const matches = Array.from(description.matchAll(urlRegex));
-
-	matches.forEach(m => {
-		m.forEach(match => {
-
-			if (!match) {
-				return;
-			}
-			// @todo: add jira match to server config
-			if (match.toLowerCase().startsWith("https://jira.it.epicgames.com")) {
-				const path = new URL(match).pathname.split("/");
-				if (path.length) {
-					jiras.push({link: match, description: path[path.length - 1]})
-				}
-				
-			}
-		})
-	})
-	
-	return jiras;
-
-}
-
 export function getIssueStatus(issue: GetIssueResponse, showResolveTime?: boolean): string {
 
 	let text = "";
