Horde: Do not fetch config data for deleted streams. It may not exist or be valid.

diff --git a/Engine/Source/Programs/Horde/Horde.Build/Projects/Impl/StreamCollection.cs b/Engine/Source/Programs/Horde/Horde.Build/Projects/Impl/StreamCollection.cs
--- a/Engine/Source/Programs/Horde/Horde.Build/Projects/Impl/StreamCollection.cs
+++ b/Engine/Source/Programs/Horde/Horde.Build/Projects/Impl/StreamCollection.cs
@@ -135,7 +135,14 @@ namespace Horde.Build.Collections.Impl
 		{
 			try
 			{
-				stream.Config = await _configCollection.GetConfigAsync<StreamConfig>(stream.ConfigRevision);
+				if (stream.Deleted)
+				{
+					stream.Config = new StreamConfig();
+				}
+				else
+				{
+					stream.Config = await _configCollection.GetConfigAsync<StreamConfig>(stream.ConfigRevision);
+				}
 			}
 			catch (Exception)
 			{
