Horde: Remove support for running BuildGraph scripts directly from Perforce (and fix CI).

diff --git a/Engine/Source/Programs/Horde/Horde.Build/Controllers/JobsController.cs b/Engine/Source/Programs/Horde/Horde.Build/Controllers/JobsController.cs
--- a/Engine/Source/Programs/Horde/Horde.Build/Controllers/JobsController.cs
+++ b/Engine/Source/Programs/Horde/Horde.Build/Controllers/JobsController.cs
@@ -154,16 +154,6 @@ namespace Horde.Build.Controllers
 			// New properties for the job
 			List<string> arguments = create.Arguments ?? template.GetDefaultArguments();
 
-			// Special handling for internal graph evaluation
-			if (arguments.Any(x => x.Equals("-InternalParser", StringComparison.OrdinalIgnoreCase)))
-			{
-				IPerforceConnection? perforceConnection = await _perforce.GetServiceUserConnection(stream.ClusterName);
-				if (perforceConnection != null)
-				{
-					graph = await _graphs.AddAsync(perforceConnection, stream, change, codeChange, arguments);
-				}
-			}
-
 			// Check the preflight change is valid
 			if (create.PreflightChange != null)
 			{
diff --git a/Engine/Source/Programs/Horde/Horde.Build/Utilities/BuildGraphExtensions.cs b/Engine/Source/Programs/Horde/Horde.Build/Utilities/BuildGraphExtensions.cs
deleted file mode 100644
