Horde: Missing files.

diff --git a/Engine/Source/Programs/Horde/Horde.Agent/Commands/Parse/ParseCommand.cs b/Engine/Source/Programs/Horde/Horde.Agent/Commands/Parse/ParseCommand.cs
--- a/Engine/Source/Programs/Horde/Horde.Agent/Commands/Parse/ParseCommand.cs
+++ b/Engine/Source/Programs/Horde/Horde.Agent/Commands/Parse/ParseCommand.cs
@@ -68,7 +68,7 @@ namespace Horde.Agent.Commands.Parse
 				PerforceViewMap viewMap = new PerforceViewMap();
 				viewMap.Entries.Add(new PerforceViewMapEntry(true, "...", $"{Stream}/..."));
 
-				JsonPerforceLogger jsonLogger = new JsonPerforceLogger(logger, WorkspaceDir ?? DirectoryReference.GetCurrentDirectory(), viewMap, Change);
+				PerforceLogger jsonLogger = new PerforceLogger(logger, WorkspaceDir ?? DirectoryReference.GetCurrentDirectory(), viewMap, Change);
 				using (LogParser parser = new LogParser(jsonLogger, ignorePatterns))
 				{
 					byte[] data = new byte[1024];
diff --git a/Engine/Source/Programs/Horde/Horde.Agent/Execution/TestExecutor.cs b/Engine/Source/Programs/Horde/Horde.Agent/Execution/TestExecutor.cs
--- a/Engine/Source/Programs/Horde/Horde.Agent/Execution/TestExecutor.cs
+++ b/Engine/Source/Programs/Horde/Horde.Agent/Execution/TestExecutor.cs
@@ -150,7 +150,7 @@ namespace Horde.Agent.Execution
 			PerforceViewMap viewMap = new PerforceViewMap();
 			viewMap.Entries.Add(new PerforceViewMapEntry(true, "...", "//UE4/Main/..."));
 			
-			JsonPerforceLogger perforceLogger = new JsonPerforceLogger(logger, new DirectoryReference("D:\\Test"), viewMap, 12345);
+			PerforceLogger perforceLogger = new PerforceLogger(logger, new DirectoryReference("D:\\Test"), viewMap, 12345);
 			using (LogParser filter = new LogParser(perforceLogger, new List<string>()))
 			{
 				if(outcome == JobStepOutcome.Warnings)
