Horde: Add retry for agent collection updates

If an outdated agent document was used, the update would previously fail. This change will allows retrying a specified number of times, and refreshing the doc between each attempt.

diff --git a/Engine/Source/Programs/Horde/Horde.Build/Agents/Impl/AgentCollection.cs b/Engine/Source/Programs/Horde/Horde.Build/Agents/Impl/AgentCollection.cs
--- a/Engine/Source/Programs/Horde/Horde.Build/Agents/Impl/AgentCollection.cs
+++ b/Engine/Source/Programs/Horde/Horde.Build/Agents/Impl/AgentCollection.cs
@@ -176,7 +176,7 @@ namespace Horde.Build.Collections.Impl
 			AgentDocument agent = (AgentDocument)agentInterface;
 
 			UpdateDefinition<AgentDocument> update = Builders<AgentDocument>.Update.Set(x => x.Deleted, true);
-			return await TryUpdateAsync(agent, update);
+			return await TryUpdateAsyncWithRetries(agent, update);
 		}
 
 		/// <inheritdoc/>
@@ -239,7 +239,7 @@ namespace Horde.Build.Collections.Impl
 		/// </summary>
 		/// <param name="current">The document to update</param>
 		/// <param name="update">The update definition</param>
-		/// <returns>True if the agent was updated</returns>
+		/// <returns>The updated agent document or null if update failed</returns>
 		private async Task<AgentDocument?> TryUpdateAsync(AgentDocument current, UpdateDefinition<AgentDocument> update)
 		{
 			uint prevUpdateIndex = current.UpdateIndex++;
@@ -251,6 +251,36 @@ namespace Horde.Build.Collections.Impl
 			return await _agents.FindOneAndUpdateAsync<AgentDocument>(filter, updateWithIndex, new FindOneAndUpdateOptions<AgentDocument, AgentDocument> { ReturnDocument = ReturnDocument.After });
 		}
 
+		/// <summary>
+		/// Update a single document, with retries to ensure it's applied
+		/// If an outdated AgentDocument is supplied, it needs to be refreshed so that the UpdateIndex matches.
+		/// </summary>
+		/// <param name="agentDoc">The document to update</param>
+		/// <param name="update">The update definition</param>
+		/// <param name="numRetries">Number of retries to attempt</param>
+		/// <returns>The updated agent document or null if update failed</returns>
+		private async Task<AgentDocument?> TryUpdateAsyncWithRetries(AgentDocument agentDoc, UpdateDefinition<AgentDocument> update, int numRetries = 5)
+		{
+			AgentDocument current = agentDoc;
+			for (int retryCount = 0; retryCount < numRetries; retryCount++)
+			{
+				AgentDocument? result = await TryUpdateAsync(current, update);
+				if (result != null)
+				{
+					return result;
+				}
+
+				AgentId id = current.Id; // Capture id to avoid doing it inside closure (and warnings)
+				current = await _agents.Find(x => x.Id == id).FirstOrDefaultAsync();
+				if (current == null)
+				{
+					return null;
+				}
+			}
+
+			return null;
+		}
+
 		/// <inheritdoc/>
 		public async Task<IAgent?> TryUpdateSettingsAsync(IAgent agentInterface, bool? enabled = null, bool? requestConform = null, bool? requestFullConform = null, bool? requestRestart = null, bool? requestShutdown = null, string? shutdownReason = null, AgentSoftwareChannelName? channel = null, List<PoolId>? pools = null, Acl? acl = null, string? comment = null)
 		{
@@ -327,7 +357,7 @@ namespace Horde.Build.Collections.Impl
 			}
 
 			// Apply the update
-			IAgent? newAgent = await TryUpdateAsync(agent, updateBuilder.Combine(updates));
+			IAgent? newAgent = await TryUpdateAsyncWithRetries(agent, updateBuilder.Combine(updates));
 			if (newAgent != null)
 			{
 				if (newAgent.RequestRestart != agent.RequestRestart || newAgent.RequestConform != agent.RequestConform || newAgent.RequestShutdown != agent.RequestShutdown || newAgent.Channel != agent.Channel)
@@ -402,7 +432,7 @@ namespace Horde.Build.Collections.Impl
 			}
 
 			// Update the agent, and try to create new lease documents if we succeed
-			return await TryUpdateAsync(agent, updateBuilder.Combine(updates));
+			return await TryUpdateAsyncWithRetries(agent, updateBuilder.Combine(updates));
 		}
 
 		static bool ResourcesEqual(IReadOnlyDictionary<string, int> dictA, IReadOnlyDictionary<string, int> dictB)
@@ -441,7 +471,7 @@ namespace Horde.Build.Collections.Impl
 			}
 
 			// Update the agent
-			return await TryUpdateAsync(agent, update);
+			return await TryUpdateAsyncWithRetries(agent, update);
 		}
 
 		/// <inheritdoc/>
@@ -470,7 +500,7 @@ namespace Horde.Build.Collections.Impl
 			updates.Add(updateBuilder.Set(x => x.LastShutdownReason, "Unexpected"));
 
 			// Apply the update
-			return await TryUpdateAsync(agent, updateBuilder.Combine(updates));
+			return await TryUpdateAsyncWithRetries(agent, updateBuilder.Combine(updates));
 		}
 
 		/// <inheritdoc/>
@@ -490,7 +520,7 @@ namespace Horde.Build.Collections.Impl
 				update = update.Set(x => x.Deleted, agent.Deleted);
 			}
 
-			return await TryUpdateAsync(agent, update);
+			return await TryUpdateAsyncWithRetries(agent, update);
 		}
 
 		/// <inheritdoc/>
@@ -506,7 +536,7 @@ namespace Horde.Build.Collections.Impl
 			leases.Add(newLease);
 
 			UpdateDefinition<AgentDocument> update = Builders<AgentDocument>.Update.Set(x => x.Leases, leases);
-			return await TryUpdateAsync(agent, update);
+			return await TryUpdateAsyncWithRetries(agent, update);
 		}
 
 		/// <inheritdoc/>
@@ -515,7 +545,7 @@ namespace Horde.Build.Collections.Impl
 			AgentDocument agent = (AgentDocument)agentInterface;
 
 			UpdateDefinition<AgentDocument> update = Builders<AgentDocument>.Update.Set(x => x.Leases![leaseIdx].State, LeaseState.Cancelled);
-			IAgent? newAgent = await TryUpdateAsync(agent, update);
+			IAgent? newAgent = await TryUpdateAsyncWithRetries(agent, update);
 			if (newAgent != null)
 			{
 				await PublishUpdateEventAsync(agent.Id);
