Horde: Allow multiple clients to be mapped into the Perforce view. This will allow producing accurate diagnostics for compile errors introduced via AutoSDK.

Also fixes AutoSdk issue test.

diff --git a/Engine/Source/Programs/Horde/Horde.Agent/Commands/Parse/ParseCommand.cs b/Engine/Source/Programs/Horde/Horde.Agent/Commands/Parse/ParseCommand.cs
--- a/Engine/Source/Programs/Horde/Horde.Agent/Commands/Parse/ParseCommand.cs
+++ b/Engine/Source/Programs/Horde/Horde.Agent/Commands/Parse/ParseCommand.cs
@@ -65,10 +65,9 @@ namespace Horde.Agent.Commands.Parse
 			// Read the file and pipe it through the event parser
 			using (FileStream inputStream = FileReference.Open(InputFile, FileMode.Open, FileAccess.Read))
 			{
-				PerforceViewMap viewMap = new PerforceViewMap();
-				viewMap.Entries.Add(new PerforceViewMapEntry(true, "...", $"{Stream}/..."));
+				PerforceLogger jsonLogger = new PerforceLogger(logger);
+				jsonLogger.AddClientView(WorkspaceDir ?? DirectoryReference.GetCurrentDirectory(), $"{Stream}/...", Change ?? 1);
 
-				PerforceLogger jsonLogger = new PerforceLogger(logger, WorkspaceDir ?? DirectoryReference.GetCurrentDirectory(), viewMap, Change);
 				using (LogParser parser = new LogParser(jsonLogger, ignorePatterns))
 				{
 					byte[] data = new byte[1024];
diff --git a/Engine/Source/Programs/Horde/Horde.Agent/Execution/BuildGraphExecutor.cs b/Engine/Source/Programs/Horde/Horde.Agent/Execution/BuildGraphExecutor.cs
--- a/Engine/Source/Programs/Horde/Horde.Agent/Execution/BuildGraphExecutor.cs
+++ b/Engine/Source/Programs/Horde/Horde.Agent/Execution/BuildGraphExecutor.cs
@@ -811,15 +811,10 @@ namespace Horde.Agent.Execution
 				FileUtils.ForceDeleteDirectoryContents(telemetryDir);
 			}
 
-			PerforceViewMap viewMap = new PerforceViewMap();
-			viewMap.Entries.Add(new PerforceViewMapEntry(true, "...", $"{_stream.Name}/..."));
-
-			PerforceLogger perforceLogger = new PerforceLogger(logger, workspaceDir, viewMap, _job.Change);
-
 			List<string> ignorePatterns = await ReadIgnorePatternsAsync(workspaceDir, logger);
 
 			int exitCode;
-			using (LogParser filter = new LogParser(perforceLogger, ignorePatterns))
+			using (LogParser filter = new LogParser(logger, ignorePatterns))
 			{
 				await ExecuteCleanupScriptAsync(cleanupScript, filter, logger);
 				try
diff --git a/Engine/Source/Programs/Horde/Horde.Agent/Execution/TestExecutor.cs b/Engine/Source/Programs/Horde/Horde.Agent/Execution/TestExecutor.cs
--- a/Engine/Source/Programs/Horde/Horde.Agent/Execution/TestExecutor.cs
+++ b/Engine/Source/Programs/Horde/Horde.Agent/Execution/TestExecutor.cs
@@ -147,10 +147,9 @@ namespace Horde.Agent.Execution
 				logger.LogInformation("Credential: {CredentialName}={CredentialValue}", credential.Key, credential.Value);
 			}
 
-			PerforceViewMap viewMap = new PerforceViewMap();
-			viewMap.Entries.Add(new PerforceViewMapEntry(true, "...", "//UE4/Main/..."));
-			
-			PerforceLogger perforceLogger = new PerforceLogger(logger, new DirectoryReference("D:\\Test"), viewMap, 12345);
+			PerforceLogger perforceLogger = new PerforceLogger(logger);
+			perforceLogger.AddClientView(new DirectoryReference("D:\\Test"), "//UE4/Main/...", 12345);
+
 			using (LogParser filter = new LogParser(perforceLogger, new List<string>()))
 			{
 				if(outcome == JobStepOutcome.Warnings)
diff --git a/Engine/Source/Programs/Horde/Horde.Agent/Utility/PerforceLogger.cs b/Engine/Source/Programs/Horde/Horde.Agent/Utility/PerforceLogger.cs
--- a/Engine/Source/Programs/Horde/Horde.Agent/Utility/PerforceLogger.cs
+++ b/Engine/Source/Programs/Horde/Horde.Agent/Utility/PerforceLogger.cs
@@ -1,6 +1,7 @@
 // Copyright Epic Games, Inc. All Rights Reserved.
 
 using System;
+using System.Collections.Generic;
 using System.IO;
 using System.Text;
 using System.Text.Json;
@@ -15,24 +16,51 @@ namespace Horde.Agent.Utility
 	/// </summary>
 	public class PerforceLogger : ILogger
 	{
+		class ClientView
+		{
+			public readonly DirectoryReference BaseDir;
+			public readonly PerforceViewMap ViewMap;
+			public readonly int Change;
+
+			public ClientView(DirectoryReference baseDir, PerforceViewMap viewMap, int change)
+			{
+				BaseDir = baseDir;
+				ViewMap = viewMap;
+				Change = change;
+			}
+		}
+
 		readonly ILogger _inner;
-		readonly DirectoryReference _baseDir;
-		readonly PerforceViewMap? _viewMap;
-		readonly int? _viewChange;
+		readonly List<ClientView> _clients = new List<ClientView>();
 
 		/// <summary>
 		/// Constructor
 		/// </summary>
 		/// <param name="inner"></param>
-		/// <param name="baseDir"></param>
-		/// <param name="viewMap"></param>
-		/// <param name="viewChange"></param>
-		public PerforceLogger(ILogger inner, DirectoryReference baseDir, PerforceViewMap? viewMap, int? viewChange)
+		public PerforceLogger(ILogger inner)
 		{
 			_inner = inner;
-			_baseDir = baseDir;
-			_viewMap = viewMap;
-			_viewChange = viewChange;
+		}
+
+		/// <summary>
+		/// Adds a new client to be included in the mapping
+		/// </summary>
+		/// <param name="baseDir">Base directory for the client</param>
+		/// <param name="depotPath">Depot path for the workspace mapping, in the form //foo/bar...</param>
+		/// <param name="change">Changelist for the client</param>
+		public void AddClientView(DirectoryReference baseDir, string depotPath, int change)
+		{
+			PerforceViewMap viewMap = new PerforceViewMap();
+			viewMap.Entries.Add(new PerforceViewMapEntry(true, "...", depotPath));
+			AddClientView(baseDir, viewMap, change);
+		}
+
+		/// <summary>
+		/// Adds a new client to be included in the mapping
+		/// </summary>
+		public void AddClientView(DirectoryReference baseDir, PerforceViewMap viewMap, int change)
+		{
+			_clients.Add(new ClientView(baseDir, viewMap, change));
 		}
 
 		/// <inheritdoc/>
@@ -182,19 +210,20 @@ namespace Horde.Agent.Utility
 
 				file ??= Encoding.UTF8.GetString(text);
 
-				FileReference location = FileReference.Combine(_baseDir, file.Replace('\\', Path.DirectorySeparatorChar));
-				if (location.IsUnderDirectory(_baseDir))
+				foreach (ClientView client in _clients)
 				{
-					string relativePath = location.MakeRelativeTo(_baseDir).Replace('\\', '/');
-					annotations.Append($",\"relativePath\":\"{JsonEncodedText.Encode(relativePath)}\"");
-
-					if (_viewMap != null && _viewMap.TryMapFile(relativePath, StringComparison.OrdinalIgnoreCase, out string depotFile))
+					FileReference location = FileReference.Combine(client.BaseDir, file.Replace('\\', Path.DirectorySeparatorChar));
+					if (location.IsUnderDirectory(client.BaseDir))
 					{
-						if (_viewChange != null)
+						string relativePath = location.MakeRelativeTo(client.BaseDir).Replace('\\', '/');
+						annotations.Append($",\"relativePath\":\"{JsonEncodedText.Encode(relativePath)}\"");
+
+						if (client.ViewMap.TryMapFile(relativePath, StringComparison.OrdinalIgnoreCase, out string depotFile))
 						{
-							depotFile = $"{depotFile}@{_viewChange}";
+							depotFile = $"{depotFile}@{client.Change}";
+							annotations.Append($",\"depotPath\":\"{JsonEncodedText.Encode(depotFile)}\"");
+							break;
 						}
-						annotations.Append($",\"depotPath\":\"{JsonEncodedText.Encode(depotFile)}\"");
 					}
 				}
 
diff --git a/Engine/Source/Programs/Horde/Horde.Agent/Utility/WorkspaceInfo.cs b/Engine/Source/Programs/Horde/Horde.Agent/Utility/WorkspaceInfo.cs
--- a/Engine/Source/Programs/Horde/Horde.Agent/Utility/WorkspaceInfo.cs
+++ b/Engine/Source/Programs/Horde/Horde.Agent/Utility/WorkspaceInfo.cs
@@ -52,6 +52,11 @@ namespace Horde.Agent.Utility
 		/// </summary>
 		public string StreamName { get; }
 
+		/// <summary>
+		/// View for the stream
+		/// </summary>
+		public PerforceViewMap StreamView { get; }
+
 		/// <summary>
 		/// The directory containing metadata for this workspace
 		/// </summary>
@@ -83,12 +88,13 @@ namespace Horde.Agent.Utility
 		/// <param name="perforce">The perforce connection</param>
 		/// <param name="hostName">Name of this host</param>
 		/// <param name="streamName">Name of the stream to sync</param>
+		/// <param name="streamView">Stream view onto the depot</param>
 		/// <param name="metadataDir">Path to the metadata directory</param>
 		/// <param name="workspaceDir">Path to the workspace directory</param>
 		/// <param name="view">View for files to be synced</param>
 		/// <param name="removeUntrackedFiles">Whether to remove untracked files when syncing</param>
 		/// <param name="repository">The repository instance</param>
-		public WorkspaceInfo(IPerforceConnection perforce, string hostName, string streamName, DirectoryReference metadataDir, DirectoryReference workspaceDir, IList<string>? view, bool removeUntrackedFiles, ManagedWorkspace repository)
+		public WorkspaceInfo(IPerforceConnection perforce, string hostName, string streamName, PerforceViewMap streamView, DirectoryReference metadataDir, DirectoryReference workspaceDir, IList<string>? view, bool removeUntrackedFiles, ManagedWorkspace repository)
 		{
 			PerforceClient = perforce;
 
@@ -99,6 +105,7 @@ namespace Horde.Agent.Utility
 
 			HostName = hostName;
 			StreamName = streamName;
+			StreamView = streamView;
 			MetadataDir = metadataDir;
 			WorkspaceDir = workspaceDir;
 			View = (view == null) ? new List<string>() : new List<string>(view);
@@ -195,6 +202,10 @@ namespace Horde.Agent.Utility
 			// Create the client Perforce connection
 			IPerforceConnection perforceClient = await perforce.WithClientAsync(clientName);
 
+			// Get the view for this stream
+			StreamRecord stream = await perforceClient.GetStreamAsync(streamName, true, cancellationToken);
+			PerforceViewMap streamView = PerforceViewMap.Parse(stream.View);
+
 			// get the workspace names
 			DirectoryReference metadataDir = DirectoryReference.Combine(rootDir, identifier);
 			DirectoryReference workspaceDir = DirectoryReference.Combine(metadataDir, "Sync");
@@ -211,7 +222,7 @@ namespace Horde.Agent.Utility
 			await newRepository.RevertAsync(perforceClient, cancellationToken);
 
 			// Create the workspace info
-			return new WorkspaceInfo(perforceClient, hostName, streamName, metadataDir, workspaceDir, view, removeUntrackedFiles, newRepository);
+			return new WorkspaceInfo(perforceClient, hostName, streamName, streamView, metadataDir, workspaceDir, view, removeUntrackedFiles, newRepository);
 		}
 
 		/// <summary>
diff --git a/Engine/Source/Programs/Horde/Horde.Build.Tests/IssueServiceTests.cs b/Engine/Source/Programs/Horde/Horde.Build.Tests/IssueServiceTests.cs
--- a/Engine/Source/Programs/Horde/Horde.Build.Tests/IssueServiceTests.cs
+++ b/Engine/Source/Programs/Horde/Horde.Build.Tests/IssueServiceTests.cs
@@ -109,6 +109,7 @@ namespace Horde.Build.Tests
 		readonly UserId _timId;
 		readonly UserId _jerryId;
 		readonly UserId _bobId;
+		readonly DirectoryReference _autoSdkDir;
 		readonly DirectoryReference _workspaceDir;
 
 		async Task<IStream> CreateStreamAsync(ProjectId projectId, StreamId streamId, string streamName)
@@ -127,10 +128,12 @@ namespace Horde.Build.Tests
 		{
 			if (RuntimeInformation.IsOSPlatform(OSPlatform.Windows))
 			{
+				_autoSdkDir = new DirectoryReference("C:\\AutoSDK");
 				_workspaceDir = new DirectoryReference("C:\\Horde");
 			}
 			else
 			{
+				_autoSdkDir = new DirectoryReference("/AutoSdk");
 				_workspaceDir = new DirectoryReference("/Horde");
 			}
 
@@ -260,10 +263,10 @@ namespace Horde.Build.Tests
 
 			await using (TestJsonLogger logger = new TestJsonLogger(LogFileService, logId))
 			{
-				PerforceViewMap viewMap = new PerforceViewMap();
-				viewMap.Entries.Add(new PerforceViewMapEntry(true, "...", "//UE4/Main/..."));
+				PerforceLogger perforceLogger = new PerforceLogger(logger);
+				perforceLogger.AddClientView(_autoSdkDir, "//depot/CarefullyRedist/...", 12345);
+				perforceLogger.AddClientView(_workspaceDir, "//UE4/Main/...", 12345);
 
-				PerforceLogger perforceLogger = new PerforceLogger(logger, _workspaceDir, viewMap, 12345);
 				using (LogParser parser = new LogParser(perforceLogger, new List<string>()))
 				{
 					for (int idx = 0; idx < lines.Length; idx++)
@@ -383,7 +386,6 @@ namespace Horde.Build.Tests
 			}
 		}
 
-		[Ignore]
 		[TestMethod]
 		public async Task AutoSdkWarningTest()
 		{
@@ -407,8 +409,8 @@ namespace Horde.Build.Tests
 			{
 				string[] lines =
 				{
-					@"  D:\build\AutoSDK\Sync\HostWin64\GDK\200604\Microsoft GDK\200604\GRDK\ExtensionLibraries\Xbox.Game.Chat.2.Cpp.API\DesignTime\CommonConfiguration\Neutral\Include\GameChat2Impl.h(90): warning C5043: 'xbox::services::game_chat_2::chat_manager::set_memory_callbacks': exception specification does not match previous declaration",
-					@"  D:\build\AutoSDK\Sync\HostWin64\GDK\200604\Microsoft GDK\200604\GRDK\ExtensionLibraries\Xbox.Game.Chat.2.Cpp.API\DesignTime\CommonConfiguration\Neutral\Include\GameChat2.h(2083): note: see declaration of 'xbox::services::game_chat_2::chat_manager::set_memory_callbacks'",
+					FileReference.Combine(_autoSdkDir, @"HostWin64\GDK\200604\Microsoft GDK\200604\GRDK\ExtensionLibraries\Xbox.Game.Chat.2.Cpp.API\DesignTime\CommonConfiguration\Neutral\Include\GameChat2Impl.h").FullName + @"(90): warning C5043: 'xbox::services::game_chat_2::chat_manager::set_memory_callbacks': exception specification does not match previous declaration",
+					FileReference.Combine(_autoSdkDir, @"HostWin64\GDK\200604\Microsoft GDK\200604\GRDK\ExtensionLibraries\Xbox.Game.Chat.2.Cpp.API\DesignTime\CommonConfiguration\Neutral\Include\GameChat2.h").FullName + @"(2083): note: see declaration of 'xbox::services::game_chat_2::chat_manager::set_memory_callbacks'",
 				};
 
 				IJob job = CreateJob(_mainStreamId, 120, "Test Build", _graph);
@@ -1520,7 +1522,6 @@ namespace Horde.Build.Tests
 			{
 				stream = Deref(await StreamService.TryUpdateTemplateRefAsync(stream, templateRefId, new List<UpdateStepStateRequest>() { new UpdateStepStateRequest() { Name = "Update Version Files", QuarantinedByUserId = UserId.GenerateNewId().ToString() } }));
 
-
 				IJob job = CreateJob(_mainStreamId, 105, "Test Build", _graph);
 				await AddEvent(job, 0, 0, new { level = nameof(LogLevel.Warning) }, EventSeverity.Warning);
 				await UpdateCompleteStep(job, 0, 0, JobStepOutcome.Warnings);
@@ -1540,10 +1541,9 @@ namespace Horde.Build.Tests
 		{
 			await using (TestJsonLogger logger = new TestJsonLogger(LogFileService, logId))
 			{
-				PerforceViewMap viewMap = new PerforceViewMap();
-				viewMap.Entries.Add(new PerforceViewMapEntry(true, "...", "//UE4/Main/..."));
+				PerforceLogger perforceLogger = new PerforceLogger(logger);
+				perforceLogger.AddClientView(DirectoryReference.GetCurrentDirectory(), "//UE4/Main/...", 12345);
 
-				PerforceLogger perforceLogger = new PerforceLogger(logger, DirectoryReference.GetCurrentDirectory(), viewMap, 12345);
 				using (LogParser parser = new LogParser(perforceLogger, new List<string>()))
 				{
 					for (int idx = 0; idx < lines.Length; idx++)
diff --git a/Engine/Source/Programs/Shared/EpicGames.Perforce/PerforceViewMap.cs b/Engine/Source/Programs/Shared/EpicGames.Perforce/PerforceViewMap.cs
--- a/Engine/Source/Programs/Shared/EpicGames.Perforce/PerforceViewMap.cs
+++ b/Engine/Source/Programs/Shared/EpicGames.Perforce/PerforceViewMap.cs
@@ -53,6 +53,20 @@ namespace EpicGames.Perforce
 			return new PerforceViewMap(entries.Select(x => PerforceViewMapEntry.Parse(x)));
 		}
 
+		/// <summary>
+		/// Gets the inverse of this mapping
+		/// </summary>
+		/// <returns></returns>
+		public PerforceViewMap Invert()
+		{
+			List<PerforceViewMapEntry> entries = new List<PerforceViewMapEntry>();
+			foreach (PerforceViewMapEntry entry in Entries)
+			{
+				entries.Add(new PerforceViewMapEntry(entry.Include, entry.Target, entry.Source));
+			}
+			return new PerforceViewMap(entries);
+		}
+
 		/// <summary>
 		/// Determines if a file is included in the view
 		/// </summary>
