Horde: Add a suspect CL button and change step ranges to be for previous step
#rnx
#skipci

diff --git a/Engine/Source/Programs/Horde/HordeDashboard/src/backend/Dashboard.ts b/Engine/Source/Programs/Horde/HordeDashboard/src/backend/Dashboard.ts
--- a/Engine/Source/Programs/Horde/HordeDashboard/src/backend/Dashboard.ts
+++ b/Engine/Source/Programs/Horde/HordeDashboard/src/backend/Dashboard.ts
@@ -384,7 +384,12 @@ export class Dashboard {
             this.updating = true;
 
             if (!this.config) {
-                this.config = await backend.getDashboardConfig();                
+                try {
+                    this.config = await backend.getDashboardConfig();
+                } catch (reason) {
+                    console.error("Error getting dashboard config, defaults used: " + reason);
+                    this.config = {};
+                }
             }
 
             if (this.polling || !this.available) {
