Horde: Adding timestamps for quarantined and paused steps
#rnx

diff --git a/Engine/Source/Programs/Horde/Horde.Build/Api/Streams.cs b/Engine/Source/Programs/Horde/Horde.Build/Api/Streams.cs
--- a/Engine/Source/Programs/Horde/Horde.Build/Api/Streams.cs
+++ b/Engine/Source/Programs/Horde/Horde.Build/Api/Streams.cs
@@ -450,11 +450,21 @@ namespace Horde.Build.Api
 		/// </summary>
 		public string? PausedByUserId { get; set; }
 
+		/// <summary>
+		/// The UTC time when the step was paused
+		/// </summary>
+		public DateTime? PauseTimeUtc { get; set; }
+
 		/// <summary>
 		/// User who quarantined the step
 		/// </summary>
 		public string? QuarantinedByUserId { get; set; }
 
+		/// <summary>
+		/// The UTC time when the step was quarantined
+		/// </summary>
+		public DateTime? QuarantineTimeUtc { get; set; }
+
 		/// <summary>
 		/// Default constructor for serialization
 		/// </summary>
@@ -469,7 +479,9 @@ namespace Horde.Build.Api
 		{
 			Name = state.Name;
 			PausedByUserId = state.PausedByUserId?.ToString();
+			PauseTimeUtc = state.PauseTimeUtc;
 			QuarantinedByUserId = state.QuarantinedByUserId?.ToString();
+			QuarantineTimeUtc = state.QuarantineTimeUtc;
 		}
 
 	}
diff --git a/Engine/Source/Programs/Horde/Horde.Build/Projects/IStream.cs b/Engine/Source/Programs/Horde/Horde.Build/Projects/IStream.cs
--- a/Engine/Source/Programs/Horde/Horde.Build/Projects/IStream.cs
+++ b/Engine/Source/Programs/Horde/Horde.Build/Projects/IStream.cs
@@ -300,11 +300,22 @@ namespace Horde.Build.Models
 		/// </summary>
 		public UserId? PausedByUserId { get; set; }
 
+		/// <summary>
+		/// The UTC time when the step was paused
+		/// </summary>
+		public DateTime? PauseTimeUtc { get; set; }
+
 		/// <summary>
 		/// User who quarantined the step
 		/// </summary>
 		public UserId? QuarantinedByUserId { get; set; }
 
+		/// <summary>
+		/// The UTC time when the step was quarantined
+		/// </summary>
+		public DateTime? QuarantineTimeUtc { get; set; }
+
+
 		/// <summary>
 		/// Default constructor for serialization
 		/// </summary>
@@ -315,11 +326,13 @@ namespace Horde.Build.Models
 		/// <summary>
 		/// Constructor
 		/// </summary>
-		public TemplateStepState( string name, UserId? pausedByUserId = null, UserId? quarantinedByUserId = null)
+		public TemplateStepState( string name, UserId? pausedByUserId = null, DateTime? pauseTimeUtc = null, UserId? quarantinedByUserId = null, DateTime? quarantineTimerUtc = null)
 		{
 			Name = name;
 			PausedByUserId = pausedByUserId;
+			PauseTimeUtc = pauseTimeUtc;
 			QuarantinedByUserId = quarantinedByUserId;
+			QuarantineTimeUtc = quarantineTimerUtc;
 		}
 
 	}
diff --git a/Engine/Source/Programs/Horde/Horde.Build/Projects/Impl/StreamCollection.cs b/Engine/Source/Programs/Horde/Horde.Build/Projects/Impl/StreamCollection.cs
--- a/Engine/Source/Programs/Horde/Horde.Build/Projects/Impl/StreamCollection.cs
+++ b/Engine/Source/Programs/Horde/Horde.Build/Projects/Impl/StreamCollection.cs
@@ -314,11 +314,30 @@ namespace Horde.Build.Collections.Impl
 							continue;
 						}
 
-						newStepStates.Add(new TemplateStepState(updateState.Name, PausedByUserId, QuarantinedByUserId));
+						newStepStates.Add(new TemplateStepState(updateState.Name, PausedByUserId, PausedByUserId != null ? DateTime.UtcNow : null, QuarantinedByUserId, QuarantinedByUserId != null ? DateTime.UtcNow : null));
 					}
 					else
 					{
+						if (PausedByUserId == null)
+						{
+							newState.PauseTimeUtc = null;
+						}
+						else if (newState.PauseTimeUtc == null)
+						{
+							newState.PauseTimeUtc = DateTime.UtcNow;
+						}
+
 						newState.PausedByUserId = PausedByUserId;
+
+						if (QuarantinedByUserId == null)
+						{
+							newState.QuarantineTimeUtc = null;
+						}
+						else if (newState.QuarantineTimeUtc == null)
+						{
+							newState.QuarantineTimeUtc = DateTime.UtcNow;
+						}
+
 						newState.QuarantinedByUserId = QuarantinedByUserId;
 					}
 				}
diff --git a/Engine/Source/Programs/Horde/HordeDashboard/src/backend/Api.ts b/Engine/Source/Programs/Horde/HordeDashboard/src/backend/Api.ts
--- a/Engine/Source/Programs/Horde/HordeDashboard/src/backend/Api.ts
+++ b/Engine/Source/Programs/Horde/HordeDashboard/src/backend/Api.ts
@@ -1786,8 +1786,14 @@ export type GetTemplateStepStateResponse = {
 	/**User who paused the step */
 	pausedByUserId?: string;
 
+	/**The UTC time when the step was paused*/
+	pauseTimeUtc?: Date | string;
+
 	/**User who quarantined the step */
 	quarantinedByUserId?: string; 
+	
+	/**The UTC time when the step was quarantined */
+	quarantineTimeUtc?: Date | string;
 
 }
 
