Horde: Missing files.

diff --git a/Engine/Source/Programs/Horde/Horde.Agent.Tests/Helper.cs b/Engine/Source/Programs/Horde/Horde.Agent.Tests/Helper.cs
--- a/Engine/Source/Programs/Horde/Horde.Agent.Tests/Helper.cs
+++ b/Engine/Source/Programs/Horde/Horde.Agent.Tests/Helper.cs
@@ -158,14 +158,14 @@ namespace Horde.Agent.Tests
 
 		public override AsyncUnaryCall<Empty> WriteOutputAsync(WriteOutputRequest request, CallOptions options)
 		{
-			_logger.LogDebug("WriteOutputAsync: " + request.Data);
+			_logger.LogDebug("WriteOutputAsync: {Data}", request.Data);
 			Empty res = new Empty();
 			return Wrap(res);
 		}
 
 		public override AsyncUnaryCall<Empty> UpdateStepAsync(UpdateStepRequest request, CallOptions options)
 		{
-			_logger.LogDebug($"UpdateStepAsync(Request: {request})");
+			_logger.LogDebug("UpdateStepAsync(Request: {Request})", request);
 			UpdateStepRequests.Add(request);
 			Empty res = new Empty();
 			return Wrap(res);
@@ -173,7 +173,7 @@ namespace Horde.Agent.Tests
 
 		public override AsyncUnaryCall<Empty> CreateEventsAsync(CreateEventsRequest request, CallOptions options)
 		{
-			_logger.LogDebug("CreateEventsAsync: " + request);
+			_logger.LogDebug("CreateEventsAsync: {Request}", request);
 			Empty res = new Empty();
 			return Wrap(res);
 		}
@@ -193,7 +193,7 @@ namespace Horde.Agent.Tests
 			return Wrap(new GetStepResponse());
 		}
 
-		private AsyncUnaryCall<T> Wrap<T>(T res)
+		private static AsyncUnaryCall<T> Wrap<T>(T res)
 		{
 			return new AsyncUnaryCall<T>(Task.FromResult(res), Task.FromResult(Metadata.Empty),
 				() => Status.DefaultSuccess, () => Metadata.Empty, null!);
@@ -218,7 +218,7 @@ namespace Horde.Agent.Tests
 		public Task<JobStepOutcome> RunAsync(BeginStepResponse step, ILogger logger,
 			CancellationToken cancellationToken)
 		{
-			logger.LogDebug($"SimpleTestExecutor.RunAsync(Step: {step})");
+			logger.LogDebug("SimpleTestExecutor.RunAsync(Step: {Step})", step);
 			return _func(step, logger, cancellationToken);
 		}
 
diff --git a/Engine/Source/Programs/Horde/Horde.Agent.Tests/LogParserTests.cs b/Engine/Source/Programs/Horde/Horde.Agent.Tests/LogParserTests.cs
--- a/Engine/Source/Programs/Horde/Horde.Agent.Tests/LogParserTests.cs
+++ b/Engine/Source/Programs/Horde/Horde.Agent.Tests/LogParserTests.cs
@@ -66,12 +66,12 @@ namespace Horde.Agent.Tests
 		public void JsonLoggerTest()
 		{
 			JsonLoggerImpl impl = new JsonLoggerImpl(null, NullLogger.Instance);
-			impl.LogInformation("Hello {0}", "world");
+			impl.LogInformation("Hello {Text}", "world");
 			impl.LogInformation("Hello {Text}", "world");
 
 			Assert.AreEqual(impl._lines.Count, 2);
 			Assert.AreEqual(impl._lines[0].ToString(), "Hello world");
-			Assert.AreEqual(impl._lines[0].Properties!["0"].ToString()!, "world");
+			Assert.AreEqual(impl._lines[0].Properties!["Text"].ToString()!, "world");
 			Assert.AreEqual(impl._lines[1].ToString(), "Hello world");
 			Assert.AreEqual(impl._lines[1].Properties!["Text"].ToString(), "world");
 		}
@@ -152,7 +152,7 @@ namespace Horde.Agent.Tests
 			}
 
 			{
-				List<LogEvent> logEvents = Parse(String.Join("\n", lines).Replace("Error:", "Warning:"));
+				List<LogEvent> logEvents = Parse(String.Join("\n", lines).Replace("Error:", "Warning:", StringComparison.Ordinal));
 				CheckEventGroup(logEvents, 0, 14, LogLevel.Warning, KnownLogEvents.Engine_Crash);
 			}
 		}
@@ -1032,28 +1032,25 @@ namespace Horde.Agent.Tests
 			}
 
 			{
-				List<LogEvent> logEvents = Parse(String.Join("\n", lines).Replace("Error:", "Warning:"));
+				List<LogEvent> logEvents = Parse(String.Join("\n", lines).Replace("Error:", "Warning:", StringComparison.Ordinal));
 				Assert.AreEqual(4, logEvents.Count);
 				CheckEventGroup(logEvents.Slice(0, 3), 1, 3, LogLevel.Warning, KnownLogEvents.Engine_LogChannel);
 				CheckEventGroup(logEvents.Slice(3, 1), 4, 1, LogLevel.Information, KnownLogEvents.Engine_LogChannel);
 			}
 		}
 
-		List<LogEvent> Parse(IEnumerable<string> lines)
+		static List<LogEvent> Parse(IEnumerable<string> lines)
 		{
 			return Parse(String.Join("\n", lines));
 		}
 
-		List<LogEvent> Parse(string text)
+		static List<LogEvent> Parse(string text)
 		{
 			return Parse(text, new DirectoryReference("C:\\Horde".Replace('\\', Path.DirectorySeparatorChar)));
 		}
 
-		List<LogEvent> Parse(string text, DirectoryReference workspaceDir)
+		static List<LogEvent> Parse(string text, DirectoryReference workspaceDir)
 		{
-			PerforceViewMap viewMap = new PerforceViewMap();
-			viewMap.Entries.Add(new PerforceViewMapEntry(true, "...", "//UE4/Main/..."));
-
 			List<string> ignorePatterns = new List<string>();
 
 			byte[] textBytes = Encoding.UTF8.GetBytes(text);
@@ -1061,7 +1058,11 @@ namespace Horde.Agent.Tests
 			Random generator = new Random(0);
 
 			LoggerCapture logger = new LoggerCapture();
-			using (LogParser parser = new LogParser(new PerforceLogger(logger, workspaceDir, viewMap, 12345), ignorePatterns))
+
+			PerforceLogger perforceLogger = new PerforceLogger(logger);
+			perforceLogger.AddClientView(workspaceDir, "//UE4/Main/...", 12345);
+
+			using (LogParser parser = new LogParser(perforceLogger, ignorePatterns))
 			{
 				int pos = 0;
 				while(pos < textBytes.Length)
@@ -1074,7 +1075,7 @@ namespace Horde.Agent.Tests
 			return logger._events;
 		}
 
-		void CheckEventGroup(IEnumerable<LogEvent> logEvents, int index, int count, LogLevel level, EventId eventId = default)
+		static void CheckEventGroup(IEnumerable<LogEvent> logEvents, int index, int count, LogLevel level, EventId eventId = default)
 		{
 			IEnumerator<LogEvent> enumerator = logEvents.GetEnumerator();
 			for (int idx = 0; idx < count; idx++)
diff --git a/Engine/Source/Programs/Horde/Horde.Agent.Tests/WorkerServiceTest.cs b/Engine/Source/Programs/Horde/Horde.Agent.Tests/WorkerServiceTest.cs
--- a/Engine/Source/Programs/Horde/Horde.Agent.Tests/WorkerServiceTest.cs
+++ b/Engine/Source/Programs/Horde/Horde.Agent.Tests/WorkerServiceTest.cs
@@ -37,7 +37,7 @@ namespace Horde.Agent.Tests
 			profile.Name = "test";
 			profile.Environment = "test-env";
 			profile.Token = "bogus-token";
-			profile.Url = "http://localhost";
+			profile.Url = new Uri("http://localhost");
 			settings.ServerProfiles.Add(profile);
 			settings.Server = "test";
 			settings.WorkingDir = Path.GetTempPath();
@@ -50,7 +50,7 @@ namespace Horde.Agent.Tests
 		[TestMethod]
 		public async Task AbortExecuteStepTest()
 		{
-			WorkerService ws = GetWorkerService();
+			using WorkerService ws = GetWorkerService();
 
 			{
 				using CancellationTokenSource cancelSource = new CancellationTokenSource();
@@ -62,7 +62,7 @@ namespace Horde.Agent.Tests
 					await Task.Delay(5000, cancelToken);
 					return JobStepOutcome.Success;
 				});
-				await Assert.ThrowsExceptionAsync<TaskCanceledException>(() => ws.ExecuteStepAsync(executor,
+				await Assert.ThrowsExceptionAsync<TaskCanceledException>(() => WorkerService.ExecuteStepAsync(executor,
 					new BeginStepResponse(), _workerLogger, cancelSource.Token, stepCancelSource.Token));
 			}
 
@@ -76,7 +76,7 @@ namespace Horde.Agent.Tests
 					await Task.Delay(5000, cancelToken);
 					return JobStepOutcome.Success;
 				});
-				(JobStepOutcome stepOutcome, JobStepState stepState) = await ws.ExecuteStepAsync(executor, new BeginStepResponse(), _workerLogger,
+				(JobStepOutcome stepOutcome, JobStepState stepState) = await WorkerService.ExecuteStepAsync(executor, new BeginStepResponse(), _workerLogger,
 					cancelSource.Token, stepCancelSource.Token);
 				Assert.AreEqual(JobStepOutcome.Failure, stepOutcome);
 				Assert.AreEqual(JobStepState.Aborted, stepState);
@@ -98,7 +98,7 @@ namespace Horde.Agent.Tests
 			executeJobTask.Workspace = new AgentWorkspace();
 
 			HordeRpcClientStub client = new HordeRpcClientStub(_workerLogger);
-			RpcConnectionStub rpcConnection = new RpcConnectionStub(null!, client);
+			await using RpcConnectionStub rpcConnection = new RpcConnectionStub(null!, client);
 
 			client.BeginStepResponses.Enqueue(new BeginStepResponse {Name = "stepName1", StepId = "stepId1"});
 			client.BeginStepResponses.Enqueue(new BeginStepResponse {Name = "stepName2", StepId = "stepId2"});
@@ -114,7 +114,7 @@ namespace Horde.Agent.Tests
 				return JobStepOutcome.Success;
 			});
 
-			WorkerService ws = GetWorkerService((a, b, c) => executor);
+			using WorkerService ws = GetWorkerService((a, b, c) => executor);
 			ws._stepAbortPollInterval = TimeSpan.FromMilliseconds(1);
 			LeaseOutcome outcome = (await ws.ExecuteJobAsync(rpcConnection, "leaseId1", executeJobTask,
 				_workerLogger, token)).Outcome;
@@ -142,11 +142,11 @@ namespace Horde.Agent.Tests
 				return JobStepOutcome.Success;
 			});
 
-			WorkerService ws = GetWorkerService((a, b, c) => executor);
+			using WorkerService ws = GetWorkerService((a, b, c) => executor);
 			ws._stepAbortPollInterval = TimeSpan.FromMilliseconds(5);
 			
 			HordeRpcClientStub client = new HordeRpcClientStub(_workerLogger);
-			RpcConnectionStub rpcConnection = new RpcConnectionStub(null!, client);
+			await using RpcConnectionStub rpcConnection = new RpcConnectionStub(null!, client);
 
 			int c = 0;
 			client._getStepFunc = (request) =>
@@ -164,7 +164,7 @@ namespace Horde.Agent.Tests
 			using CancellationTokenSource stepCancelSource = new CancellationTokenSource();
 			TaskCompletionSource<bool> stepFinishedSource = new TaskCompletionSource<bool>();
 
-			await ws.PollForStepAbort(rpcConnection, "jobId1", "batchId1", "logId1", stepPollCancelSource.Token, stepCancelSource, stepFinishedSource.Task);
+			await ws.PollForStepAbort(rpcConnection, "jobId1", "batchId1", "logId1", stepCancelSource, stepFinishedSource.Task, stepPollCancelSource.Token);
 			Assert.IsTrue(stepCancelSource.IsCancellationRequested);
 		}
 	}
