Horde: Fix exception parsing line count from formatted JSON lines.

diff --git a/Engine/Source/Programs/Horde/Horde.Agent/Utility/JsonRpcLogger.cs b/Engine/Source/Programs/Horde/Horde.Agent/Utility/JsonRpcLogger.cs
--- a/Engine/Source/Programs/Horde/Horde.Agent/Utility/JsonRpcLogger.cs
+++ b/Engine/Source/Programs/Horde/Horde.Agent/Utility/JsonRpcLogger.cs
@@ -142,11 +142,16 @@ namespace Horde.Agent.Parser
 			int lineCount = 1;
 			while(reader.Read() && reader.TokenType == JsonTokenType.PropertyName)
 			{
-				if(reader.ValueTextEquals(LogEventPropertyName.Line) && reader.TryGetInt32(out int lineIndex) && lineIndex > 0)
+				ReadOnlySpan<byte> propertyName = reader.ValueSpan;
+				if (!reader.Read())
+				{
+					break;
+				}
+				else if(propertyName.SequenceEqual(LogEventPropertyName.Line) && reader.TryGetInt32(out int lineIndex) && lineIndex > 0)
 				{
 					return 0;
 				}
-				else if(reader.ValueTextEquals(LogEventPropertyName.LineCount) && reader.TryGetInt32(out int lineCountValue))
+				else if(propertyName.SequenceEqual(LogEventPropertyName.LineCount) && reader.TryGetInt32(out int lineCountValue))
 				{
 					lineCount = lineCountValue;
 				}
