Horde: Set an environment variable to enable JSON logging to stdout.

diff --git a/Engine/Source/Programs/Horde/Horde.Agent/Execution/BuildGraphExecutor.cs b/Engine/Source/Programs/Horde/Horde.Agent/Execution/BuildGraphExecutor.cs
--- a/Engine/Source/Programs/Horde/Horde.Agent/Execution/BuildGraphExecutor.cs
+++ b/Engine/Source/Programs/Horde/Horde.Agent/Execution/BuildGraphExecutor.cs
@@ -796,6 +796,9 @@ namespace Horde.Agent.Execution
 			newEnvironment["UE_HORDE_BATCHID"] = _batchId;
 			newEnvironment["UE_HORDE_STEPID"] = step.StepId;
 
+			// Request structured JSON output directly from tools
+			newEnvironment["UE_STDOUT_JSON"] = "1";
+
 			// Enable structured logging output
 			newEnvironment["UE_LOG_JSON"] = "1";
 
