Horde: Change template step state api to use thin user responses
#rnx
#skipci

diff --git a/Engine/Source/Programs/Horde/Horde.Build/Api/Streams.cs b/Engine/Source/Programs/Horde/Horde.Build/Api/Streams.cs
--- a/Engine/Source/Programs/Horde/Horde.Build/Api/Streams.cs
+++ b/Engine/Source/Programs/Horde/Horde.Build/Api/Streams.cs
@@ -448,7 +448,7 @@ namespace Horde.Build.Api
 		/// <summary>
 		/// User who paused the step
 		/// </summary>
-		public string? PausedByUserId { get; set; }
+		public GetThinUserInfoResponse? PausedByUserInfo { get; set; }
 
 		/// <summary>
 		/// The UTC time when the step was paused
@@ -458,7 +458,7 @@ namespace Horde.Build.Api
 		/// <summary>
 		/// User who quarantined the step
 		/// </summary>
-		public string? QuarantinedByUserId { get; set; }
+		public GetThinUserInfoResponse? QuarantinedByUserInfo { get; set; }
 
 		/// <summary>
 		/// The UTC time when the step was quarantined
@@ -475,13 +475,13 @@ namespace Horde.Build.Api
 		/// <summary>
 		/// Constructor
 		/// </summary>
-		public GetTemplateStepStateResponse(TemplateStepState state)
+		public GetTemplateStepStateResponse(TemplateStepState state, GetThinUserInfoResponse? pausedByUserInfo, GetThinUserInfoResponse? quarantinedByUserInfo)
 		{
-			Name = state.Name;
-			PausedByUserId = state.PausedByUserId?.ToString();
+			Name = state.Name;			
 			PauseTimeUtc = state.PauseTimeUtc;
-			QuarantinedByUserId = state.QuarantinedByUserId?.ToString();
 			QuarantineTimeUtc = state.QuarantineTimeUtc;
+			PausedByUserInfo = pausedByUserInfo;
+			QuarantinedByUserInfo = quarantinedByUserInfo;
 		}
 
 	}
@@ -553,8 +553,9 @@ namespace Horde.Build.Api
 		/// <param name="id">The template ref id</param>
 		/// <param name="templateRef">The template ref</param>
 		/// <param name="template">The template instance</param>
+		/// <param name="stepStates">The template step states</param>
 		/// <param name="bIncludeAcl">Whether to include the ACL in the response</param>
-		public GetTemplateRefResponse(TemplateRefId id, TemplateRef templateRef, ITemplate template, bool bIncludeAcl)
+		public GetTemplateRefResponse(TemplateRefId id, TemplateRef templateRef, ITemplate template, List<GetTemplateStepStateResponse>? stepStates, bool bIncludeAcl)
 			: base(template)
 		{
 			Id = id.ToString();
@@ -565,7 +566,7 @@ namespace Horde.Build.Api
 			NotificationChannelFilter = templateRef.NotificationChannelFilter;
 			Schedule = (templateRef.Schedule != null) ? new GetScheduleResponse(templateRef.Schedule) : null;
 			ChainedJobs = (templateRef.ChainedJobs != null && templateRef.ChainedJobs.Count > 0) ? templateRef.ChainedJobs.ConvertAll(x => new GetChainedJobTemplateResponse(x)) : null;
-			StepStates = (templateRef.StepStates != null && templateRef.StepStates.Count > 0) ? templateRef.StepStates.ConvertAll(x => new GetTemplateStepStateResponse(x)) : null;
+			StepStates = stepStates;
 			Acl = (bIncludeAcl && templateRef.Acl != null)? new GetAclResponse(templateRef.Acl) : null;
 		}
 	}
diff --git a/Engine/Source/Programs/Horde/Horde.Build/Controllers/StreamsController.cs b/Engine/Source/Programs/Horde/Horde.Build/Controllers/StreamsController.cs
--- a/Engine/Source/Programs/Horde/Horde.Build/Controllers/StreamsController.cs
+++ b/Engine/Source/Programs/Horde/Horde.Build/Controllers/StreamsController.cs
@@ -31,15 +31,17 @@ namespace Horde.Build.Controllers
 		private readonly ITemplateCollection _templateCollection;
 		private readonly IJobStepRefCollection _jobStepRefCollection;
 		private readonly IPerforceService _perforceService;
+		private readonly IUserCollection _userCollection;
 		private readonly AclService _aclService;
 		/// <summary>
 		/// Constructor
 		/// </summary>
-		public StreamsController(AclService aclService, StreamService streamService, ITemplateCollection templateCollection, IJobStepRefCollection jobStepRefCollection, IPerforceService perforceService)
+		public StreamsController(AclService aclService, StreamService streamService, ITemplateCollection templateCollection, IJobStepRefCollection jobStepRefCollection, IUserCollection userCollection, IPerforceService perforceService)
 		{
 			_streamService = streamService;
 			_templateCollection = templateCollection;
 			_jobStepRefCollection = jobStepRefCollection;
+			_userCollection = userCollection;
 			_perforceService = perforceService;
 			_aclService = aclService;
 		}
@@ -114,10 +116,44 @@ namespace Horde.Build.Controllers
 				if (await _streamService.AuthorizeAsync(stream, pair.Value, AclAction.ViewTemplate, User, cache))
 				{
 					ITemplate? template = await _templateCollection.GetAsync(pair.Value.Hash);
-					if (template != null)
+					
+					if (template != null)						
 					{
+						TemplateRef tref = pair.Value;
+
+						List<GetTemplateStepStateResponse>? stepStates = null;
+						if (tref.StepStates != null)
+						{
+							for (int i = 0; i < tref.StepStates.Count; i++)
+							{
+								TemplateStepState state = tref.StepStates[i];
+
+								if (state.QuarantinedByUserId == null && state.PausedByUserId == null)
+								{
+									continue;
+								}
+
+								stepStates ??= new List<GetTemplateStepStateResponse>();
+
+								GetThinUserInfoResponse? pausedByUserInfo = null;
+								if (state.PausedByUserId != null)
+								{
+									pausedByUserInfo = new GetThinUserInfoResponse(await _userCollection.GetCachedUserAsync(state.PausedByUserId));
+								}
+
+								GetThinUserInfoResponse? quarantinedByUserInfo = null;
+								if (state.QuarantinedByUserId != null)
+								{
+									quarantinedByUserInfo = new GetThinUserInfoResponse(await _userCollection.GetCachedUserAsync(state.QuarantinedByUserId));
+								}
+
+								stepStates.Add(new GetTemplateStepStateResponse(state, pausedByUserInfo, quarantinedByUserInfo));
+
+							}
+						}
+
 						bool bIncludeTemplateAcl = pair.Value.Acl != null && await _streamService.AuthorizeAsync(stream, pair.Value, AclAction.ViewPermissions, User, cache);
-						apiTemplateRefs.Add(new GetTemplateRefResponse(pair.Key, pair.Value, template, bIncludeTemplateAcl));
+						apiTemplateRefs.Add(new GetTemplateRefResponse(pair.Key, pair.Value, template, stepStates, bIncludeTemplateAcl));
 					}
 				}
 			}
