Horde: Fix parsing of multi-line log messages.

diff --git a/Engine/Source/Programs/Horde/Horde.Agent/Utility/JsonRpcLogger.cs b/Engine/Source/Programs/Horde/Horde.Agent/Utility/JsonRpcLogger.cs
--- a/Engine/Source/Programs/Horde/Horde.Agent/Utility/JsonRpcLogger.cs
+++ b/Engine/Source/Programs/Horde/Horde.Agent/Utility/JsonRpcLogger.cs
@@ -45,7 +45,7 @@ namespace Horde.Agent.Parser
 		readonly string? _jobStepId;
 		readonly bool _warnings;
 		readonly ILogger _inner;
-		readonly Channel<QueueItem> _dataChannel;
+		readonly Channel<JsonLogEvent> _dataChannel;
 		Task? _dataWriter;
 
 		/// <summary>
@@ -76,7 +76,7 @@ namespace Horde.Agent.Parser
 			_jobStepId = jobStepId;
 			_warnings = warnings ?? true;
 			_inner = inner;
-			_dataChannel = Channel.CreateUnbounded<QueueItem>();
+			_dataChannel = Channel.CreateUnbounded<JsonLogEvent>();
 			_dataWriter = Task.Run(() => RunDataWriter());
 			Outcome = JobStepOutcome.Success;
 		}
@@ -90,8 +90,8 @@ namespace Horde.Agent.Parser
 				logLevel = LogLevel.Information;
 			}
 
-			LogEvent logEvent = LogEvent.FromState(logLevel, eventId, state, exception, formatter);
-			WriteFormattedEvent(logLevel, logEvent.ToJsonBytes());
+			JsonLogEvent jsonLogEvent = JsonLogEvent.FromLoggerState(logLevel, eventId, state, exception, formatter);
+			WriteFormattedEvent(jsonLogEvent);
 		}
 
 		/// <inheritdoc/>
@@ -100,9 +100,10 @@ namespace Horde.Agent.Parser
 		/// <inheritdoc/>
 		public IDisposable BeginScope<TState>(TState state) => _inner.BeginScope(state);
 
-		private void WriteFormattedEvent(LogLevel level, byte[] line)
+		private void WriteFormattedEvent(JsonLogEvent jsonLogEvent)
 		{
 			// Update the state of this job if this is an error status
+			LogLevel level = jsonLogEvent.Level;
 			if (level == LogLevel.Error || level == LogLevel.Critical)
 			{
 				Outcome = JobStepOutcome.Failure;
@@ -112,26 +113,14 @@ namespace Horde.Agent.Parser
 				Outcome = JobStepOutcome.Warnings;
 			}
 
-			// If we want an event for this log event, create one now
-			CreateEventRequest? eventRequest = null;
-			if (level == LogLevel.Warning || level == LogLevel.Error || level == LogLevel.Critical)
-			{
-				int lineCount = GetEventLineCount(line);
-				if (lineCount > 0)
-				{
-					eventRequest = CreateEvent(level, lineCount);
-				}
-			}
-
-			// Write the data to the output channel
-			QueueItem queueItem = new QueueItem(line, eventRequest);
-			if (!_dataChannel.Writer.TryWrite(queueItem))
+			// Write the event
+			if (!_dataChannel.Writer.TryWrite(jsonLogEvent))
 			{
 				throw new InvalidOperationException("Expected unbounded writer to complete immediately");
 			}
 		}
 
-		static int GetEventLineCount(byte[] line)
+		static int GetEventLineCount(ReadOnlySpan<byte> line)
 		{
 			Utf8JsonReader reader = new Utf8JsonReader(line);
 			if(!reader.Read() || reader.TokenType != JsonTokenType.StartObject)
@@ -234,19 +223,22 @@ namespace Horde.Agent.Parser
 				Task waitTask = Task.Delay(TimeSpan.FromSeconds(2.0));
 				while (writer.WrittenCount < 256 * 1024)
 				{
-					QueueItem? data;
-					if (_dataChannel.Reader.TryRead(out data))
+					JsonLogEvent jsonLogEvent;
+					if (_dataChannel.Reader.TryRead(out jsonLogEvent))
 					{
-						if (data.CreateEvent != null)
+						// If we want an event for this log event, create one now
+						if (jsonLogEvent.Level == LogLevel.Warning && ++numWarnings <= MaxWarnings)
+						{
+							int lineCount = GetEventLineCount(jsonLogEvent.Data.Span);
+							events.Add(new CreateEventRequest(EventSeverity.Warning, _logId, 0, lineCount));
+						}
+						else if ((jsonLogEvent.Level == LogLevel.Error || jsonLogEvent.Level == LogLevel.Critical) && ++numErrors <= MaxErrors)
 						{
-							if ((data.CreateEvent.Severity == EventSeverity.Warning && ++numWarnings <= MaxWarnings) || (data.CreateEvent.Severity == EventSeverity.Error && ++numErrors <= MaxErrors))
-							{
-								data.CreateEvent.LineIndex = lineIndex;
-								events.Add(data.CreateEvent);
-							}
+							int lineCount = GetEventLineCount(jsonLogEvent.Data.Span);
+							events.Add(new CreateEventRequest(EventSeverity.Error, _logId, 0, lineCount));
 						}
 
-						writer.Write(data.Data);
+						writer.Write(jsonLogEvent.Data.Span);
 						writer.Write(newline);
 
 						lineIndex++;
