Horde: Expose template ID and name as env vars during job execution

diff --git a/Engine/Source/Programs/Horde/Horde.Build/Services/RpcService.cs b/Engine/Source/Programs/Horde/Horde.Build/Services/RpcService.cs
--- a/Engine/Source/Programs/Horde/Horde.Build/Services/RpcService.cs
+++ b/Engine/Source/Programs/Horde/Horde.Build/Services/RpcService.cs
@@ -59,6 +59,7 @@ namespace Horde.Build.Services
 		readonly IGraphCollection _graphs;
 		readonly ITestDataCollection _testData;
 		readonly IJobStepRefCollection _jobStepRefCollection;
+		readonly ITemplateCollection _templateCollection;
 		readonly ConformTaskSource _conformTaskSource;
 		readonly IClock _clock;
 		readonly ILogger<RpcService> _logger;
@@ -66,7 +67,7 @@ namespace Horde.Build.Services
 		/// <summary>
 		/// Constructor
 		/// </summary>
-		public RpcService(AclService aclService, AgentService agentService, StreamService streamService, JobService jobService, AgentSoftwareService agentSoftwareService, IArtifactCollection artifactCollection, ILogFileService logFileService, PoolService poolService, LifetimeService lifetimeService, IGraphCollection graphs, ITestDataCollection testData, IJobStepRefCollection jobStepRefCollection, ConformTaskSource conformTaskSource, IClock clock, ILogger<RpcService> logger)
+		public RpcService(AclService aclService, AgentService agentService, StreamService streamService, JobService jobService, AgentSoftwareService agentSoftwareService, IArtifactCollection artifactCollection, ILogFileService logFileService, PoolService poolService, LifetimeService lifetimeService, IGraphCollection graphs, ITestDataCollection testData, IJobStepRefCollection jobStepRefCollection, ITemplateCollection templateCollection, ConformTaskSource conformTaskSource, IClock clock, ILogger<RpcService> logger)
 		{
 			_aclService = aclService;
 			_agentService = agentService;
@@ -80,6 +81,7 @@ namespace Horde.Build.Services
 			_graphs = graphs;
 			_testData = testData;
 			_jobStepRefCollection = jobStepRefCollection;
+			_templateCollection = templateCollection;
 			_conformTaskSource = conformTaskSource;
 			_clock = clock;
 			_logger = logger;
@@ -598,6 +600,16 @@ namespace Horde.Build.Services
 				response.Name = node.Name;
 				response.Credentials.Add(credentials);
 
+				string templateName = "<unknown>";
+				if (job.TemplateHash != null)
+				{
+					ITemplate? template = await _templateCollection.GetAsync(job.TemplateHash);
+					templateName = template != null ? template.Name : templateName;
+				}
+				
+				response.EnvVars.Add("UE_HORDE_TEMPLATEID", job.TemplateId.ToString());
+				response.EnvVars.Add("UE_HORDE_TEMPLATENAME", templateName);
+
 				IJobStepRef? lastStep = await _jobStepRefCollection.GetPrevStepForNodeAsync(job.StreamId, job.TemplateId, node.Name, job.Change);
 				if (lastStep != null)
 				{
