Horde: Add an agent command which builds UAT runs a BuildGraph script, parsing all the log output through the regular agent path.

diff --git a/Engine/Source/Programs/Horde/Horde.Agent/Commands/Utilities/BuildGraphCommand.cs b/Engine/Source/Programs/Horde/Horde.Agent/Commands/Utilities/BuildGraphCommand.cs
new file mode 100644
--- /dev/null
+++ b/Engine/Source/Programs/Horde/Horde.Agent/Commands/Utilities/BuildGraphCommand.cs
@@ -0,0 +1,155 @@
+// Copyright Epic Games, Inc. All Rights Reserved.
+
+using System;
+using System.Collections.Generic;
+using System.Diagnostics;
+using System.IO;
+using System.Text;
+using System.Text.Json;
+using System.Text.Json.Serialization;
+using System.Threading;
+using System.Threading.Tasks;
+using EpicGames.Core;
+using Horde.Agent.Parser;
+using HordeCommon;
+using HordeCommon.Rpc;
+using Microsoft.Extensions.Logging;
+
+namespace Horde.Agent.Commands.Utilities
+{
+	/// <summary>
+	/// Runs a BuildGraph script and captures the processed log output
+	/// </summary>
+	[Command("BuildGraph", "Executes a BuildGraph script with the given arguments using a build of UAT within the current branch, and runs the output through the log processor")]
+	class BuildGraphCommand : Command
+	{
+		sealed class LogSink : IJsonRpcLogSink, IDisposable
+		{
+			readonly ILogger _inner;
+			readonly FileStream _eventStream;
+			readonly FileStream _outputStream;
+
+			public LogSink(DirectoryReference baseDir, ILogger inner)
+			{
+				_inner = inner;
+
+				FileReference eventsFile = FileReference.Combine(baseDir, "horde-events.txt");
+				_eventStream = FileReference.Open(eventsFile, FileMode.Create);
+				inner.LogInformation("Writing events to {File}", eventsFile);
+
+				FileReference outputFile = FileReference.Combine(baseDir, "horde-output.txt");
+				_outputStream = FileReference.Open(outputFile, FileMode.Create);
+				inner.LogInformation("Writing output to {File}", outputFile);
+			}
+
+			public void Dispose()
+			{
+				_eventStream.Dispose();
+				_outputStream.Dispose();
+			}
+
+			public Task SetOutcomeAsync(JobStepOutcome outcome) => Task.CompletedTask;
+
+			public async Task WriteEventsAsync(List<CreateEventRequest> events)
+			{
+				foreach (CreateEventRequest request in events)
+				{
+					JsonSerializerOptions options = new JsonSerializerOptions();
+					options.Converters.Add(new JsonStringEnumConverter());
+					await JsonSerializer.SerializeAsync(_eventStream, request, options);
+					_eventStream.Write(Encoding.UTF8.GetBytes(Environment.NewLine));
+				}
+			}
+
+			public async Task WriteOutputAsync(WriteOutputRequest request)
+			{
+				PrintJson(request.Data.Span);
+				await _outputStream.WriteAsync(request.Data.Memory);
+			}
+
+			void PrintJson(ReadOnlySpan<byte> span)
+			{
+				while (span.Length > 0)
+				{
+					int idx = span.IndexOf((byte)'\n');
+					if(idx == -1)
+					{
+						break;
+					}
+					_inner.LogInformation("{Line}", Encoding.UTF8.GetString(span.Slice(0, idx)));
+					span = span.Slice(idx + 1);
+				}
+			}
+		}
+
+		readonly List<string> _arguments = new List<string>();
+
+		/// <inheritdoc/>
+		public override void Configure(CommandLineArguments arguments, ILogger logger)
+		{
+			for(int idx = 0; idx < arguments.Count; idx++)
+			{
+				if (!arguments.HasBeenUsed(idx))
+				{
+					_arguments.Add(arguments[idx]);
+					arguments.MarkAsUsed(idx);
+				}
+			}
+		}
+
+		/// <inheritdoc/>
+		public override async Task<int> ExecuteAsync(ILogger logger)
+		{
+			DirectoryReference baseDir = DirectoryReference.GetCurrentDirectory();
+
+			FileReference runUatBat;
+			for (; ; )
+			{
+				runUatBat = FileReference.Combine(baseDir, "RunUAT.bat");
+				if (FileReference.Exists(runUatBat))
+				{
+					break;
+				}
+
+				DirectoryReference? nextDir = baseDir.ParentDirectory;
+				if(nextDir == null)
+				{
+					logger.LogError("Unable to find RunUAT.bat in the current path");
+					return 1;
+				}
+
+				baseDir = nextDir;
+			}
+
+			using (LogSink sink = new LogSink(runUatBat.Directory, logger))
+			{
+				await using (JsonRpcLogger jsonLogger = new JsonRpcLogger(sink, "log", null, logger))
+				{
+					using (LogParser filter = new LogParser(jsonLogger, new List<string>()))
+					{
+						using (ManagedProcessGroup processGroup = new ManagedProcessGroup())
+						{
+							Dictionary<string, string> newEnvironment = ManagedProcess.GetCurrentEnvVars();
+							newEnvironment["UE_STDOUT_JSON"] = "1";
+
+							List<string> allArguments = new List<string>();
+							allArguments.Add(runUatBat.FullName);
+							allArguments.Add("BuildGraph");
+							allArguments.AddRange(_arguments);
+
+							string fileName = Environment.GetEnvironmentVariable("COMSPEC") ?? "cmd.exe";
+							string commandLine = $"/C \"{CommandLineArguments.Join(allArguments)}\"";
+							using (ManagedProcess process = new ManagedProcess(processGroup, fileName, commandLine, runUatBat.Directory.FullName, newEnvironment, null, ProcessPriorityClass.Normal))
+							{
+								await process.CopyToAsync((buffer, offset, length) => filter.WriteData(buffer.AsMemory(offset, length)), 4096, CancellationToken.None);
+								process.WaitForExit();
+							}
+						}
+						filter.Flush();
+					}
+				}
+			}
+			return 0;
+		}
+	}
+}
diff --git a/Engine/Source/Programs/Horde/Horde.Agent/Utility/JsonRpcLogger.cs b/Engine/Source/Programs/Horde/Horde.Agent/Utility/JsonRpcLogger.cs
--- a/Engine/Source/Programs/Horde/Horde.Agent/Utility/JsonRpcLogger.cs
+++ b/Engine/Source/Programs/Horde/Horde.Agent/Utility/JsonRpcLogger.cs
@@ -3,6 +3,7 @@
 using System;
 using System.Buffers;
 using System.Collections.Generic;
+using System.Linq;
 using System.Text;
 using System.Text.Json;
 using System.Threading;
@@ -16,6 +17,59 @@ using Microsoft.Extensions.Logging;
 
 namespace Horde.Agent.Parser
 {
+	interface IJsonRpcLogSink
+	{
+		Task WriteEventsAsync(List<CreateEventRequest> events);
+		Task WriteOutputAsync(WriteOutputRequest request);
+		Task SetOutcomeAsync(JobStepOutcome outcome);
+	}
+
+	sealed class JsonRpcLogSink : IJsonRpcLogSink
+	{
+		readonly IRpcConnection _rpcClient;
+		readonly string? _jobId;
+		readonly string? _jobBatchId;
+		readonly string? _jobStepId;
+		readonly ILogger _logger;
+
+		public JsonRpcLogSink(IRpcConnection rpcClient, string? jobId, string? jobBatchId, string? jobStepId, ILogger logger)
+		{
+			_rpcClient = rpcClient;
+			_jobId = jobId;
+			_jobBatchId = jobBatchId;
+			_jobStepId = jobStepId;
+			_logger = logger;
+		}
+
+		/// <inheritdoc/>
+		public async Task WriteEventsAsync(List<CreateEventRequest> events)
+		{
+			await _rpcClient.InvokeAsync(x => x.CreateEventsAsync(new CreateEventsRequest(events)), new RpcContext(), CancellationToken.None);
+		}
+
+		/// <inheritdoc/>
+		public async Task WriteOutputAsync(WriteOutputRequest request)
+		{
+			await _rpcClient.InvokeAsync(x => x.WriteOutputAsync(request), new RpcContext(), CancellationToken.None);
+		}
+
+		public async Task SetOutcomeAsync(JobStepOutcome outcome)
+		{
+			// Update the outcome of this jobstep
+			if (_jobId != null && _jobBatchId != null && _jobStepId != null)
+			{
+				try
+				{
+					await _rpcClient.InvokeAsync(x => x.UpdateStepAsync(new UpdateStepRequest(_jobId, _jobBatchId, _jobStepId, JobStepState.Unspecified, outcome)), new RpcContext(), CancellationToken.None);
+				}
+				catch (Exception ex)
+				{
+					_logger.LogWarning(ex, "Unable to update step outcome to {NewOutcome}", outcome);
+				}
+			}
+		}
+	}
+
 	/// <summary>
 	/// Class to handle uploading log data to the server in the background
 	/// </summary>
@@ -38,11 +92,8 @@ namespace Horde.Agent.Parser
 			}
 		}
 
-		readonly IRpcConnection _rpcClient;
+		readonly IJsonRpcLogSink _sink;
 		readonly string _logId;
-		readonly string? _jobId;
-		readonly string? _jobBatchId;
-		readonly string? _jobStepId;
 		readonly bool _warnings;
 		readonly ILogger _inner;
 		readonly Channel<JsonLogEvent> _dataChannel;
@@ -60,25 +111,33 @@ namespace Horde.Agent.Parser
 		/// <summary>
 		/// Constructor
 		/// </summary>
-		/// <param name="rpcClient">RPC client to use for server requests</param>
+		/// <param name="sink">Sink for log events</param>
 		/// <param name="logId">The log id to write to</param>
-		/// <param name="jobId">Id of the job being executed</param>
-		/// <param name="jobBatchId">Batch being executed</param>
-		/// <param name="jobStepId">Id of the step being executed</param>
 		/// <param name="warnings">Whether to include warnings in the output</param>
 		/// <param name="inner">Additional logger to write to</param>
-		public JsonRpcLogger(IRpcConnection rpcClient, string logId, string? jobId, string? jobBatchId, string? jobStepId, bool? warnings, ILogger inner)
+		public JsonRpcLogger(IJsonRpcLogSink sink, string logId, bool? warnings, ILogger inner)
 		{
-			_rpcClient = rpcClient;
+			_sink = sink;
 			_logId = logId;
-			_jobId = jobId;
-			_jobBatchId = jobBatchId;
-			_jobStepId = jobStepId;
 			_warnings = warnings ?? true;
 			_inner = inner;
 			_dataChannel = Channel.CreateUnbounded<JsonLogEvent>();
 			_dataWriter = Task.Run(() => RunDataWriter());
-			Outcome = JobStepOutcome.Success;
+		}
+
+		/// <summary>
+		/// Constructor
+		/// </summary>
+		/// <param name="rpcClient">RPC client to use for server requests</param>
+		/// <param name="logId">The log id to write to</param>
+		/// <param name="jobId">Id of the job being executed</param>
+		/// <param name="jobBatchId">Batch being executed</param>
+		/// <param name="jobStepId">Id of the step being executed</param>
+		/// <param name="warnings">Whether to include warnings in the output</param>
+		/// <param name="inner">Additional logger to write to</param>
+		public JsonRpcLogger(IRpcConnection rpcClient, string logId, string? jobId, string? jobBatchId, string? jobStepId, bool? warnings, ILogger inner)
+			: this(new JsonRpcLogSink(rpcClient, jobId, jobBatchId, jobStepId, inner), logId, warnings, inner)
+		{
 		}
 
 		/// <inheritdoc/>
@@ -150,17 +209,6 @@ namespace Horde.Agent.Parser
 			return lineCount;
 		}
 
-		/// <summary>
-		/// Makes a <see cref="CreateEventRequest"/> for the given parameters
-		/// </summary>
-		/// <param name="logLevel">Level for this log event</param>
-		/// <param name="lineCount">Number of lines in the event</param>
-		CreateEventRequest CreateEvent(LogLevel logLevel, int lineCount)
-		{
-			EventSeverity severity = (logLevel == LogLevel.Warning) ? EventSeverity.Warning : EventSeverity.Error;
-			return new CreateEventRequest(severity, _logId, 0, lineCount);
-		}
-
 		/// <summary>
 		/// Stops the log writer's background task
 		/// </summary>
@@ -205,7 +253,6 @@ namespace Horde.Agent.Parser
 
 			// Line separator for JSON events
 			byte[] newline = { (byte)'\n' };
-			JsonEncodedText timestamp = JsonEncodedText.Encode("");
 
 			// The current jobstep outcome
 			JobStepOutcome postedOutcome = JobStepOutcome.Success;
@@ -229,13 +276,11 @@ namespace Horde.Agent.Parser
 						// If we want an event for this log event, create one now
 						if (jsonLogEvent.Level == LogLevel.Warning && ++numWarnings <= MaxWarnings)
 						{
-							int lineCount = GetEventLineCount(jsonLogEvent.Data.Span);
-							events.Add(new CreateEventRequest(EventSeverity.Warning, _logId, 0, lineCount));
+							AddEvent(jsonLogEvent.Data.Span, lineIndex, EventSeverity.Warning, events);
 						}
 						else if ((jsonLogEvent.Level == LogLevel.Error || jsonLogEvent.Level == LogLevel.Critical) && ++numErrors <= MaxErrors)
 						{
-							int lineCount = GetEventLineCount(jsonLogEvent.Data.Span);
-							events.Add(new CreateEventRequest(EventSeverity.Error, _logId, 0, lineCount));
+							AddEvent(jsonLogEvent.Data.Span, lineIndex, EventSeverity.Error, events);
 						}
 
 						writer.Write(jsonLogEvent.Data.Span);
@@ -263,7 +308,7 @@ namespace Horde.Agent.Parser
 					byte[] data = writer.WrittenSpan.ToArray();
 					try
 					{
-						await _rpcClient.InvokeAsync(x => x.WriteOutputAsync(new WriteOutputRequest(_logId, offset, initialLineIndex, data, false)), new RpcContext(), CancellationToken.None);
+						await _sink.WriteOutputAsync(new WriteOutputRequest(_logId, offset, initialLineIndex, data, false));
 					}
 					catch (Exception ex)
 					{
@@ -277,7 +322,7 @@ namespace Horde.Agent.Parser
 				{
 					try
 					{
-						await _rpcClient.InvokeAsync(x => x.CreateEventsAsync(new CreateEventsRequest(events)), new RpcContext(), CancellationToken.None);
+						await _sink.WriteEventsAsync(events);
 					}
 					catch (Exception ex)
 					{
@@ -286,11 +331,11 @@ namespace Horde.Agent.Parser
 				}
 
 				// Update the outcome of this jobstep
-				if (_jobId != null && _jobBatchId != null && _jobStepId != null && Outcome != postedOutcome)
+				if (Outcome != postedOutcome)
 				{
 					try
 					{
-						await _rpcClient.InvokeAsync(x => x.UpdateStepAsync(new UpdateStepRequest(_jobId, _jobBatchId, _jobStepId, JobStepState.Unspecified, Outcome)), new RpcContext(), CancellationToken.None);
+						await _sink.SetOutcomeAsync(Outcome);
 					}
 					catch (Exception ex)
 					{
@@ -304,7 +349,7 @@ namespace Horde.Agent.Parser
 				{
 					try
 					{
-						await _rpcClient.InvokeAsync(x => x.WriteOutputAsync(new WriteOutputRequest(_logId, offset, lineIndex, Array.Empty<byte>(), true)), new RpcContext(), CancellationToken.None);
+						await _sink.WriteOutputAsync(new WriteOutputRequest(_logId, offset, lineIndex, Array.Empty<byte>(), true));
 					}
 					catch (Exception ex)
 					{
@@ -314,5 +359,21 @@ namespace Horde.Agent.Parser
 				}
 			}
 		}
+
+		void AddEvent(ReadOnlySpan<byte> span, int lineIndex, EventSeverity severity, List<CreateEventRequest> events)
+		{
+			try
+			{
+				int lineCount = GetEventLineCount(span);
+				if (lineCount > 0)
+				{
+					events.Add(new CreateEventRequest(severity, _logId, lineIndex, lineCount));
+				}
+			}
+			catch (Exception ex)
+			{
+				_inner.LogError(ex, "Exception while trying to parse line count from data ({Message})", Encoding.UTF8.GetString(span));
+			}
+		}
 	}
 }
diff --git a/Engine/Source/Programs/Shared/EpicGames.Core/ManagedProcess.cs b/Engine/Source/Programs/Shared/EpicGames.Core/ManagedProcess.cs
--- a/Engine/Source/Programs/Shared/EpicGames.Core/ManagedProcess.cs
+++ b/Engine/Source/Programs/Shared/EpicGames.Core/ManagedProcess.cs
@@ -1323,5 +1323,24 @@ namespace EpicGames.Core
 				}
 			}
 		}
+
+		/// <summary>
+		/// Gets all the current environment variables
+		/// </summary>
+		/// <returns></returns>
+		public static Dictionary<string, string> GetCurrentEnvVars()
+		{
+			Dictionary<string, string> newEnvironment = new Dictionary<string, string>();
+			foreach (object? envVar in Environment.GetEnvironmentVariables())
+			{
+				System.Collections.DictionaryEntry entry = (System.Collections.DictionaryEntry)envVar!;
+				string key = entry.Key.ToString()!;
+				if (!newEnvironment.ContainsKey(key))
+				{
+					newEnvironment[key] = entry.Value!.ToString()!;
+				}
+			}
+			return newEnvironment;
+		}
 	}
 }
