Horde: Fix agent test failures due to uninitialized variable.

diff --git a/Engine/Source/Programs/Horde/Horde.Agent/Utility/JsonRpcLogger.cs b/Engine/Source/Programs/Horde/Horde.Agent/Utility/JsonRpcLogger.cs
--- a/Engine/Source/Programs/Horde/Horde.Agent/Utility/JsonRpcLogger.cs
+++ b/Engine/Source/Programs/Horde/Horde.Agent/Utility/JsonRpcLogger.cs
@@ -123,6 +123,8 @@ namespace Horde.Agent.Parser
 			_inner = inner;
 			_dataChannel = Channel.CreateUnbounded<JsonLogEvent>();
 			_dataWriter = Task.Run(() => RunDataWriter());
+
+			Outcome = JobStepOutcome.Success;
 		}
 
 		/// <summary>
