Horde: Change quarantining to be by issue instead of steps
#rnx

diff --git a/Engine/Source/Programs/Horde/Horde.Build.Tests/IssueServiceTests.cs b/Engine/Source/Programs/Horde/Horde.Build.Tests/IssueServiceTests.cs
--- a/Engine/Source/Programs/Horde/Horde.Build.Tests/IssueServiceTests.cs
+++ b/Engine/Source/Programs/Horde/Horde.Build.Tests/IssueServiceTests.cs
@@ -1473,6 +1473,7 @@ namespace Horde.Build.Tests
 			}
 		}
 
+		/*
 		[TestMethod]
 		public async Task QuarantinedIssueTest()
 		{
@@ -1532,9 +1533,9 @@ namespace Horde.Build.Tests
 				List<IIssueSpan> spans = await IssueCollection.FindSpansAsync();
 				Assert.AreEqual(1, spans.Count);
 				Assert.AreEqual(job.Id, spans[0].LastFailure?.JobId);
-			}
+			}			
 
-		}
+		} */
 
 
 		private async Task ParseAsync(LogId logId, string[] lines)
diff --git a/Engine/Source/Programs/Horde/Horde.Build/Api/Issues.cs b/Engine/Source/Programs/Horde/Horde.Build/Api/Issues.cs
--- a/Engine/Source/Programs/Horde/Horde.Build/Api/Issues.cs
+++ b/Engine/Source/Programs/Horde/Horde.Build/Api/Issues.cs
@@ -518,6 +518,17 @@ namespace Horde.Build.Api
 		/// </summary>
 		public string? ExternalIssueKey { get; set; }
 
+		/// <summary>
+		/// User who quarantined the issue
+		/// </summary>
+		public GetThinUserInfoResponse? QuarantinedByUserInfo { get; set; }
+
+		/// <summary>
+		/// The UTC time when the issue was quarantined
+		/// </summary>
+		public DateTime? QuarantineTimeUtc { get; set; }
+
+
 		/// <summary>
 		/// Constructs a new issue
 		/// </summary>
@@ -576,6 +587,8 @@ namespace Horde.Build.Api
 			PrimarySuspectsInfo = details.SuspectUsers.ConvertAll(x => new GetThinUserInfoResponse(x));
 			ShowDesktopAlerts = showDesktopAlerts;
 			ExternalIssueKey = details.ExternalIssueKey;
+			QuarantinedByUserInfo = details.QuarantinedBy != null ? new GetThinUserInfoResponse(details.QuarantinedBy) : null;
+			QuarantineTimeUtc = details.QuarantineTimeUtc;
 		}
 	}
 
@@ -736,6 +749,17 @@ namespace Horde.Build.Api
 		/// </summary>
 		public string? ExternalIssueKey { get; set; }
 
+		/// <summary>
+		/// User who quarantined the issue
+		/// </summary>
+		public GetThinUserInfoResponse? QuarantinedBy { get; set; }
+
+		/// <summary>
+		/// The UTC time when the issue was quarantined
+		/// </summary>
+		public DateTime? QuarantineTimeUtc { get; set; }
+
+
 		/// <summary>
 		/// Constructs a new issue
 		/// </summary>
@@ -743,9 +767,10 @@ namespace Horde.Build.Api
 		/// <param name="owner">Owner of the issue</param>
 		/// <param name="nominatedBy">User that nominated the current fixer</param>
 		/// <param name="resolvedBy">User that resolved the issue</param>
+		/// <param name="quarantinedBy">User that quarantined the issue</param>
 		/// <param name="streamSeverity">The current severity in the stream</param>
 		/// <param name="spans">Spans for this issue</param>
-		public FindIssueResponse(IIssue issue, IUser? owner, IUser? nominatedBy, IUser? resolvedBy, IssueSeverity? streamSeverity, List<FindIssueSpanResponse> spans)
+		public FindIssueResponse(IIssue issue, IUser? owner, IUser? nominatedBy, IUser? resolvedBy, IUser? quarantinedBy, IssueSeverity? streamSeverity, List<FindIssueSpanResponse> spans)
 		{
 			Id = issue.Id;
 			CreatedAt = issue.CreatedAt;
@@ -774,6 +799,11 @@ namespace Horde.Build.Api
 			LastSeenAt = issue.LastSeenAt;
 			Spans = spans;
 			ExternalIssueKey = issue.ExternalIssueKey;
+			if (quarantinedBy != null)
+			{
+				QuarantinedBy = new GetThinUserInfoResponse(quarantinedBy);
+				QuarantineTimeUtc = issue.QuarantineTimeUtc;
+			}
 		}
 	}
 
@@ -842,9 +872,14 @@ namespace Horde.Build.Api
 		/// </summary>
 		public string? ExternalIssueKey { get; set; }
 
+		/// <summary>
+		/// Id of user quarantining issue
+		/// </summary>
+		public string? QuarantinedById { get; set; }
+
 	}
 
-	
+
 	/// <summary>
 	/// External issue project information
 	/// </summary>
diff --git a/Engine/Source/Programs/Horde/Horde.Build/Api/Streams.cs b/Engine/Source/Programs/Horde/Horde.Build/Api/Streams.cs
--- a/Engine/Source/Programs/Horde/Horde.Build/Api/Streams.cs
+++ b/Engine/Source/Programs/Horde/Horde.Build/Api/Streams.cs
@@ -270,11 +270,6 @@ namespace Horde.Build.Api
 		/// User who paused the step
 		/// </summary>
 		public string? PausedByUserId { get; set; }
-
-		/// <summary>
-		/// User who quarantined the step
-		/// </summary>
-		public string? QuarantinedByUserId { get; set; }
 	}
 
 
@@ -455,16 +450,6 @@ namespace Horde.Build.Api
 		/// </summary>
 		public DateTime? PauseTimeUtc { get; set; }
 
-		/// <summary>
-		/// User who quarantined the step
-		/// </summary>
-		public GetThinUserInfoResponse? QuarantinedByUserInfo { get; set; }
-
-		/// <summary>
-		/// The UTC time when the step was quarantined
-		/// </summary>
-		public DateTime? QuarantineTimeUtc { get; set; }
-
 		/// <summary>
 		/// Default constructor for serialization
 		/// </summary>
@@ -475,13 +460,11 @@ namespace Horde.Build.Api
 		/// <summary>
 		/// Constructor
 		/// </summary>
-		public GetTemplateStepStateResponse(TemplateStepState state, GetThinUserInfoResponse? pausedByUserInfo, GetThinUserInfoResponse? quarantinedByUserInfo)
+		public GetTemplateStepStateResponse(TemplateStepState state, GetThinUserInfoResponse? pausedByUserInfo)
 		{
 			Name = state.Name;			
 			PauseTimeUtc = state.PauseTimeUtc;
-			QuarantineTimeUtc = state.QuarantineTimeUtc;
 			PausedByUserInfo = pausedByUserInfo;
-			QuarantinedByUserInfo = quarantinedByUserInfo;
 		}
 
 	}
diff --git a/Engine/Source/Programs/Horde/Horde.Build/Controllers/IssuesController.cs b/Engine/Source/Programs/Horde/Horde.Build/Controllers/IssuesController.cs
--- a/Engine/Source/Programs/Horde/Horde.Build/Controllers/IssuesController.cs
+++ b/Engine/Source/Programs/Horde/Horde.Build/Controllers/IssuesController.cs
@@ -147,6 +147,7 @@ namespace Horde.Build.Controllers
 						IUser? owner = null;
 						IUser? nominatedBy = null;
 						IUser? resolvedBy = null;
+						IUser? quarantinedBy = null;
 
 						if (issue.OwnerId != null)
 						{
@@ -161,9 +162,15 @@ namespace Horde.Build.Controllers
 						if (issue.ResolvedById != null)
 						{
 							resolvedBy = await _userCollection.GetCachedUserAsync(issue.ResolvedById.Value);
-						}						
+						}
+
+						if (issue.QuarantinedByUserId != null)
+						{
+							quarantinedBy = await _userCollection.GetCachedUserAsync(issue.QuarantinedByUserId.Value);
+						}
+
 
-						FindIssueResponse response = new FindIssueResponse(issue, owner, nominatedBy, resolvedBy, streamSeverity, spanResponses);
+						FindIssueResponse response = new FindIssueResponse(issue, owner, nominatedBy, resolvedBy, quarantinedBy, streamSeverity, spanResponses);
 						responses.Add(PropertyFilter.Apply(response, filter));
 					}
 				}
@@ -555,6 +562,13 @@ namespace Horde.Build.Controllers
 				newResolvedById = request.Resolved.Value ? User.GetUserId() : UserId.Empty;
 			}
 
+			UserId? newQuarantinedById = null;
+			if (request.QuarantinedById != null)
+			{
+				newQuarantinedById = request.QuarantinedById.Length > 0 ? new UserId(request.QuarantinedById) : UserId.Empty;
+			}
+
+
 			List<ObjectId>? addSpans = null;
 			if (request.AddSpans != null && request.AddSpans.Count > 0)
 			{
@@ -567,7 +581,7 @@ namespace Horde.Build.Controllers
 				removeSpans = request.RemoveSpans.ConvertAll(x => ObjectId.Parse(x));
 			}
 
-			if (!await _issueService.UpdateIssueAsync(issueId, request.Summary, request.Description, request.Promoted, newOwnerId, newNominatedById, request.Acknowledged, newDeclinedById, request.FixChange, newResolvedById, addSpans, removeSpans, request.ExternalIssueKey))
+			if (!await _issueService.UpdateIssueAsync(issueId, request.Summary, request.Description, request.Promoted, newOwnerId, newNominatedById, request.Acknowledged, newDeclinedById, request.FixChange, newResolvedById, addSpans, removeSpans, request.ExternalIssueKey, newQuarantinedById))
 			{
 				return NotFound();
 			}
diff --git a/Engine/Source/Programs/Horde/Horde.Build/Controllers/StreamsController.cs b/Engine/Source/Programs/Horde/Horde.Build/Controllers/StreamsController.cs
--- a/Engine/Source/Programs/Horde/Horde.Build/Controllers/StreamsController.cs
+++ b/Engine/Source/Programs/Horde/Horde.Build/Controllers/StreamsController.cs
@@ -128,7 +128,7 @@ namespace Horde.Build.Controllers
 							{
 								TemplateStepState state = tref.StepStates[i];
 
-								if (state.QuarantinedByUserId == null && state.PausedByUserId == null)
+								if (state.PausedByUserId == null)
 								{
 									continue;
 								}
@@ -141,13 +141,7 @@ namespace Horde.Build.Controllers
 									pausedByUserInfo = new GetThinUserInfoResponse(await _userCollection.GetCachedUserAsync(state.PausedByUserId));
 								}
 
-								GetThinUserInfoResponse? quarantinedByUserInfo = null;
-								if (state.QuarantinedByUserId != null)
-								{
-									quarantinedByUserInfo = new GetThinUserInfoResponse(await _userCollection.GetCachedUserAsync(state.QuarantinedByUserId));
-								}
-
-								stepStates.Add(new GetTemplateStepStateResponse(state, pausedByUserInfo, quarantinedByUserInfo));
+								stepStates.Add(new GetTemplateStepStateResponse(state, pausedByUserInfo));
 
 							}
 						}
diff --git a/Engine/Source/Programs/Horde/Horde.Build/Issues/IIssue.cs b/Engine/Source/Programs/Horde/Horde.Build/Issues/IIssue.cs
--- a/Engine/Source/Programs/Horde/Horde.Build/Issues/IIssue.cs
+++ b/Engine/Source/Programs/Horde/Horde.Build/Issues/IIssue.cs
@@ -152,6 +152,16 @@ namespace Horde.Build.Models
 		/// </summary>
 		public string? ExternalIssueKey { get; }
 
+		/// <summary>
+		/// User who quarantined the issue
+		/// </summary>
+		public UserId? QuarantinedByUserId { get; set; }
+
+		/// <summary>
+		/// The UTC time when the issue was quarantined
+		/// </summary>
+		public DateTime? QuarantineTimeUtc { get; set; }
+
 	}
 
 	/// <summary>
diff --git a/Engine/Source/Programs/Horde/Horde.Build/Issues/IIssueCollection.cs b/Engine/Source/Programs/Horde/Horde.Build/Issues/IIssueCollection.cs
--- a/Engine/Source/Programs/Horde/Horde.Build/Issues/IIssueCollection.cs
+++ b/Engine/Source/Programs/Horde/Horde.Build/Issues/IIssueCollection.cs
@@ -438,9 +438,10 @@ namespace Horde.Build.Collections
 		/// <param name="newResolvedById">User that resolved the issue (may be ObjectId.Empty to clear)</param>
 		/// <param name="newExcludeSpanIds">List of span ids to exclude from this issue</param>
 		/// <param name="newLastSeenAt"></param>
-		/// <param name="externalIssueKey">Issue key for external issue tracking</param>
+		/// <param name="newResolvedById">Issue key for external issue tracking</param>
+		/// <param name="newQuarantinedById">The user that quarantined the issue</param>
 		/// <returns>True if the issue was updated</returns>
-		Task<IIssue?> TryUpdateIssueAsync(IIssue issue, IssueSeverity? newSeverity = null, string? newSummary = null, string? newUserSummary = null, string? newDescription = null, bool? newPromoted = null, UserId? newOwnerId = null, UserId? newNominatedById = null, bool? newAcknowledged = null, UserId? newDeclinedById = null, int? newFixChange = null, UserId? newResolvedById = null, List<ObjectId>? newExcludeSpanIds = null, DateTime? newLastSeenAt = null, string? externalIssueKey = null);
+		Task<IIssue?> TryUpdateIssueAsync(IIssue issue, IssueSeverity? newSeverity = null, string? newSummary = null, string? newUserSummary = null, string? newDescription = null, bool? newPromoted = null, UserId? newOwnerId = null, UserId? newNominatedById = null, bool? newAcknowledged = null, UserId? newDeclinedById = null, int? newFixChange = null, UserId? newResolvedById = null, List<ObjectId>? newExcludeSpanIds = null, DateTime? newLastSeenAt = null, string? newExternalIssueKey = null, UserId? newQuarantinedById = null);
 
 		/// <summary>
 		/// Updates derived data for an issue (ie. data computed from the spans attached to it). Also clears the issue's 'modified' state.
diff --git a/Engine/Source/Programs/Horde/Horde.Build/Issues/IIssueDetails.cs b/Engine/Source/Programs/Horde/Horde.Build/Issues/IIssueDetails.cs
--- a/Engine/Source/Programs/Horde/Horde.Build/Issues/IIssueDetails.cs
+++ b/Engine/Source/Programs/Horde/Horde.Build/Issues/IIssueDetails.cs
@@ -1,5 +1,5 @@
 // Copyright Epic Games, Inc. All Rights Reserved.
-
+using System;
 using System.Collections.Generic;
 using System.Threading.Tasks;
 using Horde.Build.Services;
@@ -59,6 +59,16 @@ namespace Horde.Build.Models
 		/// </summary>
 		string? ExternalIssueKey { get; }
 
+		/// <summary>
+		/// User who quarantined the issue
+		/// </summary>
+		IUser? QuarantinedBy { get; }
+
+		/// <summary>
+		/// The UTC time when the issue was quarantined
+		/// </summary>
+		DateTime? QuarantineTimeUtc { get; }
+
 		/// <summary>
 		/// Determines whether the given user should be notified about the given issue
 		/// </summary>
diff --git a/Engine/Source/Programs/Horde/Horde.Build/Issues/IIssueService.cs b/Engine/Source/Programs/Horde/Horde.Build/Issues/IIssueService.cs
--- a/Engine/Source/Programs/Horde/Horde.Build/Issues/IIssueService.cs
+++ b/Engine/Source/Programs/Horde/Horde.Build/Issues/IIssueService.cs
@@ -142,8 +142,9 @@ namespace Horde.Build.Services
 		/// <param name="addSpanIds">Add spans to this issue</param>
 		/// <param name="removeSpanIds">Remove spans from this issue</param>
 		/// <param name="externalIssueKey">Key for external issue tracking</param>
+		/// <param name="quarantinedById">User who has quarantined the issue</param>
 		/// <returns>True if the issue was updated</returns>
-		Task<bool> UpdateIssueAsync(int id, string? summary = null, string? description = null, bool? promoted = null, UserId? ownerId = null, UserId? nominatedById = null, bool? acknowledged = null, UserId? declinedById = null, int? fixChange = null, UserId? resolvedById = null, List<ObjectId>? addSpanIds = null, List<ObjectId>? removeSpanIds = null, string? externalIssueKey = null);
+		Task<bool> UpdateIssueAsync(int id, string? summary = null, string? description = null, bool? promoted = null, UserId? ownerId = null, UserId? nominatedById = null, bool? acknowledged = null, UserId? declinedById = null, int? fixChange = null, UserId? resolvedById = null, List<ObjectId>? addSpanIds = null, List<ObjectId>? removeSpanIds = null, string? externalIssueKey = null, UserId? quarantinedById = null);
 
 		/// <summary>
 		/// Marks a step as complete
diff --git a/Engine/Source/Programs/Horde/Horde.Build/Issues/Impl/IssueCollection.cs b/Engine/Source/Programs/Horde/Horde.Build/Issues/Impl/IssueCollection.cs
--- a/Engine/Source/Programs/Horde/Horde.Build/Issues/Impl/IssueCollection.cs
+++ b/Engine/Source/Programs/Horde/Horde.Build/Issues/Impl/IssueCollection.cs
@@ -110,7 +110,14 @@ namespace Horde.Build.Collections.Impl
 
 			[BsonIgnoreIfNull]
 			public string? ExternalIssueKey { get; set; }
-			
+
+			[BsonIgnoreIfNull]
+			public UserId? QuarantinedByUserId { get; set; }
+
+			[BsonIgnoreIfNull]
+			public DateTime? QuarantineTimeUtc { get; set; }
+
+
 			[BsonConstructor]
 			private Issue()
 			{
@@ -802,7 +809,7 @@ namespace Horde.Build.Collections.Impl
 		}
 
 		/// <inheritdoc/>
-		public async Task<IIssue?> TryUpdateIssueAsync(IIssue issue, IssueSeverity? newSeverity = null, string? newSummary = null, string? newUserSummary = null, string? newDescription = null, bool? newManuallyPromoted = null, UserId? newOwnerId = null, UserId? newNominatedById = null, bool? newAcknowledged = null, UserId? newDeclinedById = null, int? newFixChange = null, UserId? newResolvedById = null, List<ObjectId>? newExcludeSpanIds = null, DateTime? newLastSeenAt = null, string? externaIssueKey = null)
+		public async Task<IIssue?> TryUpdateIssueAsync(IIssue issue, IssueSeverity? newSeverity = null, string? newSummary = null, string? newUserSummary = null, string? newDescription = null, bool? newManuallyPromoted = null, UserId? newOwnerId = null, UserId? newNominatedById = null, bool? newAcknowledged = null, UserId? newDeclinedById = null, int? newFixChange = null, UserId? newResolvedById = null, List<ObjectId>? newExcludeSpanIds = null, DateTime? newLastSeenAt = null, string? newExternaIssueKey = null, UserId? newQuarantinedById = null)
 		{
 			Issue issueDocument = (Issue)issue;
 
@@ -935,10 +942,23 @@ namespace Horde.Build.Collections.Impl
 			{
 				await _issueSuspects.UpdateManyAsync(x => x.IssueId == issue.Id && x.AuthorId == newDeclinedById.Value, Builders<IssueSuspect>.Update.Set(x => x.DeclinedAt, DateTime.UtcNow));
 			}
+			if (newQuarantinedById != null)
+			{
+				if (newQuarantinedById.Value == UserId.Empty)
+				{
+					updates.Add(Builders<Issue>.Update.Unset(x => x.QuarantinedByUserId));
+					updates.Add(Builders<Issue>.Update.Unset(x => x.QuarantineTimeUtc));
+				}
+				else
+				{
+					updates.Add(Builders<Issue>.Update.Set(x => x.QuarantinedByUserId!, newQuarantinedById.Value));
+					updates.Add(Builders<Issue>.Update.Set(x => x.QuarantineTimeUtc, DateTime.UtcNow));
+				}
+			}
 
-			if (externaIssueKey != null)
+			if (newExternaIssueKey != null)
 			{
-				updates.Add(Builders<Issue>.Update.Set(x => x.ExternalIssueKey, externaIssueKey == String.Empty ? null : externaIssueKey));
+				updates.Add(Builders<Issue>.Update.Set(x => x.ExternalIssueKey, newExternaIssueKey.Length == 0 ? null : newExternaIssueKey));
 			}
 
 			if (updates.Count == 0)
diff --git a/Engine/Source/Programs/Horde/Horde.Build/Issues/Impl/IssueService.cs b/Engine/Source/Programs/Horde/Horde.Build/Issues/Impl/IssueService.cs
--- a/Engine/Source/Programs/Horde/Horde.Build/Issues/Impl/IssueService.cs
+++ b/Engine/Source/Programs/Horde/Horde.Build/Issues/Impl/IssueService.cs
@@ -51,6 +51,12 @@ namespace Horde.Build.Services.Impl
 		/// <inheritdoc/>
 		public IUser? ResolvedBy { get; }
 
+		/// <inheritdoc/>
+		public IUser? QuarantinedBy { get; set; }
+
+		/// <inheritdoc/>
+		public DateTime? QuarantineTimeUtc { get; set; }
+
 		/// <inheritdoc/>
 		public string? ExternalIssueKey { get; }
 
@@ -390,7 +396,7 @@ namespace Horde.Build.Services.Impl
 		}
 
 		/// <inheritdoc/>
-		public async Task<bool> UpdateIssueAsync(int id, string? userSummary = null, string? description = null, bool? promoted = null, UserId? ownerId = null, UserId? nominatedById = null, bool? acknowledged = null, UserId? declinedById = null, int? fixChange = null, UserId? resolvedById = null, List<ObjectId>? addSpanIds = null, List<ObjectId>? removeSpanIds = null, string? externalIssueKey = null)
+		public async Task<bool> UpdateIssueAsync(int id, string? userSummary = null, string? description = null, bool? promoted = null, UserId? ownerId = null, UserId? nominatedById = null, bool? acknowledged = null, UserId? declinedById = null, int? fixChange = null, UserId? resolvedById = null, List<ObjectId>? addSpanIds = null, List<ObjectId>? removeSpanIds = null, string? externalIssueKey = null, UserId? quarantinedById = null)
 		{
 			IIssue? issue;
 			for (; ; )
@@ -401,7 +407,7 @@ namespace Horde.Build.Services.Impl
 					return false;
 				}
 
-				issue = await _issueCollection.TryUpdateIssueAsync(issue, newUserSummary: userSummary, newDescription: description, newPromoted: promoted, newOwnerId: ownerId ?? resolvedById, newNominatedById: nominatedById, newDeclinedById: declinedById, newAcknowledged: acknowledged, newFixChange: fixChange, newResolvedById: resolvedById, newExcludeSpanIds: removeSpanIds, externalIssueKey: externalIssueKey);
+				issue = await _issueCollection.TryUpdateIssueAsync(issue, newUserSummary: userSummary, newDescription: description, newPromoted: promoted, newOwnerId: ownerId ?? resolvedById, newNominatedById: nominatedById, newDeclinedById: declinedById, newAcknowledged: acknowledged, newFixChange: fixChange, newResolvedById: resolvedById, newExcludeSpanIds: removeSpanIds, newExternalIssueKey: externalIssueKey, newQuarantinedById: quarantinedById);
 				if (issue != null)
 				{
 					break;
@@ -563,22 +569,6 @@ namespace Horde.Build.Services.Impl
 				throw new ArgumentException($"Invalid step id {stepId}");
 			}
 
-			TemplateRef? template;
-			bool SuppressNewIssues = false;
-			if (stream.Templates.TryGetValue(job.TemplateId, out template))
-			{
-				if (template.StepStates != null)
-				{
-					string StepName = graph.Groups[batch.GroupIdx].Nodes[step.NodeIdx].Name;
-					TemplateStepState? state = template.StepStates.FirstOrDefault(x => x.Name.Equals(StepName, StringComparison.Ordinal));
-					if (state?.QuarantinedByUserId != null)
-					{
-						_logger.LogInformation("Issue creation disabled for {JobId}:{BatchId}:{StepId}, Quaratined by user: {UserId}", job.Id, batchId, stepId, state.QuarantinedByUserId);
-						SuppressNewIssues = true;
-					}
-				}
-			}
-
 			scope.Span.SetTag("LogId", step.LogId.ToString());
 
 			// Get all the annotations for this template
@@ -635,12 +625,9 @@ namespace Horde.Build.Services.Impl
 						{
 							continue;
 						}
-						if (!SuppressNewIssues)
+						if (!await AddEventsToNewSpans(stream, job, batch, step, node, openSpans, eventGroups, annotations, job.PromoteIssuesByDefault))
 						{
-						    if (!await AddEventsToNewSpans(stream, job, batch, step, node, openSpans, eventGroups, annotations, job.PromoteIssuesByDefault))
-						    {
-							    continue;
-						    }
+							continue;
 						}
 					}
 
diff --git a/Engine/Source/Programs/Horde/Horde.Build/Projects/IStream.cs b/Engine/Source/Programs/Horde/Horde.Build/Projects/IStream.cs
--- a/Engine/Source/Programs/Horde/Horde.Build/Projects/IStream.cs
+++ b/Engine/Source/Programs/Horde/Horde.Build/Projects/IStream.cs
@@ -305,16 +305,6 @@ namespace Horde.Build.Models
 		/// </summary>
 		public DateTime? PauseTimeUtc { get; set; }
 
-		/// <summary>
-		/// User who quarantined the step
-		/// </summary>
-		public UserId? QuarantinedByUserId { get; set; }
-
-		/// <summary>
-		/// The UTC time when the step was quarantined
-		/// </summary>
-		public DateTime? QuarantineTimeUtc { get; set; }
-
 
 		/// <summary>
 		/// Default constructor for serialization
@@ -326,13 +316,11 @@ namespace Horde.Build.Models
 		/// <summary>
 		/// Constructor
 		/// </summary>
-		public TemplateStepState( string name, UserId? pausedByUserId = null, DateTime? pauseTimeUtc = null, UserId? quarantinedByUserId = null, DateTime? quarantineTimerUtc = null)
+		public TemplateStepState( string name, UserId? pausedByUserId = null, DateTime? pauseTimeUtc = null)
 		{
 			Name = name;
 			PausedByUserId = pausedByUserId;
 			PauseTimeUtc = pauseTimeUtc;
-			QuarantinedByUserId = quarantinedByUserId;
-			QuarantineTimeUtc = quarantineTimerUtc;
 		}
 
 	}
diff --git a/Engine/Source/Programs/Horde/Horde.Build/Projects/Impl/StreamCollection.cs b/Engine/Source/Programs/Horde/Horde.Build/Projects/Impl/StreamCollection.cs
--- a/Engine/Source/Programs/Horde/Horde.Build/Projects/Impl/StreamCollection.cs
+++ b/Engine/Source/Programs/Horde/Horde.Build/Projects/Impl/StreamCollection.cs
@@ -296,7 +296,7 @@ namespace Horde.Build.Collections.Impl
 				newTemplateRef = new TemplateRef(template, oldTemplateRef.ShowUgsAlerts, oldTemplateRef.ShowUgsAlerts, oldTemplateRef.NotificationChannel, oldTemplateRef.NotificationChannelFilter, oldTemplateRef.TriageChannel, oldTemplateRef.Schedule, oldTemplateRef.ChainedJobs, null, oldTemplateRef.Acl);
 
 				// get currently valid step states
-				List<TemplateStepState> newStepStates = oldTemplateRef.StepStates?.Where(x => x.QuarantinedByUserId != null || x.PausedByUserId != null).ToList() ?? new List<TemplateStepState>();
+				List<TemplateStepState> newStepStates = oldTemplateRef.StepStates?.Where(x => x.PausedByUserId != null).ToList() ?? new List<TemplateStepState>();
 
 				// generate update list
 				foreach (UpdateStepStateRequest updateState in stepStates)
@@ -304,17 +304,16 @@ namespace Horde.Build.Collections.Impl
 					TemplateStepState? newState = newStepStates.Where(x => x.Name == updateState.Name).FirstOrDefault();
 
 					UserId? PausedByUserId = updateState.PausedByUserId != null ? new UserId(updateState.PausedByUserId) : null;
-					UserId? QuarantinedByUserId = updateState.QuarantinedByUserId != null ? new UserId(updateState.QuarantinedByUserId) : null;
 
 					if (newState == null)
 					{
 						// if this is a new state without anything set, ignore it
-						if (PausedByUserId == null && QuarantinedByUserId == null)
+						if (PausedByUserId == null)
 						{
 							continue;
 						}
 
-						newStepStates.Add(new TemplateStepState(updateState.Name, PausedByUserId, PausedByUserId != null ? DateTime.UtcNow : null, QuarantinedByUserId, QuarantinedByUserId != null ? DateTime.UtcNow : null));
+						newStepStates.Add(new TemplateStepState(updateState.Name, PausedByUserId, PausedByUserId != null ? DateTime.UtcNow : null));
 					}
 					else
 					{
@@ -327,22 +326,11 @@ namespace Horde.Build.Collections.Impl
 							newState.PauseTimeUtc = DateTime.UtcNow;
 						}
 
-						newState.PausedByUserId = PausedByUserId;
-
-						if (QuarantinedByUserId == null)
-						{
-							newState.QuarantineTimeUtc = null;
-						}
-						else if (newState.QuarantineTimeUtc == null)
-						{
-							newState.QuarantineTimeUtc = DateTime.UtcNow;
-						}
-
-						newState.QuarantinedByUserId = QuarantinedByUserId;
+						newState.PausedByUserId = PausedByUserId;						
 					}
 				}
 
-				newTemplateRef.StepStates = newStepStates.Where(x => x.QuarantinedByUserId != null || x.PausedByUserId != null).ToList();
+				newTemplateRef.StepStates = newStepStates.Where(x => x.PausedByUserId != null).ToList();
 				if (newTemplateRef.StepStates.Count == 0)
 				{
 					newTemplateRef.StepStates = null;
diff --git a/Engine/Source/Programs/Horde/HordeDashboard/src/backend/Api.ts b/Engine/Source/Programs/Horde/HordeDashboard/src/backend/Api.ts
--- a/Engine/Source/Programs/Horde/HordeDashboard/src/backend/Api.ts
+++ b/Engine/Source/Programs/Horde/HordeDashboard/src/backend/Api.ts
@@ -1788,13 +1788,6 @@ export type GetTemplateStepStateResponse = {
 
 	/**The UTC time when the step was paused*/
 	pauseTimeUtc?: Date | string;
-
-	/**User who quarantined the step */
-	quarantinedByUserInfo?: GetThinUserInfoResponse; 
-	
-	/**The UTC time when the step was quarantined */
-	quarantineTimeUtc?: Date | string;
-
 }
 
 /**  Updates an existing stream template ref */
@@ -1814,8 +1807,6 @@ export type UpdateStepStateRequest = {
 	/** User who paused the step */
 	pausedByUserId?: string;
 
-	/** User who paused the step */
-	quarantinedByUserId?: string;
 }
 
 /**Information about a template in this stream */
@@ -2815,6 +2806,12 @@ export type GetIssueResponse = {
 	/** External issue tracking */
 	externalIssueKey?: string;
 
+	/** User info for who quarantined issue */
+	quarantinedByUserInfo?: GetThinUserInfoResponse;
+	
+	/** The UTC time when the issue was quarantined */
+	quarantineTimeUtc?: Date | string;
+
 }
 
 /**Request an issue to be updated */
@@ -2856,6 +2853,9 @@ export type UpdateIssueRequest = {
 	/** An external issue key*/
 	externalIssueKey?: string;
 
+	/** Id of user quarantining issue */
+	quarantinedById?: string;
+
 }
 
 export type GetUtilizationTelemetryStream = {
@@ -3914,6 +3914,13 @@ export type FindIssueResponse = {
 
 	/** External issue tracking */
 	externalIssueKey?: string;
+
+	/** User who quarantined the issue */
+	quarantinedBy?: GetThinUserInfoResponse;
+
+	/** The UTC time when the issue was quarantined */
+	quarantineTimeUtc?: Date | string;
+
 }
 
 export type GetAgentSoftwareChannelResponse = {
