Horde: Fix tests
#rnx

diff --git a/Engine/Source/Programs/Horde/Horde.Build.Tests/JobTaskSourceTests.cs b/Engine/Source/Programs/Horde/Horde.Build.Tests/JobTaskSourceTests.cs
--- a/Engine/Source/Programs/Horde/Horde.Build.Tests/JobTaskSourceTests.cs
+++ b/Engine/Source/Programs/Horde/Horde.Build.Tests/JobTaskSourceTests.cs
@@ -109,8 +109,7 @@ namespace Horde.Build.Tests
 			Fixture fixture = await SetupPoolWithAgentAsync(isPoolAutoScaled: true, shouldCreateAgent: true, isAgentEnabled: true);
 
 			// update template with some step states
-			IStream Stream = Deref(await StreamService.TryUpdateTemplateRefAsync(fixture.Stream!, fixture.TemplateRefId1, new List<UpdateStepStateRequest>() { new UpdateStepStateRequest() { Name = "Quarantined Step", QuarantinedByUserId = UserId.GenerateNewId().ToString() } }));
-			Stream = Deref(await StreamService.TryUpdateTemplateRefAsync(Stream, fixture.TemplateRefId1, new List<UpdateStepStateRequest>() { new UpdateStepStateRequest() { Name = "Paused Step", PausedByUserId = UserId.GenerateNewId().ToString() } }));
+			IStream Stream = Deref(await StreamService.TryUpdateTemplateRefAsync(fixture.Stream!, fixture.TemplateRefId1, new List<UpdateStepStateRequest>() { new UpdateStepStateRequest() { Name = "Paused Step", PausedByUserId = UserId.GenerateNewId().ToString() } }));
 
 			// create a new graph with the associated nodes
 			List<NewGroup> newGroups = new List<NewGroup>();
