Horde: Fix warning
#rnx

diff --git a/Engine/Source/Programs/Horde/Horde.Build/Issues/IIssueCollection.cs b/Engine/Source/Programs/Horde/Horde.Build/Issues/IIssueCollection.cs
--- a/Engine/Source/Programs/Horde/Horde.Build/Issues/IIssueCollection.cs
+++ b/Engine/Source/Programs/Horde/Horde.Build/Issues/IIssueCollection.cs
@@ -438,7 +438,7 @@ namespace Horde.Build.Collections
 		/// <param name="newResolvedById">User that resolved the issue (may be ObjectId.Empty to clear)</param>
 		/// <param name="newExcludeSpanIds">List of span ids to exclude from this issue</param>
 		/// <param name="newLastSeenAt"></param>
-		/// <param name="newResolvedById">Issue key for external issue tracking</param>
+		/// <param name="newExternalIssueKey">Issue key for external issue tracking</param>
 		/// <param name="newQuarantinedById">The user that quarantined the issue</param>
 		/// <returns>True if the issue was updated</returns>
 		Task<IIssue?> TryUpdateIssueAsync(IIssue issue, IssueSeverity? newSeverity = null, string? newSummary = null, string? newUserSummary = null, string? newDescription = null, bool? newPromoted = null, UserId? newOwnerId = null, UserId? newNominatedById = null, bool? newAcknowledged = null, UserId? newDeclinedById = null, int? newFixChange = null, UserId? newResolvedById = null, List<ObjectId>? newExcludeSpanIds = null, DateTime? newLastSeenAt = null, string? newExternalIssueKey = null, UserId? newQuarantinedById = null);
