Horde Storage - Improvements for debugging issues with the role based auth

diff --git a/Engine/Source/Programs/Horde/HordeStorage/Horde.Storage/Implementation/AnyRequirement.cs b/Engine/Source/Programs/Horde/HordeStorage/Horde.Storage/Implementation/AnyRequirement.cs
--- a/Engine/Source/Programs/Horde/HordeStorage/Horde.Storage/Implementation/AnyRequirement.cs
+++ b/Engine/Source/Programs/Horde/HordeStorage/Horde.Storage/Implementation/AnyRequirement.cs
@@ -2,6 +2,7 @@
 
 using System;
 using System.Collections.Generic;
+using System.Linq;
 using System.Threading.Tasks;
 using Microsoft.AspNetCore.Authorization;
 using Microsoft.AspNetCore.Authorization.Infrastructure;
@@ -41,6 +42,12 @@ namespace Horde.Storage
                 }
             }
         }
+
+        public override string ToString()
+        {
+            string components = string.Join(',', InnerRequirements.Select(r => r.ToString()));
+            return $"Any({components})";
+        }
     }
 
     /// <summary>
