Horde: Missing file.

diff --git a/Engine/Source/Programs/Horde/Horde.Agent/Horde.Agent.csproj b/Engine/Source/Programs/Horde/Horde.Agent/Horde.Agent.csproj
--- a/Engine/Source/Programs/Horde/Horde.Agent/Horde.Agent.csproj
+++ b/Engine/Source/Programs/Horde/Horde.Agent/Horde.Agent.csproj
@@ -71,6 +71,9 @@
   </ItemGroup>
 
   <ItemGroup>
+    <Compile Include="..\..\UnrealBuildTool\Matchers\*.cs">
+      <Link>Parser\Matchers\%(FileName)%(Extension)</Link>
+    </Compile>
     <Content Include="..\Horde.Build\Protos\horde\horde_rpc_messages.proto">
       <Link>Protos\horde_rpc_messages.proto</Link>
     </Content>
