Horde Storage - Fixed issue in the orphan blob cleanup that wasnt deleting the blob from all namespace that share a pool, generating a inconsistency and making the consistency check attempt to replicate the blob from remote regions when it should just be deleted.

diff --git a/Engine/Source/Programs/Horde/HordeStorage/Horde.Storage/Implementation/GC/OrphanBlobCleanupRefs.cs b/Engine/Source/Programs/Horde/HordeStorage/Horde.Storage/Implementation/GC/OrphanBlobCleanupRefs.cs
--- a/Engine/Source/Programs/Horde/HordeStorage/Horde.Storage/Implementation/GC/OrphanBlobCleanupRefs.cs
+++ b/Engine/Source/Programs/Horde/HordeStorage/Horde.Storage/Implementation/GC/OrphanBlobCleanupRefs.cs
@@ -87,8 +87,9 @@ namespace Horde.Storage.Implementation
                     bool found = false;
                     NamespacePolicy policy = _namespacePolicyResolver.GetPoliciesForNs(@namespace);
                     
+                    List<NamespaceId> namespacesThatSharePool = namespaces.Where(ns => _namespacePolicyResolver.GetPoliciesForNs(ns).StoragePool == policy.StoragePool).ToList();
                     // check all other namespaces that share the same storage pool for presence of the blob
-                    foreach (NamespaceId blobNamespace in namespaces.Where(ns => _namespacePolicyResolver.GetPoliciesForNs(ns).StoragePool == policy.StoragePool))
+                    foreach (NamespaceId blobNamespace in namespacesThatSharePool)
                     {
                         if (cancellationToken.IsCancellationRequested)
                         {
@@ -135,7 +136,10 @@ namespace Horde.Storage.Implementation
 
                     if (!found)
                     {
-                        await RemoveBlob(@namespace, blob);
+                        foreach (NamespaceId namespaceToDeleteFrom in namespacesThatSharePool)
+                        {
+                            await RemoveBlob(namespaceToDeleteFrom, blob);
+                        }
                         ++countOfBlobsRemoved;
                     }
                 }
@@ -151,6 +155,10 @@ namespace Horde.Storage.Implementation
             {
                 await _blobService.DeleteObject(ns, blob);
             }
+            catch (BlobNotFoundException)
+            {
+                // ignore blob not found exceptions, if it didn't exist it has been removed so we are happy either way
+            }
             catch (Exception e)
             {
                 _logger.Warning("Failed to delete blob {Blob} from {Namespace} due to {Error}", blob, ns, e.Message);
