Horde: Issue quarantine work, populating in issue details and most of the dashboard side
#rnx

diff --git a/Engine/Source/Programs/Horde/Horde.Build/Issues/Impl/IssueService.cs b/Engine/Source/Programs/Horde/Horde.Build/Issues/Impl/IssueService.cs
--- a/Engine/Source/Programs/Horde/Horde.Build/Issues/Impl/IssueService.cs
+++ b/Engine/Source/Programs/Horde/Horde.Build/Issues/Impl/IssueService.cs
@@ -85,18 +85,21 @@ namespace Horde.Build.Services.Impl
 		/// <param name="owner"></param>
 		/// <param name="nominatedBy"></param>
 		/// <param name="resolvedBy"></param>
+		/// <param name="quarantinedBy"></param>
 		/// <param name="spans">Spans for this issue</param>
 		/// <param name="steps">Steps for this issue</param>
 		/// <param name="suspects"></param>
 		/// <param name="suspectUsers"></param>
 		/// <param name="showDesktopAlerts"></param>
 		/// <param name="externalIssueKey"></param>
-		public IssueDetails(IIssue issue, IUser? owner, IUser? nominatedBy, IUser? resolvedBy, IReadOnlyList<IIssueSpan> spans, IReadOnlyList<IIssueStep> steps, IReadOnlyList<IIssueSuspect> suspects, IReadOnlyList<IUser> suspectUsers, bool showDesktopAlerts, string? externalIssueKey)
+		public IssueDetails(IIssue issue, IUser? owner, IUser? nominatedBy, IUser? resolvedBy, IUser? quarantinedBy, IReadOnlyList<IIssueSpan> spans, IReadOnlyList<IIssueStep> steps, IReadOnlyList<IIssueSuspect> suspects, IReadOnlyList<IUser> suspectUsers, bool showDesktopAlerts, string? externalIssueKey)
 		{
 			Issue = issue;
 			Owner = owner;
 			NominatedBy = nominatedBy;
 			ResolvedBy = resolvedBy;
+			QuarantinedBy = quarantinedBy;
+			QuarantineTimeUtc = issue.QuarantineTimeUtc;
 			Spans = spans;
 			Steps = steps;
 			Suspects = suspects;
@@ -300,6 +303,7 @@ namespace Horde.Build.Services.Impl
 			IUser? owner = issue.OwnerId.HasValue ? await _userCollection.GetCachedUserAsync(issue.OwnerId.Value) : null;
 			IUser? nominatedBy = issue.NominatedById.HasValue ? await _userCollection.GetCachedUserAsync(issue.NominatedById.Value) : null;
 			IUser? resolvedBy = (issue.ResolvedById.HasValue && issue.ResolvedById != IIssue.ResolvedByTimeoutId && issue.ResolvedById != IIssue.ResolvedByUnknownId)? await _userCollection.GetCachedUserAsync(issue.ResolvedById.Value) : null;
+			IUser? quarantinedBy = issue.QuarantinedByUserId.HasValue ? await _userCollection.GetCachedUserAsync(issue.QuarantinedByUserId.Value) : null;
 
 			List<IIssueSpan> spans = await _issueCollection.FindSpansAsync(issue.Id);
 			List<IIssueStep> steps = await _issueCollection.FindStepsAsync(spans.Select(x => x.Id));
@@ -315,7 +319,7 @@ namespace Horde.Build.Services.Impl
 				}
 			}
 
-			return new IssueDetails(issue, owner, nominatedBy, resolvedBy, spans, steps, suspects, suspectUsers, ShowDesktopAlertsForIssue(issue, spans), issue.ExternalIssueKey);
+			return new IssueDetails(issue, owner, nominatedBy, resolvedBy, quarantinedBy, spans, steps, suspects, suspectUsers, ShowDesktopAlertsForIssue(issue, spans), issue.ExternalIssueKey);
 		}
 
 		/// <inheritdoc/>
