Horde: Fix server build.

diff --git a/Engine/Source/Programs/Horde/BuildHorde.xml b/Engine/Source/Programs/Horde/BuildHorde.xml
--- a/Engine/Source/Programs/Horde/BuildHorde.xml
+++ b/Engine/Source/Programs/Horde/BuildHorde.xml
@@ -20,6 +20,7 @@
 				Engine/Binaries/DotNET/EpicGames.Perforce.Native/...
 				Engine/Source/Programs/Shared/...
 				Engine/Source/Programs/Horde/...
+				Engine/Source/Programs/UnrealBuildTool/Matchers/...
 				Engine/Source/ThirdParty/Perforce/...
 			</Property>
 			<Tag Files="$(StagePaths)" Except=".../.vs/...;.../.git/...;.../bin/...;.../obj/..." With="#InputFiles"/>
