Horde Storage - Fixes to private namespace always returning Forbidden, even on corp  routes
Also fixed up configuration for development instances of Horde Storage so that we have our 3 sets of ports as well as http1 and 2 variants of them, to better reflect how we operate in prod.

diff --git a/Engine/Source/Programs/Horde/HordeStorage/Horde.Storage/Controllers/RequestHelper.cs b/Engine/Source/Programs/Horde/HordeStorage/Horde.Storage/Controllers/RequestHelper.cs
--- a/Engine/Source/Programs/Horde/HordeStorage/Horde.Storage/Controllers/RequestHelper.cs
+++ b/Engine/Source/Programs/Horde/HordeStorage/Horde.Storage/Controllers/RequestHelper.cs
@@ -46,10 +46,11 @@ public class RequestHelper
 
         // namespace is a restricted namespace
         HttpContext context = request.HttpContext;
-        bool isPublicPort = _settings!.CurrentValue.PublicApiPorts.Contains(context.Connection.LocalPort) ||
-            /* unit tests do not run on ports, we consider them always on the internal port */ !(context.Connection.LocalPort == 0 && context.Connection.LocalIpAddress == null);
+        bool isUnitTest = context.Connection.LocalPort == 0 && context.Connection.LocalIpAddress == null;
+        /* unit tests do not run on ports, we consider them always on the internal port */
+        bool isPublicPort = _settings!.CurrentValue.PublicApiPorts.Contains(context.Connection.LocalPort);
             
-        if (isPublicPort)
+        if (isPublicPort && !isUnitTest)
         {
             // trying to access restricted namespace on a public port, this is not allowed
             return new ForbidResult();
diff --git a/Engine/Source/Programs/Horde/HordeStorage/Horde.Storage/appsettings.Development.json b/Engine/Source/Programs/Horde/HordeStorage/Horde.Storage/appsettings.Development.json
--- a/Engine/Source/Programs/Horde/HordeStorage/Horde.Storage/appsettings.Development.json
+++ b/Engine/Source/Programs/Horde/HordeStorage/Horde.Storage/appsettings.Development.json
@@ -1,102 +1,120 @@
 {
-  "Serilog": {
-    "Using": [ "Serilog.Sinks.Console" ],
-    "MinimumLevel": {
-      "Default": "Information",
-      "Override": {
-        "Microsoft": "Information",
-        "System": "Information"
-      }
-    },
-    "WriteTo": [
-      {
-        "Name": "Console"
-      }
-    ],
-    "Enrich": [ "FromLogContext" ]
-  },
+	"Serilog": {
+		"Using": [ "Serilog.Sinks.Console" ],
+		"MinimumLevel": {
+			"Default": "Information",
+			"Override": {
+				"Microsoft": "Information",
+				"System": "Information"
+			}
+		},
+		"WriteTo": [
+			{
+				"Name": "Console"
+			}
+		],
+		"Enrich": [ "FromLogContext" ]
+	},
 
-  "Mongo": {
-    // this assumes a mongo db were we have localhost exemption, e.g. we do not need to authenticate
-    // if using mongo via a docker container you will need to update this
-    "ConnectionString": "mongodb://localhost:27017",
-    "RequireTls12": true
-  },
+	"Mongo": {
+		// this assumes a mongo db were we have localhost exemption, e.g. we do not need to authenticate
+		// if using mongo via a docker container you will need to update this
+		"ConnectionString": "mongodb://localhost:27017",
+		"RequireTls12": true
+	},
 
-  // scylla is normally not enabled in development but we still configure it so we can easily switch
-  "Scylla": {
-    "LocalDatacenterName": "datacenter1",
-    "LocalKeyspaceSuffix": "local"
-  },
+	// scylla is normally not enabled in development but we still configure it so we can easily switch
+	"Scylla": {
+		"LocalDatacenterName": "datacenter1",
+		"LocalKeyspaceSuffix": "local"
+	},
 
-  "Horde_Storage": {
-    // assume a single instance is running
-    "LeaderElectionImplementation": "Static",
+	"Horde_Storage": {
+		// assume a single instance is running
+		"LeaderElectionImplementation": "Static",
 
-    // use mongo for any database needs
-    "ReferencesDbImplementation": "Mongo",
-    "ContentIdStoreImplementation": "Mongo",
-    "BlobIndexImplementation": "Mongo",
+		// use mongo for any database needs
+		"ReferencesDbImplementation": "Mongo",
+		"ContentIdStoreImplementation": "Mongo",
+		"BlobIndexImplementation": "Mongo",
 
-    // store blobs in the filesystem
-    "StorageImplementations": [
-      "Filesystem"
-    ],
+		// store blobs in the filesystem
+		"StorageImplementations": [
+			"Filesystem"
+		],
 
-    // disable old api that we do not want to use
-    "TransactionLogWriterImplementation": "Memory",
-    "RefDbImplementation": "Memory"
-  },
+		// disable old api that we do not want to use
+		"TransactionLogWriterImplementation": "Memory",
+		"RefDbImplementation": "Memory"
+	},
 
-  "Filesystem": {
-    "RootDir": "$(ExecutableLocation)/BlobStorage"
-  },
+	"Filesystem": {
+		"RootDir": "$(ExecutableLocation)/BlobStorage"
+	},
 
-  "Replication": {
-    "StateRoot": "$(ExecutableLocation)/ganymede",
-    "ReplicationPollFrequencySeconds": 15,
-    "CurrentSite": "local-test",
-    "Enabled": false
-  },
+	"Replication": {
+		"StateRoot": "$(ExecutableLocation)/ganymede",
+		"ReplicationPollFrequencySeconds": 15,
+		"CurrentSite": "local-test",
+		"Enabled": false
+	},
 
-  "GC": {
-    "BlobCleanupServiceEnabled": false,
-    "CleanOldBlobs": false,
-    "CleanOldRefRecords": false
-  },
+	"GC": {
+		"BlobCleanupServiceEnabled": false,
+		"CleanOldBlobs": false,
+		"CleanOldRefRecords": false
+	},
 
-  "Jupiter": {
-    "CurrentSite": "local",
-    "InternalApiPort": 57003
-  },
+	"Jupiter": {
+		"CurrentSite": "local",
+		"PublicApiPorts": [ 57001, 57002 ],
+		"CorpApiPorts": [ 57003, 57004 ],
+		"InternalApiPorts": [ 57005, 57006 ]
+	},
 
-  // disable auth and allow any namespace
-  "Auth": {
-    "Method": "Disabled"
-  },
+	// disable auth and allow any namespace
+	"Auth": {
+		"Method": "Disabled"
+	},
 
-  "Namespaces": {
-    "Policies": {
-      "*": {
-        "Claims": [ "*" ]
-      }
-    }
-  },
+	"Namespaces": {
+		"Policies": {
+			"*": {
+				"Claims": [ "*" ]
+			},
+			"test.private": {
+				"Claims": [ "*" ],
+				"IsPublicNamespace": false
+			}
+		}
+	},
 
-  "AllowedHosts": "*",
+	"AllowedHosts": "*",
 
-  "Kestrel": {
-    "Endpoints": {
-      "Http": {
-        "Url": "http://0.0.0.0:57001"
-      },
-      "Http2": {
-        "Url": "http://0.0.0.0:57002",
-        "Protocols": "Http2"
-      },
-      "Http-Internal": {
-        "Url": "http://0.0.0.0:57003"
-      }
-    }
-  }
+
+	"Kestrel": {
+		"Endpoints": {
+			"Http": {
+				"Url": "http://0.0.0.0:57001"
+			},
+			"Http2": {
+				"Url": "http://0.0.0.0:57002",
+				"Protocols": "Http2"
+			},
+			"CorpHttp": {
+				"Url": "http://0.0.0.0:57003"
+			},
+			"CorpHttp2": {
+				"Url": "http://0.0.0.0:57004",
+				"Protocols": "Http2"
+			},
+			"InternalHttp": {
+				"Url": "http://0.0.0.0:57005"
+			},
+			"InternalHttp2": {
+				"Url": "http://0.0.0.0:57006",
+				"Protocols": "Http2"
+			}
+		}
+	}
 }
