Horde: Suppress issue notifications for quarantined issues and surface status  in reports
#rnx

diff --git a/Engine/Source/Programs/Horde/Horde.Build.Tests/IssueServiceTests.cs b/Engine/Source/Programs/Horde/Horde.Build.Tests/IssueServiceTests.cs
--- a/Engine/Source/Programs/Horde/Horde.Build.Tests/IssueServiceTests.cs
+++ b/Engine/Source/Programs/Horde/Horde.Build.Tests/IssueServiceTests.cs
@@ -15,11 +15,9 @@ using Horde.Build.Server;
 using Horde.Build.Services;
 using Horde.Build.Tests.Stubs.Services;
 using Horde.Build.Utilities;
-using Horde.Build.Api;
 using HordeCommon;
 using Microsoft.Extensions.DependencyInjection;
 using Microsoft.Extensions.Logging;
-using Microsoft.Extensions.Logging.Abstractions;
 using Microsoft.VisualStudio.TestTools.UnitTesting;
 using MongoDB.Bson;
 using MongoDB.Driver;
@@ -27,7 +25,6 @@ using Moq;
 using ProjectId = Horde.Build.Utilities.StringId<Horde.Build.Models.IProject>;
 using StreamId = Horde.Build.Utilities.StringId<Horde.Build.Models.IStream>;
 using TemplateRefId = Horde.Build.Utilities.StringId<Horde.Build.Models.TemplateRef>;
-using EpicGames.Perforce;
 
 namespace Horde.Build.Tests
 {
@@ -1473,71 +1470,6 @@ namespace Horde.Build.Tests
 			}
 		}
 
-		/*
-		[TestMethod]
-		public async Task QuarantinedIssueTest()
-		{
-			IStream stream = Deref(await StreamCollection.GetAsync(_mainStreamId));
-			TemplateRefId templateRefId = new TemplateRefId("test-template");
-
-			// #1
-			// Scenario: Warning in first step, step is quarantined
-			// Expected: No issue generated
-			{
-				stream = Deref(await StreamService.TryUpdateTemplateRefAsync(stream, templateRefId, new List<UpdateStepStateRequest>() { new UpdateStepStateRequest() { Name = "Update Version Files", QuarantinedByUserId = UserId.GenerateNewId().ToString() } }));
-
-				IJob job = CreateJob(_mainStreamId, 105, "Test Build", _graph);
-				await AddEvent(job, 0, 0, new { level = nameof(LogLevel.Warning) }, EventSeverity.Warning);
-				await UpdateCompleteStep(job, 0, 0, JobStepOutcome.Warnings);
-
-				List<IIssue> issues = await IssueService.FindIssuesAsync();
-				Assert.AreEqual(0, issues.Count);				
-			}
-
-			// #2
-			// Scenario: Warning in first step, step is not quarantined
-			// Expected: issue is generated
-			{
-				stream = Deref(await StreamService.TryUpdateTemplateRefAsync(stream, templateRefId, new List<UpdateStepStateRequest>()));
-
-				// step states should be cleared now
-				Assert.AreEqual(null, stream.Templates[templateRefId].StepStates);
-
-				IJob job = CreateJob(_mainStreamId, 105, "Test Build", _graph);
-				await AddEvent(job, 0, 0, new { level = nameof(LogLevel.Warning) }, EventSeverity.Warning);
-				await UpdateCompleteStep(job, 0, 0, JobStepOutcome.Warnings);
-
-				List<IIssue> issues = await IssueService.FindIssuesAsync();
-				Assert.AreEqual(1, issues.Count);
-
-				Assert.AreEqual("Warnings in Update Version Files", issues[0].Summary);
-
-				List<IIssueSpan> spans = await IssueCollection.FindSpansAsync();
-				Assert.AreEqual(1, spans.Count);
-				Assert.AreEqual(job.Id, spans[0].LastFailure?.JobId);
-			}
-
-			// #3
-			// Scenario: Warning in first step, step is quarantined, issue exists
-			// Expected: issue is updated with new failure
-			{
-				stream = Deref(await StreamService.TryUpdateTemplateRefAsync(stream, templateRefId, new List<UpdateStepStateRequest>() { new UpdateStepStateRequest() { Name = "Update Version Files", QuarantinedByUserId = UserId.GenerateNewId().ToString() } }));
-
-				IJob job = CreateJob(_mainStreamId, 105, "Test Build", _graph);
-				await AddEvent(job, 0, 0, new { level = nameof(LogLevel.Warning) }, EventSeverity.Warning);
-				await UpdateCompleteStep(job, 0, 0, JobStepOutcome.Warnings);
-
-				List<IIssue> issues = await IssueService.FindIssuesAsync();
-				Assert.AreEqual(1, issues.Count);
-
-				List<IIssueSpan> spans = await IssueCollection.FindSpansAsync();
-				Assert.AreEqual(1, spans.Count);
-				Assert.AreEqual(job.Id, spans[0].LastFailure?.JobId);
-			}			
-
-		} */
-
-
 		private async Task ParseAsync(LogId logId, string[] lines)
 		{
 			await using (TestJsonLogger logger = new TestJsonLogger(LogFileService, logId))
diff --git a/Engine/Source/Programs/Horde/Horde.Build/Issues/Impl/IssueService.cs b/Engine/Source/Programs/Horde/Horde.Build/Issues/Impl/IssueService.cs
--- a/Engine/Source/Programs/Horde/Horde.Build/Issues/Impl/IssueService.cs
+++ b/Engine/Source/Programs/Horde/Horde.Build/Issues/Impl/IssueService.cs
@@ -266,6 +266,13 @@ namespace Horde.Build.Services.Impl
 			for (int idx = 0; idx < openIssues.Count; idx++)
 			{
 				IIssue openIssue = openIssues[idx];
+
+				// Do not automaically close quarantined issues
+				if (openIssue.QuarantinedByUserId != null)
+				{
+					continue;
+				}
+
 				if (openIssue.LastSeenAt < utcNow - TimeSpan.FromDays(7.0))
 				{
 					await _issueCollection.TryUpdateIssueAsync(openIssue, newResolvedById: IIssue.ResolvedByTimeoutId);
diff --git a/Engine/Source/Programs/Horde/Horde.Build/Notifications/Impl/SlackNotificationSink.cs b/Engine/Source/Programs/Horde/Horde.Build/Notifications/Impl/SlackNotificationSink.cs
--- a/Engine/Source/Programs/Horde/Horde.Build/Notifications/Impl/SlackNotificationSink.cs
+++ b/Engine/Source/Programs/Horde/Horde.Build/Notifications/Impl/SlackNotificationSink.cs
@@ -595,6 +595,12 @@ namespace Horde.Build.Notifications.Impl
 		/// <inheritdoc/>
 		public async Task NotifyIssueUpdatedAsync(IIssue issue)
 		{
+			// Do not send notifications for quarantined issues
+			if (issue.QuarantinedByUserId != null)
+			{
+				return;
+			}
+
 			IIssueDetails details = await _issueService.GetIssueDetailsAsync(issue);
 
 			bool notifySuspects = issue.Promoted || details.Spans.Any(x => x.LastFailure.Annotations.NotifySubmitters ?? false);
@@ -1121,6 +1127,11 @@ namespace Horde.Build.Notifications.Impl
 						}
 					}
 
+					if (issue.QuarantinedByUserId != null)
+					{
+						status = $"{status} - *Quarantined*";
+					}
+
 					body.Append($"\n\u2022 *Issue <{issueUrl}|{issue.Id}>");
 					
 					string? triageChannel = report.Workflow.TriageChannel;
