Horde: Add a setting that will disable [REDACTED_HANDLE] for particular workflows. Should allow us to test new build health workflows in the prod environment.

diff --git a/Engine/Source/Programs/Horde/Horde.Build/Notifications/Impl/SlackNotificationSink.cs b/Engine/Source/Programs/Horde/Horde.Build/Notifications/Impl/SlackNotificationSink.cs
--- a/Engine/Source/Programs/Horde/Horde.Build/Notifications/Impl/SlackNotificationSink.cs
+++ b/Engine/Source/Programs/Horde/Horde.Build/Notifications/Impl/SlackNotificationSink.cs
@@ -50,6 +50,8 @@ namespace Horde.Build.Notifications.Impl
 		const string UpdateMessageUrl = "https://slack.com/api/chat.update";
 		const string GetPermalinkUrl = "https://slack.com/api/chat.getPermalink";
 
+		const bool defaultAllowMentions = true;
+
 		class SocketResponse
 		{
 			[JsonPropertyName("ok")]
@@ -685,7 +687,7 @@ namespace Horde.Build.Notifications.Impl
 			{
 				foreach (string channel in channels)
 				{
-					await SendIssueMessageAsync(channel, issue, details, null);
+					await SendIssueMessageAsync(channel, issue, details, null, defaultAllowMentions);
 				}
 			}
 		}
@@ -762,7 +764,7 @@ namespace Horde.Build.Notifications.Impl
 						List<string> suspectList = new List<string>();
 						foreach (IGrouping<UserId, IIssueSuspect> suspectGroup in suspectGroups)
 						{
-							string mention = await FormatMentionAsync(suspectGroup.Key);
+							string mention = await FormatMentionAsync(suspectGroup.Key, workflow.AllowMentions);
 							string changes = String.Join(", ", suspectGroup.Select(x => $"<ugs://change?number={x.Change}|CL {x.Change}>"));
 							suspectList.Add($"{mention} ({changes})");
 						}
@@ -795,7 +797,7 @@ namespace Horde.Build.Notifications.Impl
 				return;
 			}
 
-			await SendIssueMessageAsync(slackUserId, issue, details, user.Id);
+			await SendIssueMessageAsync(slackUserId, issue, details, user.Id, defaultAllowMentions);
 		}
 
 		Uri GetJobUrl(JobId jobId)
@@ -813,7 +815,7 @@ namespace Horde.Build.Notifications.Impl
 			return new Uri(_settings.DashboardUrl, $"job/{step.JobId}?step={step.StepId}&issue={issue.Id}");
 		}
 
-		async Task SendIssueMessageAsync(string recipient, IIssue issue, IIssueDetails details, UserId? userId)
+		async Task SendIssueMessageAsync(string recipient, IIssue issue, IIssueDetails details, UserId? userId, bool allowMentions)
 		{
 			using IDisposable scope = _logger.BeginScope("SendIssueMessageAsync (User: {SlackUser}, Issue: {IssueId})", recipient, issue.Id);
 
@@ -888,14 +890,9 @@ namespace Horde.Build.Notifications.Impl
 				{
 					if (issue.NominatedById != null)
 					{
-						IUser? nominatedByUser = await _userCollection.GetUserAsync(issue.NominatedById.Value);
-						if (nominatedByUser != null)
-						{
-							string? nominatedBySlackUserId = await GetSlackUserId(nominatedByUser);
-							string mention = (nominatedBySlackUserId != null) ? $"<@{nominatedBySlackUserId}>" : nominatedByUser.Login ?? $"User {nominatedByUser.Id}";
-							string text = $"You were nominated to fix this issue by {mention} at {FormatSlackTime(issue.NominatedAt ?? DateTime.UtcNow)}";
-							attachment.Blocks.Add(new SectionBlock(text));
-						}
+						string mention = await FormatMentionAsync(issue.NominatedById.Value, allowMentions);
+						string text = $"You were nominated to fix this issue by {mention} at {FormatSlackTime(issue.NominatedAt ?? DateTime.UtcNow)}";
+						attachment.Blocks.Add(new SectionBlock(text));
 					}
 					else
 					{
@@ -915,7 +912,7 @@ namespace Horde.Build.Notifications.Impl
 			}
 			else if (issue.OwnerId != null)
 			{
-				string ownerMention = await FormatMentionAsync(issue.OwnerId.Value);
+				string ownerMention = await FormatMentionAsync(issue.OwnerId.Value, allowMentions);
 				if (issue.AcknowledgedAt.HasValue)
 				{
 					attachment.Blocks.Add(new SectionBlock($":+1: Acknowledged by {ownerMention} at {FormatSlackTime(issue.AcknowledgedAt.Value)}"));
@@ -930,7 +927,7 @@ namespace Horde.Build.Notifications.Impl
 				}
 				else
 				{
-					attachment.Blocks.Add(new SectionBlock($"{ownerMention} was nominated to fix this issue by {await FormatMentionAsync(issue.NominatedById.Value)}"));
+					attachment.Blocks.Add(new SectionBlock($"{ownerMention} was nominated to fix this issue by {await FormatMentionAsync(issue.NominatedById.Value, allowMentions)}"));
 				}
 			}
 			else if (userId != null)
@@ -995,7 +992,7 @@ namespace Horde.Build.Notifications.Impl
 			return user.Name;
 		}
 
-		async Task<string> FormatMentionAsync(UserId userId)
+		async Task<string> FormatMentionAsync(UserId userId, bool allowMentions)
 		{
 			IUser? user = await _userCollection.GetUserAsync(userId);
 			if (user == null)
@@ -1009,7 +1006,7 @@ namespace Horde.Build.Notifications.Impl
 				return user.Login;
 			}
 
-			if (!_environment.IsProduction())
+			if (!_environment.IsProduction() && allowMentions)
 			{
 				return $"{user.Name} [{slackUserId}]";
 			}
@@ -1066,7 +1063,7 @@ namespace Horde.Build.Notifications.Impl
 					if (recipient != null)
 					{
 						IIssueDetails details = await _issueService.GetIssueDetailsAsync(newIssue);
-						await SendIssueMessageAsync(recipient, newIssue, details, userId);
+						await SendIssueMessageAsync(recipient, newIssue, details, userId, defaultAllowMentions);
 					}
 				}
 			}
diff --git a/Engine/Source/Programs/Horde/Horde.Build/Server/Config.cs b/Engine/Source/Programs/Horde/Horde.Build/Server/Config.cs
--- a/Engine/Source/Programs/Horde/Horde.Build/Server/Config.cs
+++ b/Engine/Source/Programs/Horde/Horde.Build/Server/Config.cs
@@ -367,6 +367,11 @@ namespace Horde.Build.Server
 		/// </summary>
 		public int MaxMentions { get; set; } = 5;
 
+		/// <summary>
+		/// Whether to mention people on this thread. Useful to disable for testing.
+		/// </summary>
+		public bool AllowMentions { get; set; } = true;
+
 		/// <summary>
 		/// Additional node annotations implicit in this workflow
 		/// </summary>
