Horde: Do not auto resolve quarantined issues
#rnx

diff --git a/Engine/Source/Programs/Horde/Horde.Build/Issues/Impl/IssueService.cs b/Engine/Source/Programs/Horde/Horde.Build/Issues/Impl/IssueService.cs
--- a/Engine/Source/Programs/Horde/Horde.Build/Issues/Impl/IssueService.cs
+++ b/Engine/Source/Programs/Horde/Horde.Build/Issues/Impl/IssueService.cs
@@ -991,21 +991,24 @@ namespace Horde.Build.Services.Impl
 
 				// Calculate the time at which the issue is verified fixed
 				DateTime? newVerifiedAt = null;
-				if (spans.Count == 0)
+				if (issue.QuarantinedByUserId == null)
 				{
-					newVerifiedAt = issue.VerifiedAt ?? DateTime.UtcNow;
-				}
-				else
-				{
-					foreach (IIssueSpan span in spans)
+					if (spans.Count == 0)
 					{
-						if (span.NextSuccess == null)
-						{
-							break;
-						}
-						else if (newVerifiedAt == null || span.NextSuccess.StepTime < newVerifiedAt.Value)
+						newVerifiedAt = issue.VerifiedAt ?? DateTime.UtcNow;
+					}
+					else
+					{
+						foreach (IIssueSpan span in spans)
 						{
-							newVerifiedAt = span.NextSuccess.StepTime;
+							if (span.NextSuccess == null)
+							{
+								break;
+							}
+							else if (newVerifiedAt == null || span.NextSuccess.StepTime < newVerifiedAt.Value)
+							{
+								newVerifiedAt = span.NextSuccess.StepTime;
+							}
 						}
 					}
 				}
