Horde: Make the help email and slack settings configurable
#rnx

diff --git a/Engine/Source/Programs/Horde/Horde.Build/Api/Dashboard.cs b/Engine/Source/Programs/Horde/Horde.Build/Api/Dashboard.cs
--- a/Engine/Source/Programs/Horde/Horde.Build/Api/Dashboard.cs
+++ b/Engine/Source/Programs/Horde/Horde.Build/Api/Dashboard.cs
@@ -22,6 +22,16 @@ namespace Horde.Build.Api
 		/// </summary>
 		public string? PerforceSwarmUrl { get; set; }
 
+		/// <summary>
+		/// Help email address that users can contact with issues
+		/// </summary>
+		public string? HelpEmailAddress { get; set; }
+
+		/// <summary>
+		/// Help slack channel that users can use for issues
+		/// </summary>
+		public string? HelpSlackChannel { get; set; }
+
 		/// <summary>
 		/// Response constructor
 		/// </summary>
diff --git a/Engine/Source/Programs/Horde/Horde.Build/Controllers/DashboardController.cs b/Engine/Source/Programs/Horde/Horde.Build/Controllers/DashboardController.cs
--- a/Engine/Source/Programs/Horde/Horde.Build/Controllers/DashboardController.cs
+++ b/Engine/Source/Programs/Horde/Horde.Build/Controllers/DashboardController.cs
@@ -129,6 +129,9 @@ namespace Horde.Build
 				dashboardConfigResponse.PerforceSwarmUrl = _settings.P4SwarmUrl.ToString().TrimEnd('/');
 			}
 
+			dashboardConfigResponse.HelpEmailAddress = _settings.HelpEmailAddress;
+			dashboardConfigResponse.HelpSlackChannel = _settings.HelpSlackChannel;
+
 			return dashboardConfigResponse;
 		}
 
diff --git a/Engine/Source/Programs/Horde/Horde.Build/ServerSettings.cs b/Engine/Source/Programs/Horde/Horde.Build/ServerSettings.cs
--- a/Engine/Source/Programs/Horde/Horde.Build/ServerSettings.cs
+++ b/Engine/Source/Programs/Horde/Horde.Build/ServerSettings.cs
@@ -438,7 +438,17 @@ namespace Horde.Build
 		/// The URl to use for generating links back to the dashboard.
 		/// </summary>
 		public Uri DashboardUrl { get; set; } = new Uri("https://localhost:3000");
-		
+
+		/// <summary>
+		/// Help email address that users can contact with issues
+		/// </summary>
+		public string? HelpEmailAddress { get; set; }
+
+		/// <summary>
+		/// Help slack channel that users can use for issues
+		/// </summary>
+		public string? HelpSlackChannel { get; set; }
+
 		/// <summary>
 		/// The p4 bridge server
 		/// </summary>
diff --git a/Engine/Source/Programs/Horde/HordeDashboard/src/backend/Api.ts b/Engine/Source/Programs/Horde/HordeDashboard/src/backend/Api.ts
--- a/Engine/Source/Programs/Horde/HordeDashboard/src/backend/Api.ts
+++ b/Engine/Source/Programs/Horde/HordeDashboard/src/backend/Api.ts
@@ -385,8 +385,13 @@ export type GetDashboardConfigResponse = {
 
 	/** The url of the perforce swarm installation */
 	perforceSwarmUrl?: string;
-}
 
+	/** Help email address that users can contact with issues */
+	helpEmailAddress?: string;
+
+	/** Help slack channel that users can use for issues */
+	helpSlackChannel?: string;
+}
 
 /**Parameters to register a new agent */
 export type CreateAgentRequest = {
diff --git a/Engine/Source/Programs/Horde/HordeDashboard/src/backend/Dashboard.ts b/Engine/Source/Programs/Horde/HordeDashboard/src/backend/Dashboard.ts
--- a/Engine/Source/Programs/Horde/HordeDashboard/src/backend/Dashboard.ts
+++ b/Engine/Source/Programs/Horde/HordeDashboard/src/backend/Dashboard.ts
@@ -138,6 +138,13 @@ export class Dashboard {
 
     }
 
+    get helpEmail(): string | undefined {
+        return this.config?.helpEmailAddress;
+    }
+
+    get helpSlack(): string | undefined {
+        return this.config?.helpSlackChannel;
+    }
 
     get swarmUrl(): string | undefined {
         return this.config?.perforceSwarmUrl;
