Horde: Add Slack reactions to indicate build health status (tick when an issue is resolved, eyes when an issue is acknowledged), and mention the owner of an issue in the thread when it is created.

diff --git a/Engine/Source/Programs/Horde/Horde.Build/Notifications/Impl/SlackNotificationSink.cs b/Engine/Source/Programs/Horde/Horde.Build/Notifications/Impl/SlackNotificationSink.cs
--- a/Engine/Source/Programs/Horde/Horde.Build/Notifications/Impl/SlackNotificationSink.cs
+++ b/Engine/Source/Programs/Horde/Horde.Build/Notifications/Impl/SlackNotificationSink.cs
@@ -46,6 +46,8 @@ namespace Horde.Build.Notifications.Impl
 	/// </summary>
 	public class SlackNotificationSink : BackgroundService, INotificationSink, IAvatarService
 	{
+		const string AddReactionUrl = "https://slack.com/api/reactions.add";
+		const string RemoveReactionUrl = "https://slack.com/api/reactions.remove";
 		const string PostMessageUrl = "https://slack.com/api/chat.postMessage";
 		const string UpdateMessageUrl = "https://slack.com/api/chat.update";
 		const string GetPermalinkUrl = "https://slack.com/api/chat.getPermalink";
@@ -73,6 +75,18 @@ namespace Horde.Build.Notifications.Impl
 			public EventPayload? Payload { get; set; }
 		}
 
+		class ReactionMessage
+		{
+			[JsonPropertyName("channel")]
+			public string? Channel { get; set; }
+
+			[JsonPropertyName("timestamp")]
+			public string? Ts { get; set; }
+
+			[JsonPropertyName("name")]
+			public string? Name { get; set; }
+		}
+
 		class EventPayload
 		{
 			[JsonPropertyName("type")]
@@ -642,7 +656,7 @@ namespace Horde.Build.Notifications.Impl
 					IStream? stream = await _streamService.GetStreamAsync(span.StreamId);
 					if(stream != null && stream.Config.TryGetWorkflow(workflowId.Value, out WorkflowConfig? workflow))
 					{
-						await SendTriageMessageAsync(issue, span, workflow);
+						await CreateOrUpdateWorkflowThreadAsync(issue, span, workflow);
 					}
 				}
 			}
@@ -692,9 +706,37 @@ namespace Horde.Build.Notifications.Impl
 			}
 		}
 
+		async Task AddReactionAsync(string channel, string ts, string name)
+		{
+			ReactionMessage message = new ReactionMessage();
+			message.Channel = channel;
+			message.Ts = ts;
+			message.Name = name;
+
+			SlackResponse response = await SendRequestAsync<SlackResponse>(AddReactionUrl, message);
+			if(!response.Ok && !String.Equals(response.Error, "already_reacted", StringComparison.Ordinal))
+			{
+				_logger.LogWarning("Error adding reaction {Name} to thread {Channel}/{Ts}: {Error}", name, channel, ts, response.Error);
+			}
+		}
+
+		async Task RemoveReactionAsync(string channel, string ts, string name)
+		{
+			ReactionMessage message = new ReactionMessage();
+			message.Channel = channel;
+			message.Ts = ts;
+			message.Name = name;
+
+			SlackResponse response = await SendRequestAsync<SlackResponse>(RemoveReactionUrl, message);
+			if (!response.Ok && !String.Equals(response.Error, "no_reaction", StringComparison.Ordinal))
+			{
+				_logger.LogWarning("Error removing reaction {Name} from thread {Channel}/{Ts}: {Error}", name, channel, ts, response.Error);
+			}
+		}
+
 		static string GetTriageThreadEventId(int issueId) => $"issue_triage_{issueId}";
 
-		async Task SendTriageMessageAsync(IIssue issue, IIssueSpan span, WorkflowConfig workflow)
+		async Task CreateOrUpdateWorkflowThreadAsync(IIssue issue, IIssueSpan span, WorkflowConfig workflow)
 		{
 			string? triageChannel = workflow.TriageChannel;
 			if (triageChannel != null)
@@ -756,25 +798,43 @@ namespace Horde.Build.Notifications.Impl
 						permalink = await GetPermalinkAsync(state.Channel, summaryTs);
 					}
 
-					List<IIssueSuspect> suspects = await _issueService.GetIssueSuspectsAsync(issue);
-					IGrouping<UserId, IIssueSuspect>[] suspectGroups = suspects.GroupBy(x => x.AuthorId).ToArray();
-
-					if (suspectGroups.Length > 0 && suspectGroups.Length <= workflow.MaxMentions)
+					// If it has an owner, show that
+					if (issue.OwnerId != null)
 					{
-						List<string> suspectList = new List<string>();
-						foreach (IGrouping<UserId, IIssueSuspect> suspectGroup in suspectGroups)
+						string mention = await FormatMentionAsync(issue.OwnerId.Value, workflow.AllowMentions);
+						await SendMessageToThread(triageChannel, state.Ts, $"Assigned to {mention}");
+					}
+					else
+					{
+						List<IIssueSuspect> suspects = await _issueService.GetIssueSuspectsAsync(issue);
+						IGrouping<UserId, IIssueSuspect>[] suspectGroups = suspects.GroupBy(x => x.AuthorId).ToArray();
+
+						if (suspectGroups.Length > 0 && suspectGroups.Length <= workflow.MaxMentions)
 						{
-							string mention = await FormatMentionAsync(suspectGroup.Key, workflow.AllowMentions);
-							string changes = String.Join(", ", suspectGroup.Select(x => $"<ugs://change?number={x.Change}|CL {x.Change}>"));
-							suspectList.Add($"{mention} ({changes})");
-						}
+							List<string> suspectList = new List<string>();
+							foreach (IGrouping<UserId, IIssueSuspect> suspectGroup in suspectGroups)
+							{
+								string mention = await FormatMentionAsync(suspectGroup.Key, workflow.AllowMentions);
+								string changes = String.Join(", ", suspectGroup.Select(x => $"<ugs://change?number={x.Change}|CL {x.Change}>"));
+								suspectList.Add($"{mention} ({changes})");
+							}
 
-						string suspectMessage = $"Possibly {StringUtils.FormatList(suspectList, "or")}.";
-						await SendMessageToThread(triageChannel, state.Ts, suspectMessage);
+							string suspectMessage = $"Possibly {StringUtils.FormatList(suspectList, "or")}.";
+							await SendMessageToThread(triageChannel, state.Ts, suspectMessage);
+						}
 					}
 
 					await SetMessageTimestampAsync(state.Id, state.Channel, state.Ts, permalink);
 				}
+
+				if (issue.AcknowledgedAt != null)
+				{
+					await AddReactionAsync(state.Channel, state.Ts, "eyes");
+				}
+				if (issue.ResolvedAt != null)
+				{
+					await AddReactionAsync(state.Channel, state.Ts, "tick");
+				}
 			}
 		}
 
