Horde: Post a message to the workflow thread when an issue is marked as fixed.

diff --git a/Engine/Source/Programs/Horde/Horde.Build/Notifications/Impl/SlackNotificationSink.cs b/Engine/Source/Programs/Horde/Horde.Build/Notifications/Impl/SlackNotificationSink.cs
--- a/Engine/Source/Programs/Horde/Horde.Build/Notifications/Impl/SlackNotificationSink.cs
+++ b/Engine/Source/Programs/Horde/Horde.Build/Notifications/Impl/SlackNotificationSink.cs
@@ -815,7 +815,7 @@ namespace Horde.Build.Notifications.Impl
 							foreach (IGrouping<UserId, IIssueSuspect> suspectGroup in suspectGroups)
 							{
 								string mention = await FormatMentionAsync(suspectGroup.Key, workflow.AllowMentions);
-								string changes = String.Join(", ", suspectGroup.Select(x => $"<ugs://change?number={x.Change}|CL {x.Change}>"));
+								string changes = String.Join(", ", suspectGroup.Select(x => FormatChange(x.Change)));
 								suspectList.Add($"{mention} ({changes})");
 							}
 
@@ -835,6 +835,22 @@ namespace Horde.Build.Notifications.Impl
 				{
 					await AddReactionAsync(state.Channel, state.Ts, "tick");
 				}
+
+				if (issue.ResolvedAt != null && issue.FixChange != null)
+				{
+					string fixedEventId = $"issue_{issue.Id}_fixed_{issue.ResolvedAt.Value}";
+
+					(MessageStateDocument fixedState, bool fixedIsNew) = await AddOrUpdateMessageStateAsync(triageChannel, eventId, null, "");
+					if (fixedIsNew)
+					{
+						string fixedMessage = $"Marked as fixed in {FormatChange(issue.FixChange.Value)}";
+						string? fixedTs = await SendMessageToThread(triageChannel, state.Ts, fixedMessage);
+						if (fixedTs != null)
+						{
+							await SetMessageTimestampAsync(fixedState.Id, triageChannel, fixedTs);
+						}
+					}
+				}
 			}
 		}
 
@@ -1052,6 +1068,8 @@ namespace Horde.Build.Notifications.Impl
 			return user.Name;
 		}
 
+		static string FormatChange(int change) => $"<ugs://change?number={change}|CL {change}>";
+
 		async Task<string> FormatMentionAsync(UserId userId, bool allowMentions)
 		{
 			IUser? user = await _userCollection.GetUserAsync(userId);
