Horde: Post a message to the Slack thread whenever a Jira issue is linked to a Horde issue.

diff --git a/Engine/Source/Programs/Horde/Horde.Build/Issues/External/IExternalIssueService.cs b/Engine/Source/Programs/Horde/Horde.Build/Issues/External/IExternalIssueService.cs
--- a/Engine/Source/Programs/Horde/Horde.Build/Issues/External/IExternalIssueService.cs
+++ b/Engine/Source/Programs/Horde/Horde.Build/Issues/External/IExternalIssueService.cs
@@ -48,6 +48,13 @@ namespace Horde.Build.Services
 	/// </summary>
 	public interface IExternalIssueService : IHostedService
 	{
+		/// <summary>
+		/// Gets the URL for an issue from a key
+		/// </summary>
+		/// <param name="key">Key for the issue</param>
+		/// <returns>Url to the issue</returns>
+		string? GetIssueUrl(string key);
+
 		/// <summary>
 		/// Get issues associated with provided keys
 		/// </summary>
@@ -83,6 +90,12 @@ namespace Horde.Build.Services
 	/// </summary>
 	public class DefaultExternalIssueService : IExternalIssueService
 	{
+		/// <inheritdoc/>
+		public string? GetIssueUrl(string key)
+		{
+			return null;
+		}
+
 		/// <inheritdoc/>
 		public Task<List<IExternalIssue>> GetIssuesAsync(string[] keys)
 		{
diff --git a/Engine/Source/Programs/Horde/Horde.Build/Issues/External/Impl/JiraService.cs b/Engine/Source/Programs/Horde/Horde.Build/Issues/External/Impl/JiraService.cs
--- a/Engine/Source/Programs/Horde/Horde.Build/Issues/External/Impl/JiraService.cs
+++ b/Engine/Source/Programs/Horde/Horde.Build/Issues/External/Impl/JiraService.cs
@@ -167,7 +167,6 @@ namespace Horde.Build.Services
 		/// <inheritdoc/>
 		public Task StopAsync(CancellationToken cancellationToken) => _ticker.StopAsync();
 
-
 		/// <inheritdoc/>
 		public void Dispose()
 		{
@@ -175,6 +174,9 @@ namespace Horde.Build.Services
 			_client.Dispose();
 		}
 
+		/// <inheritdoc/>
+		public string? GetIssueUrl(string key) => $"{_jiraUrl}browse/{key}";
+
 		public async Task<List<IExternalIssueProject>> GetProjects(IStream stream)
 		{
 			HashSet<string> projectKeys = new HashSet<string>();
diff --git a/Engine/Source/Programs/Horde/Horde.Build/Notifications/Impl/SlackNotificationSink.cs b/Engine/Source/Programs/Horde/Horde.Build/Notifications/Impl/SlackNotificationSink.cs
--- a/Engine/Source/Programs/Horde/Horde.Build/Notifications/Impl/SlackNotificationSink.cs
+++ b/Engine/Source/Programs/Horde/Horde.Build/Notifications/Impl/SlackNotificationSink.cs
@@ -228,6 +228,7 @@ namespace Horde.Build.Notifications.Impl
 		readonly IMongoCollection<MessageStateDocument> _messageStates;
 		readonly IMongoCollection<SlackUser> _slackUsers;
 		readonly HashSet<string>? _allowUsers;
+		readonly IExternalIssueService _externalIssueService;
 		readonly ILogger _logger;
 
 		/// <summary>
@@ -238,13 +239,14 @@ namespace Horde.Build.Notifications.Impl
 		/// <summary>
 		/// Constructor
 		/// </summary>
-		public SlackNotificationSink(MongoService mongoService, IIssueService issueService, IUserCollection userCollection, ILogFileService logFileService, StreamService streamService, ConfigCollection configCollection, IWebHostEnvironment environment, IOptions<ServerSettings> settings, ILogger<SlackNotificationSink> logger)
+		public SlackNotificationSink(MongoService mongoService, IIssueService issueService, IUserCollection userCollection, ILogFileService logFileService, StreamService streamService, ConfigCollection configCollection, IExternalIssueService externalIssueService, IWebHostEnvironment environment, IOptions<ServerSettings> settings, ILogger<SlackNotificationSink> logger)
 		{
 			_issueService = issueService;
 			_userCollection = userCollection;
 			_logFileService = logFileService;
 			_streamService = streamService;
 			_configCollection = configCollection;
+			_externalIssueService = externalIssueService;
 			_environment = environment;
 			_settings = settings.Value;
 			_messageStates = mongoService.Database.GetCollection<MessageStateDocument>("Slack");
@@ -839,17 +841,28 @@ namespace Horde.Build.Notifications.Impl
 				if (issue.ResolvedAt != null && issue.FixChange != null)
 				{
 					string fixedEventId = $"issue_{issue.Id}_fixed_{issue.ResolvedAt.Value}";
+					string fixedMessage = $"Marked as fixed in {FormatChange(issue.FixChange.Value)}";
+					await PostSingleMessageToThreadAsync(triageChannel, state.Ts, fixedEventId, fixedMessage);
+				}
 
-					(MessageStateDocument fixedState, bool fixedIsNew) = await AddOrUpdateMessageStateAsync(triageChannel, eventId, null, "");
-					if (fixedIsNew)
-					{
-						string fixedMessage = $"Marked as fixed in {FormatChange(issue.FixChange.Value)}";
-						string? fixedTs = await SendMessageToThread(triageChannel, state.Ts, fixedMessage);
-						if (fixedTs != null)
-						{
-							await SetMessageTimestampAsync(fixedState.Id, triageChannel, fixedTs);
-						}
-					}
+				if(issue.ExternalIssueKey != null)
+				{
+					string extIssueEventId = $"issue_{issue.Id}_ext_{issue.ExternalIssueKey}";
+					string extIssueMessage = $"Linked to issue {FormatExternalIssue(issue.ExternalIssueKey)}";
+					await PostSingleMessageToThreadAsync(triageChannel, state.Ts, extIssueEventId, extIssueMessage);
+				}
+			}
+		}
+
+		async Task PostSingleMessageToThreadAsync(string channel, string threadTs, string eventId, string message)
+		{
+			(MessageStateDocument state, bool isNew) = await AddOrUpdateMessageStateAsync(channel, eventId, null, "");
+			if (isNew)
+			{
+				string? ts = await SendMessageToThread(channel, threadTs, message);
+				if (ts != null)
+				{
+					await SetMessageTimestampAsync(state.Id, channel, ts);
 				}
 			}
 		}
@@ -1070,6 +1083,19 @@ namespace Horde.Build.Notifications.Impl
 
 		static string FormatChange(int change) => $"<ugs://change?number={change}|CL {change}>";
 
+		string FormatExternalIssue(string key)
+		{
+			string? url = _externalIssueService.GetIssueUrl(key);
+			if (url == null)
+			{
+				return key;
+			}
+			else
+			{
+				return $"<{_externalIssueService.GetIssueUrl(key)}|{key}>";
+			}
+		}
+
 		async Task<string> FormatMentionAsync(UserId userId, bool allowMentions)
 		{
 			IUser? user = await _userCollection.GetUserAsync(userId);
