Horde: Ignore empty responses from sync commands.

diff --git a/Engine/Source/Programs/Horde/Horde.Build/Perforce/Impl/CommitService.cs b/Engine/Source/Programs/Horde/Horde.Build/Perforce/Impl/CommitService.cs
--- a/Engine/Source/Programs/Horde/Horde.Build/Perforce/Impl/CommitService.cs
+++ b/Engine/Source/Programs/Horde/Horde.Build/Perforce/Impl/CommitService.cs
@@ -896,6 +896,12 @@ namespace Horde.Build.Commits.Impl
 				int count = 0;
 				await foreach (PerforceResponse<SyncRecord> record in perforce.StreamCommandAsync<SyncRecord>("sync", new[] { "-k", $"{queryPath}@{change}" }, null, default))
 				{
+					PerforceError? error = record.Error;
+					if (error != null && error.Generic == PerforceGenericCode.Empty)
+					{
+						continue;
+					}
+
 					SyncRecord syncRecord = record.Data;
 					if (!syncRecord.Path.StartsWith(clientRoot))
 					{
