Horde: Add external issue claim and when creating Jira issues add the user as a watcher once issue is created
#rnx

diff --git a/Engine/Source/Programs/Horde/Horde.Build/Controllers/IssuesController.cs b/Engine/Source/Programs/Horde/Horde.Build/Controllers/IssuesController.cs
--- a/Engine/Source/Programs/Horde/Horde.Build/Controllers/IssuesController.cs
+++ b/Engine/Source/Programs/Horde/Horde.Build/Controllers/IssuesController.cs
@@ -654,7 +654,7 @@ namespace Horde.Build.Controllers
 				return BadRequest($"Missing user for {User.GetUserName()}");
 			}
 
-			(string? key, string? url) = await _externalIssueService.CreateIssueAsync(user, issueRequest.IssueId, issueRequest.Summary, issueRequest.ProjectId, issueRequest.ComponentId, issueRequest.IssueTypeId, issueRequest.Description, issueRequest.HordeIssueLink);
+			(string? key, string? url) = await _externalIssueService.CreateIssueAsync(user, User.GetExternalIssueUser(), issueRequest.IssueId, issueRequest.Summary, issueRequest.ProjectId, issueRequest.ComponentId, issueRequest.IssueTypeId, issueRequest.Description, issueRequest.HordeIssueLink);
 
 			if (key == null)
 			{
diff --git a/Engine/Source/Programs/Horde/Horde.Build/Issues/External/IExternalIssueService.cs b/Engine/Source/Programs/Horde/Horde.Build/Issues/External/IExternalIssueService.cs
--- a/Engine/Source/Programs/Horde/Horde.Build/Issues/External/IExternalIssueService.cs
+++ b/Engine/Source/Programs/Horde/Horde.Build/Issues/External/IExternalIssueService.cs
@@ -65,7 +65,8 @@ namespace Horde.Build.Services
 		/// <summary>
 		/// Create and link an external issue
 		/// </summary>		
-		/// <param name="User"></param>
+		/// <param name="user"></param>
+		/// <param name="externalIssueUser"></param>
 		/// <param name="issueId"></param>
 		/// <param name="summary"></param>
 		/// <param name="projectId"></param>
@@ -73,7 +74,7 @@ namespace Horde.Build.Services
 		/// <param name="issueType"></param>
 		/// <param name="description"></param>
 		/// <param name="hordeIssueLink"></param>
-		Task<(string? key, string? url)> CreateIssueAsync(IUser User, int issueId, string summary, string projectId, string componentId, string issueType, string? description, string? hordeIssueLink);
+		Task<(string? key, string? url)> CreateIssueAsync(IUser user, string? externalIssueUser, int issueId, string summary, string projectId, string componentId, string issueType, string? description, string? hordeIssueLink);
 
 
 		/// <summary>
@@ -103,7 +104,7 @@ namespace Horde.Build.Services
 		}
 
 		/// <inheritdoc/>
-		public Task<(string? key, string? url)> CreateIssueAsync(IUser User, int issueId, string summary, string projectId, string componentId, string issueType, string? description, string? hordeIssueLink)
+		public Task<(string? key, string? url)> CreateIssueAsync(IUser user, string? externalIssueUser, int issueId, string summary, string projectId, string componentId, string issueType, string? description, string? hordeIssueLink)
 		{
 			throw new NotImplementedException();
 		}
diff --git a/Engine/Source/Programs/Horde/Horde.Build/Issues/External/Impl/JiraService.cs b/Engine/Source/Programs/Horde/Horde.Build/Issues/External/Impl/JiraService.cs
--- a/Engine/Source/Programs/Horde/Horde.Build/Issues/External/Impl/JiraService.cs
+++ b/Engine/Source/Programs/Horde/Horde.Build/Issues/External/Impl/JiraService.cs
@@ -372,7 +372,7 @@ namespace Horde.Build.Services
 
 		}
 
-		public async Task<(string? key, string? url)> CreateIssueAsync(IUser User, int issueId, string summary, string projectId, string componentId, string issueType, string? description, string? hordeIssueLink)
+		public async Task<(string? key, string? url)> CreateIssueAsync(IUser user, string? externalIssueUser, int issueId, string summary, string projectId, string componentId, string issueType, string? description, string? hordeIssueLink)
 		{
 
 			IIssue? issue = await _issueService.GetIssueAsync(issueId);
@@ -385,11 +385,11 @@ namespace Horde.Build.Services
 
 			if (hordeIssueLink != null)
 			{
-				issueDesc = $"\n{issueDesc}\n\nJira issue created by {User.Name} for *[Horde Issue {issueId}|{hordeIssueLink}]*\n\n";
+				issueDesc = $"\n{issueDesc}\n\nJira issue created by {user.Name} for *[Horde Issue {issueId}|{hordeIssueLink}]*\n\n";
 			}
 			else
 			{
-				issueDesc = $"\n{issueDesc}\n\nJira Issue created by {User.Name} from Horde\n\n";
+				issueDesc = $"\n{issueDesc}\n\nJira Issue created by {user.Name} from Horde\n\n";
 			}
 
 
@@ -434,6 +434,23 @@ namespace Horde.Build.Services
 
 			await _issueService.UpdateIssueAsync(issueId, externalIssueKey: jiraResponse.Key);
 
+			// add the user as a watcher to the newly created issue
+			if (externalIssueUser != null)
+			{
+				uri = new Uri(_jiraUrl, $"/rest/api/2/issue/{jiraResponse.Key}/watchers");
+
+				try
+				{
+					response = await _retryPolicy.ExecuteAsync(() => _client.PostAsync(uri, new StringContent(JsonSerializer.Serialize(externalIssueUser), Encoding.UTF8, "application/json")));
+					response.EnsureSuccessStatusCode();
+				}
+				catch (Exception ex)
+				{
+					_logger.LogError(ex, "Unable to add {ExternalUser} as watcher to issue {IssueKey}", externalIssueUser, jiraResponse.Key);
+				}
+				
+			}
+
 			return (jiraResponse.Key, $"{_jiraUrl}browse/{jiraResponse.Key}");
 		}
 
diff --git a/Engine/Source/Programs/Horde/Horde.Build/Utilities/HordeClaimTypes.cs b/Engine/Source/Programs/Horde/Horde.Build/Utilities/HordeClaimTypes.cs
--- a/Engine/Source/Programs/Horde/Horde.Build/Utilities/HordeClaimTypes.cs
+++ b/Engine/Source/Programs/Horde/Horde.Build/Utilities/HordeClaimTypes.cs
@@ -58,6 +58,12 @@ namespace Horde.Build.Utilities
 		/// Claim for the Perforce username
 		/// </summary>
 		public const string PerforceUser = Prefix + "perforce-user";
+
+		/// <summary>
+		/// Claim for the external issue username
+		/// </summary>
+		public const string ExternalIssueUser = Prefix + "external-issue-user";
+
 	}
 
 	/// <summary>
@@ -117,5 +123,21 @@ namespace Horde.Build.Utilities
 			}
 			return claim.Value;
 		}
+
+		/// <summary>
+		/// Gets the external issue username for the given principal
+		/// </summary>
+		/// <param name="user">The principal to get the external issue user for</param>
+		/// <returns>Jira user name</returns>
+		public static string? GetExternalIssueUser(this ClaimsPrincipal user)
+		{
+			Claim? claim = user.FindFirst(HordeClaimTypes.ExternalIssueUser);
+			if (claim == null)
+			{
+				return null;
+			}
+			return claim.Value;
+		}
+
 	}
 }
