Horde: Send legacy triage messages to Slack immediately, rather than waiting an hour. Since issues do not necessarily trigger an update after an hour, this can result in notifications not being sent at all.

diff --git a/Engine/Source/Programs/Horde/Horde.Build/Notifications/Impl/SlackNotificationSink.cs b/Engine/Source/Programs/Horde/Horde.Build/Notifications/Impl/SlackNotificationSink.cs
--- a/Engine/Source/Programs/Horde/Horde.Build/Notifications/Impl/SlackNotificationSink.cs
+++ b/Engine/Source/Programs/Horde/Horde.Build/Notifications/Impl/SlackNotificationSink.cs
@@ -616,9 +616,13 @@ namespace Horde.Build.Notifications.Impl
 			// Do not send notifications for quarantined issues
 			if (issue.QuarantinedByUserId != null)
 			{
+				_logger.LogInformation("Skipping notifications for quarantined issue {IssueId}", issue.Id);
 				return;
 			}
 
+			using IDisposable scope = _logger.BeginScope("Slack notifications for issue {IssueId}", issue.Id);
+			_logger.LogInformation("Updating Slack notifications for issue {IssueId}", issue.Id);
+
 			IIssueDetails details = await _issueService.GetIssueDetailsAsync(issue);
 
 			bool notifySuspects = issue.Promoted || details.Spans.Any(x => x.LastFailure.Annotations.NotifySubmitters ?? false);
@@ -663,22 +667,19 @@ namespace Horde.Build.Notifications.Impl
 				}
 			}
 
-			if (issue.OwnerId == null && (details.Suspects.Count == 0 || details.Suspects.All(x => x.DeclinedAt != null) || issue.CreatedAt < DateTime.UtcNow - TimeSpan.FromHours(1.0)))
+			foreach (IIssueSpan span in details.Spans)
 			{
-				foreach (IIssueSpan span in details.Spans)
+				IStream? stream = await _streamService.GetCachedStream(span.StreamId);
+				if (stream != null)
 				{
-					IStream? stream = await _streamService.GetCachedStream(span.StreamId);
-					if (stream != null)
+					TemplateRef? templateRef;
+					if (stream.Templates.TryGetValue(span.TemplateRefId, out templateRef) && templateRef.TriageChannel != null)
 					{
-						TemplateRef? templateRef;
-						if (stream.Templates.TryGetValue(span.TemplateRefId, out templateRef) && templateRef.TriageChannel != null)
-						{
-							channels.Add(templateRef.TriageChannel);
-						}
-						else if (stream.TriageChannel != null)
-						{
-							channels.Add(stream.TriageChannel);
-						}
+						channels.Add(templateRef.TriageChannel);
+					}
+					else if (stream.TriageChannel != null)
+					{
+						channels.Add(stream.TriageChannel);
 					}
 				}
 			}
