Horde: Fix issue log error and use _blank for view log in issue modal
#rnx

diff --git a/Engine/Source/Programs/Horde/Horde.Build/Issues/External/Impl/JiraService.cs b/Engine/Source/Programs/Horde/Horde.Build/Issues/External/Impl/JiraService.cs
--- a/Engine/Source/Programs/Horde/Horde.Build/Issues/External/Impl/JiraService.cs
+++ b/Engine/Source/Programs/Horde/Horde.Build/Issues/External/Impl/JiraService.cs
@@ -241,7 +241,7 @@ namespace Horde.Build.Services
 				}
 				catch (Exception)
 				{
-					_logger.LogError("Unable to get project info for {ProjectKey", projectKey);
+					_logger.LogError("Unable to get project info for {ProjectKey}", projectKey);
 					continue;
 				}
 
@@ -287,7 +287,7 @@ namespace Horde.Build.Services
 				}
 				catch (Exception)
 				{
-					_logger.LogError("Unable to get components info for {ProjectKey", projectKey);
+					_logger.LogError("Unable to get components info for {ProjectKey}", projectKey);
 					continue;
 				}
 
