#Horde: Remove possible permissions issue messaging
#skipci
#rnx

diff --git a/Engine/Source/Programs/Horde/HordeDashboard/src/backend/Fetch.ts b/Engine/Source/Programs/Horde/HordeDashboard/src/backend/Fetch.ts
--- a/Engine/Source/Programs/Horde/HordeDashboard/src/backend/Fetch.ts
+++ b/Engine/Source/Programs/Horde/HordeDashboard/src/backend/Fetch.ts
@@ -72,7 +72,7 @@ export class Fetch {
 
                 // Note: the fetch request sets error for redirects, this is opaque as per the spec
                 // so all we get is an error string and can't detect AccessDenied, etc in redirect :/
-                let message = `Possible permissions issue, ${reason}`;
+                let message = `Error in request, ${reason}`;
 
                 handleError({
                     reason: message,
