Horde: Guard against GetAgentCapabilities() failing and causing the current session to terminate.

diff --git a/Engine/Source/Programs/Horde/Horde.Agent/Services/WorkerService.cs b/Engine/Source/Programs/Horde/Horde.Agent/Services/WorkerService.cs
--- a/Engine/Source/Programs/Horde/Horde.Agent/Services/WorkerService.cs
+++ b/Engine/Source/Programs/Horde/Horde.Agent/Services/WorkerService.cs
@@ -534,7 +534,14 @@ namespace Horde.Agent.Services
 					// Update the capabilities every 5m
 					if (updateCapabilitiesTimer.Elapsed > TimeSpan.FromMinutes(5.0))
 					{
-						updateSessionRequest.Capabilities = await GetAgentCapabilities(_workingDir, _logger);
+						try
+						{
+							updateSessionRequest.Capabilities = await GetAgentCapabilities(_workingDir, _logger);
+						}
+						catch (Exception ex)
+						{
+							_logger.LogWarning(ex, "Unable to query agent capabilities. Ignoring.");
+						}
 						updateCapabilitiesTimer.Restart();
 					}
 
