Horde: Ignore "safe" error return codes from Slack.

diff --git a/Engine/Source/Programs/Horde/Horde.Build/Notifications/Impl/SlackNotificationSink.cs b/Engine/Source/Programs/Horde/Horde.Build/Notifications/Impl/SlackNotificationSink.cs
--- a/Engine/Source/Programs/Horde/Horde.Build/Notifications/Impl/SlackNotificationSink.cs
+++ b/Engine/Source/Programs/Horde/Horde.Build/Notifications/Impl/SlackNotificationSink.cs
@@ -715,12 +715,7 @@ namespace Horde.Build.Notifications.Impl
 			message.Channel = channel;
 			message.Ts = ts;
 			message.Name = name;
-
-			SlackResponse response = await SendRequestAsync<SlackResponse>(AddReactionUrl, message);
-			if(!response.Ok && !String.Equals(response.Error, "already_reacted", StringComparison.Ordinal))
-			{
-				_logger.LogWarning("Error adding reaction {Name} to thread {Channel}/{Ts}: {Error}", name, channel, ts, response.Error);
-			}
+			await SendRequestAsync<SlackResponse>(AddReactionUrl, message);
 		}
 
 		async Task RemoveReactionAsync(string channel, string ts, string name)
@@ -729,12 +724,7 @@ namespace Horde.Build.Notifications.Impl
 			message.Channel = channel;
 			message.Ts = ts;
 			message.Name = name;
-
-			SlackResponse response = await SendRequestAsync<SlackResponse>(RemoveReactionUrl, message);
-			if (!response.Ok && !String.Equals(response.Error, "no_reaction", StringComparison.Ordinal))
-			{
-				_logger.LogWarning("Error removing reaction {Name} from thread {Channel}/{Ts}: {Error}", name, channel, ts, response.Error);
-			}
+			await SendRequestAsync<SlackResponse>(RemoveReactionUrl, message);
 		}
 
 		static string GetTriageThreadEventId(int issueId) => $"issue_triage_{issueId}";
@@ -1711,6 +1701,12 @@ namespace Horde.Build.Notifications.Impl
 			}
 		}
 
+		static readonly string[] s_ignoreErrorCodes =
+		{
+			"no_reaction",
+			"already_reacted"
+		};
+
 		private async Task<TResponse> SendRequestAsync<TResponse>(HttpRequestMessage request, string? requestJson) where TResponse : SlackResponse
 		{
 			using (HttpClient client = new HttpClient())
@@ -1721,7 +1717,7 @@ namespace Horde.Build.Notifications.Impl
 				byte[] responseBytes = await response.Content.ReadAsByteArrayAsync();
 
 				TResponse responseObject = JsonSerializer.Deserialize<TResponse>(responseBytes)!;
-				if (!responseObject.Ok)
+				if (!responseObject.Ok && !s_ignoreErrorCodes.Contains(responseObject.Error, StringComparer.Ordinal))
 				{
 					_logger.LogError("Failed to send Slack message ({Error}). Request: {Request}. Response: {Response}", responseObject.Error, requestJson, Encoding.UTF8.GetString(responseBytes));
 				}
