Horde: Use Robomerge tags to infer in which stream a build breakage originated. We will be able to use this information to notify when a fix failed.

diff --git a/Engine/Source/Programs/Horde/Horde.Build.Tests/IssueServiceTests.cs b/Engine/Source/Programs/Horde/Horde.Build.Tests/IssueServiceTests.cs
--- a/Engine/Source/Programs/Horde/Horde.Build.Tests/IssueServiceTests.cs
+++ b/Engine/Source/Programs/Horde/Horde.Build.Tests/IssueServiceTests.cs
@@ -1486,5 +1486,60 @@ namespace Horde.Build.Tests
 				}
 			}
 		}
+
+		//		static List<IIssueSpan> GetOriginSpans(List<IIssueSpan> spans)
+
+		static IIssueSpanSuspect MockSuspect(int change, int? originatingChange)
+		{
+			Mock<IIssueSpanSuspect> suspect = new Mock<IIssueSpanSuspect>(MockBehavior.Strict);
+			suspect.SetupGet(x => x.Change).Returns(change);
+			suspect.SetupGet(x => x.OriginatingChange).Returns(originatingChange);
+			return suspect.Object;
+		}
+
+		static IIssueSpan MockSpan(StreamId streamId, params IIssueSpanSuspect[] suspects)
+		{
+			Mock<IIssueSpan> span = new Mock<IIssueSpan>(MockBehavior.Strict);
+			span.SetupGet(x => x.StreamId).Returns(streamId);
+			span.SetupGet(x => x.Suspects).Returns(suspects);
+			return span.Object;
+		}
+
+		[TestMethod]
+		public void FindMergeOrigins()
+		{
+			{
+				List<IIssueSpan> spans = new List<IIssueSpan>();
+				spans.Add(MockSpan(new StreamId("ue5-main"), MockSuspect(201, 1), MockSuspect(202, 2)));
+				spans.Add(MockSpan(new StreamId("ue5-release-staging"), MockSuspect(101, 1), MockSuspect(102, 2)));
+				spans.Add(MockSpan(new StreamId("ue5-release"), MockSuspect(1, null), MockSuspect(2, null)));
+
+				List<IIssueSpan> results = Horde.Build.Services.Impl.IssueService.FindMergeOriginSpans(spans);
+				Assert.AreEqual(1, results.Count);
+				Assert.AreEqual(new StreamId("ue5-release"), results[0].StreamId);
+			}
+
+			{
+				List<IIssueSpan> spans = new List<IIssueSpan>();
+				spans.Add(MockSpan(new StreamId("ue5-main"), MockSuspect(201, null), MockSuspect(202, 2)));
+				spans.Add(MockSpan(new StreamId("ue5-release-staging"), MockSuspect(101, 1), MockSuspect(102, null)));
+				spans.Add(MockSpan(new StreamId("ue5-release"), MockSuspect(1, null), MockSuspect(2, null)));
+
+				List<IIssueSpan> results = Horde.Build.Services.Impl.IssueService.FindMergeOriginSpans(spans);
+				Assert.AreEqual(1, results.Count);
+				Assert.AreEqual(new StreamId("ue5-release"), results[0].StreamId);
+			}
+
+			{
+				List<IIssueSpan> spans = new List<IIssueSpan>();
+				spans.Add(MockSpan(new StreamId("ue5-release"), MockSuspect(1, null), MockSuspect(2, null)));
+				spans.Add(MockSpan(new StreamId("ue5-release-staging"), MockSuspect(101, 1), MockSuspect(102, null)));
+				spans.Add(MockSpan(new StreamId("ue5-main"), MockSuspect(201, null), MockSuspect(202, 2)));
+
+				List<IIssueSpan> results = Horde.Build.Services.Impl.IssueService.FindMergeOriginSpans(spans);
+				Assert.AreEqual(1, results.Count);
+				Assert.AreEqual(new StreamId("ue5-release"), results[0].StreamId);
+			}
+		}
 	}
 }
diff --git a/Engine/Source/Programs/Horde/Horde.Build/Issues/IIssue.cs b/Engine/Source/Programs/Horde/Horde.Build/Issues/IIssue.cs
--- a/Engine/Source/Programs/Horde/Horde.Build/Issues/IIssue.cs
+++ b/Engine/Source/Programs/Horde/Horde.Build/Issues/IIssue.cs
@@ -132,6 +132,11 @@ namespace Horde.Build.Models
 		/// </summary>
 		public int? FixChange { get; }
 
+//		/// <summary>
+//		/// The first stream that encountered the error. The fix will be considered failed if an error after FixChange occurs in this stream. 
+//		/// </summary>
+//		public StreamId? OriginStreamId { get; }
+
 		/// <summary>
 		/// List of streams affected by this issue
 		/// </summary>
@@ -174,6 +179,11 @@ namespace Horde.Build.Models
 		/// </summary>
 		StreamId StreamId { get; }
 
+		/// <summary>
+		/// Whether this stream is the merge origin for other streams
+		/// </summary>
+		bool? MergeOrigin { get; }
+
 		/// <summary>
 		/// Whether this stream contains the fix change
 		/// </summary>
diff --git a/Engine/Source/Programs/Horde/Horde.Build/Issues/IIssueCollection.cs b/Engine/Source/Programs/Horde/Horde.Build/Issues/IIssueCollection.cs
--- a/Engine/Source/Programs/Horde/Horde.Build/Issues/IIssueCollection.cs
+++ b/Engine/Source/Programs/Horde/Horde.Build/Issues/IIssueCollection.cs
@@ -352,6 +352,9 @@ namespace Horde.Build.Collections
 		/// <inheritdoc cref="IIssueStream.StreamId"/>
 		public StreamId StreamId { get; set; }
 
+		/// <inheritdoc cref="IIssueStream.MergeOrigin"/>
+		public bool? MergeOrigin { get; set; }
+
 		/// <inheritdoc cref="IIssueStream.ContainsFix"/>
 		public bool? ContainsFix { get; set; }
 
diff --git a/Engine/Source/Programs/Horde/Horde.Build/Issues/Impl/IssueCollection.cs b/Engine/Source/Programs/Horde/Horde.Build/Issues/Impl/IssueCollection.cs
--- a/Engine/Source/Programs/Horde/Horde.Build/Issues/Impl/IssueCollection.cs
+++ b/Engine/Source/Programs/Horde/Horde.Build/Issues/Impl/IssueCollection.cs
@@ -165,6 +165,7 @@ namespace Horde.Build.Collections.Impl
 		class IssueStream : IIssueStream
 		{
 			public StreamId StreamId { get; set; }
+			public bool? MergeOrigin { get; set; }
 			public bool? ContainsFix { get; set; }
 
 			public IssueStream()
@@ -174,6 +175,7 @@ namespace Horde.Build.Collections.Impl
 			public IssueStream(IIssueStream other)
 			{
 				StreamId = other.StreamId;
+				MergeOrigin = other.MergeOrigin;
 				ContainsFix = other.ContainsFix;
 			}
 		}
diff --git a/Engine/Source/Programs/Horde/Horde.Build/Issues/Impl/IssueService.cs b/Engine/Source/Programs/Horde/Horde.Build/Issues/Impl/IssueService.cs
--- a/Engine/Source/Programs/Horde/Horde.Build/Issues/Impl/IssueService.cs
+++ b/Engine/Source/Programs/Horde/Horde.Build/Issues/Impl/IssueService.cs
@@ -910,6 +910,58 @@ namespace Horde.Build.Services.Impl
 			return issue;
 		}
 
+		static void GetChangeMap(IIssueSpan span, Dictionary<int, int> originChangeToSuspectChange)
+		{
+			foreach (IIssueSpanSuspect suspect in span.Suspects)
+			{
+				originChangeToSuspectChange[suspect.OriginatingChange ?? suspect.Change] = suspect.Change;
+			}
+		}
+
+		static int CompareSpans(Dictionary<int, int> originChangeToSuspectChange, IIssueSpan span)
+		{
+			foreach (IIssueSpanSuspect suspect in span.Suspects)
+			{
+				int originChange = suspect.OriginatingChange ?? suspect.Change;
+				if (originChangeToSuspectChange.TryGetValue(originChange, out int suspectChange) && suspectChange != suspect.Change)
+				{
+					return suspectChange - suspect.Change;
+				}
+			}
+			return 0;
+		}
+
+		internal static List<IIssueSpan> FindMergeOriginSpans(List<IIssueSpan> spans)
+		{
+			// Determine the stream(s) highest up in the merge hierarchy, as determined by the minimum changelist number
+			// for any robomerged change.
+			List<IIssueSpan> originSpans = new List<IIssueSpan>();
+			if (spans.Count > 0)
+			{
+				originSpans.Add(spans[0]);
+
+				Dictionary<int, int> originChangeToSuspectChange = new Dictionary<int, int>();
+				GetChangeMap(spans[0], originChangeToSuspectChange);
+
+				for (int idx = 1; idx < spans.Count; idx++)
+				{
+					IIssueSpan span = spans[idx];
+					int comparison = CompareSpans(originChangeToSuspectChange, span);
+					if (comparison > 0)
+					{
+						originChangeToSuspectChange.Clear();
+						originSpans.Clear();
+					}
+					if (comparison >= 0)
+					{
+						GetChangeMap(span, originChangeToSuspectChange);
+						originSpans.Add(span);
+					}
+				}
+			}
+			return originSpans;
+		}
+
 		async Task<IIssue?> UpdateIssueDerivedDataAsync(IIssue issue)
 		{
 			Dictionary<(StreamId, int), bool> fixChangeCache = new Dictionary<(StreamId, int), bool>();
@@ -1038,6 +1090,17 @@ namespace Horde.Build.Services.Impl
 					}
 				}
 
+				// Find the stream where the error was introduced.
+				List<IIssueSpan> mergeOrigins = FindMergeOriginSpans(spans);
+				foreach (IIssueSpan mergeOrigin in mergeOrigins)
+				{
+					NewIssueStream? newStream = newStreams.FirstOrDefault(x => x.StreamId == mergeOrigin.StreamId);
+					if (newStream != null)
+					{
+						newStream.MergeOrigin = true;
+					}
+				}
+
 				// If a fix changelist has been specified and it's not valid, clear it out
 				if (newVerifiedAt == null && newResolvedAt != null)
 				{
