Horde: Add messaging for fixes that do not resolve build health issues, and for fixes that are not applied to the correct stream.

diff --git a/Engine/Source/Programs/Horde/Horde.Build/Notifications/Impl/SlackNotificationSink.cs b/Engine/Source/Programs/Horde/Horde.Build/Notifications/Impl/SlackNotificationSink.cs
--- a/Engine/Source/Programs/Horde/Horde.Build/Notifications/Impl/SlackNotificationSink.cs
+++ b/Engine/Source/Programs/Horde/Horde.Build/Notifications/Impl/SlackNotificationSink.cs
@@ -38,6 +38,7 @@ namespace Horde.Build.Notifications.Impl
 {
 	using JobId = ObjectId<IJob>;
 	using LogId = ObjectId<ILogFile>;
+	using StreamId = StringId<IStream>;
 	using UserId = ObjectId<IUser>;
 	using WorkflowId = StringId<WorkflowConfig>;
 
@@ -662,7 +663,7 @@ namespace Horde.Build.Notifications.Impl
 					IStream? stream = await _streamService.GetStreamAsync(span.StreamId);
 					if(stream != null && stream.Config.TryGetWorkflow(workflowId.Value, out WorkflowConfig? workflow))
 					{
-						await CreateOrUpdateWorkflowThreadAsync(issue, span, workflow);
+						await CreateOrUpdateWorkflowThreadAsync(issue, span, details.Spans, workflow);
 					}
 				}
 			}
@@ -729,7 +730,7 @@ namespace Horde.Build.Notifications.Impl
 
 		static string GetTriageThreadEventId(int issueId) => $"issue_triage_{issueId}";
 
-		async Task CreateOrUpdateWorkflowThreadAsync(IIssue issue, IIssueSpan span, WorkflowConfig workflow)
+		async Task CreateOrUpdateWorkflowThreadAsync(IIssue issue, IIssueSpan span, IReadOnlyList<IIssueSpan> spans, WorkflowConfig workflow)
 		{
 			string? triageChannel = workflow.TriageChannel;
 			if (triageChannel != null)
@@ -739,7 +740,7 @@ namespace Horde.Build.Notifications.Impl
 				string eventId = GetTriageThreadEventId(issue.Id);
 
 				string text = $"{workflow.TriagePrefix}*Issue <{issueUrl}|{issue.Id}>*: {issue.Summary}{workflow.TriageSuffix}";
-				if (issue.ResolvedAt != null)
+				if (!spans.Any(x => x.NextSuccess == null))
 				{
 					text = $"~{text}~";
 				}
@@ -773,7 +774,7 @@ namespace Horde.Build.Notifications.Impl
 					}
 
 					StringBuilder summary = new StringBuilder();
-					summary.AppendLine($"From <{GetJobUrl(span.FirstFailure.JobId)}|{span.FirstFailure.JobName}> / <{GetStepUrl(span.FirstFailure.JobId, span.FirstFailure.StepId)}|{span.NodeName}>:");
+					summary.AppendLine($"From {FormatJobStep(span.FirstFailure, span.NodeName)}:");
 					foreach (ILogEventData eventDataItem in eventDataItems)
 					{
 						summary.AppendLine($"```{eventDataItem.Message}```");
@@ -829,11 +830,43 @@ namespace Horde.Build.Notifications.Impl
 					await AddReactionAsync(state.Channel, state.Ts, "tick");
 				}
 
-				if (issue.ResolvedAt != null && issue.FixChange != null)
+				if (issue.FixChange != null)
 				{
-					string fixedEventId = $"issue_{issue.Id}_fixed_{issue.ResolvedAt.Value}";
-					string fixedMessage = $"Marked as fixed in {FormatChange(issue.FixChange.Value)}";
-					await PostSingleMessageToThreadAsync(triageChannel, state.Ts, fixedEventId, fixedMessage);
+					HashSet<StreamId> originStreams = new HashSet<StreamId>(issue.Streams.Where(x => x.MergeOrigin ?? false).Select(x => x.StreamId));
+
+					IIssueSpan? fixFailedSpan = null;
+					foreach (IIssueSpan originSpan in spans.Where(x => originStreams.Contains(x.StreamId)).OrderBy(x => x.StreamId.ToString()))
+					{
+						if (originSpan.LastFailure.Change > issue.FixChange.Value)
+						{
+							fixFailedSpan = originSpan;
+							break;
+						}
+					}
+
+					if (fixFailedSpan == null)
+					{
+						string fixedEventId = $"issue_{issue.Id}_fixed_{issue.FixChange}";
+						string fixedMessage = $"Marked as fixed in {FormatChange(issue.FixChange.Value)}";
+						await PostSingleMessageToThreadAsync(triageChannel, state.Ts, fixedEventId, fixedMessage);
+					}
+					else
+					{
+						string fixFailedEventId = $"issue_{issue.Id}_fixfailed_{issue.FixChange}";
+						string fixFailedMessage = $"Issue not fixed by {FormatChange(issue.FixChange.Value)}; see {FormatJobStep(fixFailedSpan.LastFailure, fixFailedSpan.NodeName)} in {fixFailedSpan.StreamName}.";
+						await PostSingleMessageToThreadAsync(triageChannel, state.Ts, fixFailedEventId, fixFailedMessage);
+					}
+
+					foreach (IIssueStream stream in issue.Streams)
+					{
+						if ((stream.MergeOrigin ?? false) && !(stream.ContainsFix ?? false))
+						{
+							string streamName = spans.FirstOrDefault(x => x.StreamId == stream.StreamId)?.StreamName ?? stream.StreamId.ToString();
+							string missingEventId = $"issue_{issue.Id}_fixmissing_{issue.FixChange}_{stream.StreamId}";
+							string missingMessage = $"Note: Fix may need manually merging to {streamName}";
+							await PostSingleMessageToThreadAsync(triageChannel, state.Ts, missingEventId, missingMessage);
+						}
+					}
 				}
 
 				if(issue.ExternalIssueKey != null)
@@ -1074,6 +1107,8 @@ namespace Horde.Build.Notifications.Impl
 
 		static string FormatChange(int change) => $"<ugs://change?number={change}|CL {change}>";
 
+		string FormatJobStep(IIssueStep step, string name) => $"<{GetJobUrl(step.JobId)}|{step.JobName}> / <{GetStepUrl(step.JobId, step.StepId)}|{name}>";
+
 		string FormatExternalIssue(string key)
 		{
 			string? url = _externalIssueService.GetIssueUrl(key);
