Horde: Improve error message when template is not defined for a stream during job create

diff --git a/Engine/Source/Programs/Horde/Horde.Build/Controllers/JobsController.cs b/Engine/Source/Programs/Horde/Horde.Build/Controllers/JobsController.cs
--- a/Engine/Source/Programs/Horde/Horde.Build/Controllers/JobsController.cs
+++ b/Engine/Source/Programs/Horde/Horde.Build/Controllers/JobsController.cs
@@ -91,7 +91,7 @@ namespace Horde.Build.Controllers
 			TemplateRef? templateRef;
 			if (!stream.Templates.TryGetValue(templateRefId, out templateRef))
 			{
-				return BadRequest($"Invalid {create.TemplateId} parameter");
+				return BadRequest($"Template {create.TemplateId} is not available for stream {stream.Id}");
 			}
 			if (!await _streamService.AuthorizeAsync(stream, templateRef, AclAction.CreateJob, User, null))
 			{
