Horde: Add issue audit logging for external issue keys and quarantine status
#rnx

diff --git a/Engine/Source/Programs/Horde/Horde.Build/Issues/Impl/IssueCollection.cs b/Engine/Source/Programs/Horde/Horde.Build/Issues/Impl/IssueCollection.cs
--- a/Engine/Source/Programs/Horde/Horde.Build/Issues/Impl/IssueCollection.cs
+++ b/Engine/Source/Programs/Horde/Horde.Build/Issues/Impl/IssueCollection.cs
@@ -572,6 +572,31 @@ namespace Horde.Build.Collections.Impl
 				}
 			}
 
+			if (newIssue.ExternalIssueKey != oldIssue.ExternalIssueKey)
+			{
+				if (newIssue.ExternalIssueKey != null)
+				{
+					issueLogger.LogInformation("Linked to external issue {ExternalIssueKey}", newIssue.ExternalIssueKey);
+				}
+				else
+				{
+					issueLogger.LogInformation("Unlinked from external issue {ExternalIssueKey}", oldIssue.ExternalIssueKey);
+				}
+			}
+
+			if (newIssue.QuarantinedByUserId != oldIssue.QuarantinedByUserId)
+			{
+				if (newIssue.QuarantinedByUserId != null)
+				{
+					issueLogger.LogInformation("Quarantined by {UserId}", newIssue.QuarantinedByUserId);
+				}
+				else
+				{
+					issueLogger.LogInformation("Quarantine cleared");
+				}
+			}
+
+
 			string oldFingerprints = oldIssue.FingerprintsDesc;
 			string newFingerprints = newIssue.FingerprintsDesc;
 			if (oldFingerprints != newFingerprints)
