Horde: Enable P4 annotations on source file properties.

diff --git a/Engine/Source/Programs/Horde/Horde.Agent/Execution/PerforceExecutor.cs b/Engine/Source/Programs/Horde/Horde.Agent/Execution/PerforceExecutor.cs
--- a/Engine/Source/Programs/Horde/Horde.Agent/Execution/PerforceExecutor.cs
+++ b/Engine/Source/Programs/Horde/Horde.Agent/Execution/PerforceExecutor.cs
@@ -193,14 +193,27 @@ namespace Horde.Agent.Execution
 			}
 		}
 
+		PerforceLogger CreatePerforceLogger(ILogger logger)
+		{
+			PerforceLogger perforceLogger = new PerforceLogger(logger);
+			perforceLogger.AddClientView(_workspace.WorkspaceDir, _workspace.StreamView, _job.Change);
+			if (_autoSdkWorkspace != null)
+			{
+				perforceLogger.AddClientView(_autoSdkWorkspace.WorkspaceDir, _autoSdkWorkspace.StreamView, _job.Change);
+			}
+			return perforceLogger;
+		}
+
 		protected override async Task<bool> SetupAsync(BeginStepResponse step, ILogger logger, CancellationToken cancellationToken)
 		{
-			return await SetupAsync(step, _workspace.WorkspaceDir, _sharedStorageDir, _envVars, logger, cancellationToken);
+			PerforceLogger perforceLogger = CreatePerforceLogger(logger);
+			return await SetupAsync(step, _workspace.WorkspaceDir, _sharedStorageDir, _envVars, perforceLogger, cancellationToken);
 		}
 
 		protected override async Task<bool> ExecuteAsync(BeginStepResponse step, ILogger logger, CancellationToken cancellationToken)
 		{
-			return await ExecuteAsync(step, _workspace.WorkspaceDir, _sharedStorageDir, _envVars, logger, cancellationToken);
+			PerforceLogger perforceLogger = CreatePerforceLogger(logger);
+			return await ExecuteAsync(step, _workspace.WorkspaceDir, _sharedStorageDir, _envVars, perforceLogger, cancellationToken);
 		}
 
 		public override async Task FinalizeAsync(ILogger logger, CancellationToken cancellationToken)
diff --git a/Engine/Source/Programs/Horde/Horde.Agent/Utility/PerforceLogger.cs b/Engine/Source/Programs/Horde/Horde.Agent/Utility/PerforceLogger.cs
--- a/Engine/Source/Programs/Horde/Horde.Agent/Utility/PerforceLogger.cs
+++ b/Engine/Source/Programs/Horde/Horde.Agent/Utility/PerforceLogger.cs
@@ -60,7 +60,7 @@ namespace Horde.Agent.Utility
 		/// </summary>
 		public void AddClientView(DirectoryReference baseDir, PerforceViewMap viewMap, int change)
 		{
-			_clients.Add(new ClientView(baseDir, viewMap, change));
+			_clients.Add(new ClientView(baseDir, viewMap.Invert(), change));
 		}
 
 		/// <inheritdoc/>
