Horde: Fix tests.

diff --git a/Engine/Source/Programs/Horde/Horde.Agent/Utility/PerforceLogger.cs b/Engine/Source/Programs/Horde/Horde.Agent/Utility/PerforceLogger.cs
--- a/Engine/Source/Programs/Horde/Horde.Agent/Utility/PerforceLogger.cs
+++ b/Engine/Source/Programs/Horde/Horde.Agent/Utility/PerforceLogger.cs
@@ -51,7 +51,7 @@ namespace Horde.Agent.Utility
 		public void AddClientView(DirectoryReference baseDir, string depotPath, int change)
 		{
 			PerforceViewMap viewMap = new PerforceViewMap();
-			viewMap.Entries.Add(new PerforceViewMapEntry(true, "...", depotPath));
+			viewMap.Entries.Add(new PerforceViewMapEntry(true, depotPath, "..."));
 			AddClientView(baseDir, viewMap, change);
 		}
 
