Horde: Adding model constraints to device service
#rnx

diff --git a/Engine/Source/Programs/Horde/Horde.Build/Controllers/DevicesController.cs b/Engine/Source/Programs/Horde/Horde.Build/Controllers/DevicesController.cs
--- a/Engine/Source/Programs/Horde/Horde.Build/Controllers/DevicesController.cs
+++ b/Engine/Source/Programs/Horde/Horde.Build/Controllers/DevicesController.cs
@@ -726,23 +726,21 @@ namespace Horde.Build.Controllers
 
 			List<DeviceRequestData> requestedDevices = new List<DeviceRequestData>();
 
-			List<string> perfSpecs = new List<string>();
+			HashSet<string> perfSpecs = new HashSet<string> { "Unspecified", "Minimum", "Recommended", "High" };
 
 			foreach (string deviceType in request.DeviceTypes)
 			{
 
 				string platformName = deviceType;
-				string perfSpecName = "Unspecified";
+				string? constraint = null;
 
 				if (deviceType.Contains(':', StringComparison.Ordinal))
 				{
 					string[] tokens = deviceType.Split(":");
 					platformName = tokens[0];
-					perfSpecName = tokens[1];
+					constraint = tokens[1];
 				}
 
-				perfSpecs.Add(perfSpecName);
-
 				DevicePlatformId platformId;
 
 				DevicePlatformMapV1 mapV1 = await _deviceService.GetPlatformMapV1();
@@ -768,22 +766,33 @@ namespace Horde.Build.Controllers
 				List<string> includeModels = new List<string>();
 				List<string> excludeModels = new List<string>();
 
-				if (perfSpecName == "High")
+				if (constraint != null && perfSpecs.Contains(constraint))
 				{
-					string? model = null;
-					if (mapV1.PerfSpecHighMap.TryGetValue(platformId, out model))
+					if (constraint == "High")
+					{
+						string? model = null;
+						if (mapV1.PerfSpecHighMap.TryGetValue(platformId, out model))
+						{
+							includeModels.Add(model);
+						}
+					}
+					else if (constraint == "Minimum" || constraint == "Recommended")
 					{
-						includeModels.Add(model);
+						string? model = null;
+						if (mapV1.PerfSpecHighMap.TryGetValue(platformId, out model))
+						{
+							excludeModels.Add(model);
+						}
 					}
 				}
-
-				if (perfSpecName == "Minimum" || perfSpecName == "Recommended")
+				else if (constraint != null)
 				{
-					string? model = null;
-					if (mapV1.PerfSpecHighMap.TryGetValue(platformId, out model))
+					string? model = platform.Models?.FirstOrDefault(x => x == constraint);
+					if (model == null)
 					{
-						excludeModels.Add(model);
+						return NotFound($"Platform {platform.Id} has no model {model}");
 					}
+					includeModels.Add(constraint);
 				}
 
 				requestedDevices.Add(new DeviceRequestData(platformId, platformName, includeModels, excludeModels));
