Horde: Remove span check when filtering resolved state for issues
#rnx

diff --git a/Engine/Source/Programs/Horde/Horde.Build/Controllers/IssuesController.cs b/Engine/Source/Programs/Horde/Horde.Build/Controllers/IssuesController.cs
--- a/Engine/Source/Programs/Horde/Horde.Build/Controllers/IssuesController.cs
+++ b/Engine/Source/Programs/Horde/Horde.Build/Controllers/IssuesController.cs
@@ -120,12 +120,6 @@ namespace Horde.Build.Controllers
 						List<FindIssueSpanResponse> spanResponses = new List<FindIssueSpanResponse>();
 						if (issueIdToSpans.TryGetValue(issue.Id, out List<IIssueSpan>? spansForIssue))
 						{
-							// Filter issues on resolved state
-							if (resolved != null && (resolved.Value != spansForIssue.All(x => x.NextSuccess != null)))
-							{
-								continue;
-							}
-
 							// Find the current severity in the stream
 							DateTime lastStepTime = DateTime.MinValue;
 							foreach (IIssueSpan span in spansForIssue)
