Horde: Adding issue audit logging for changing ResolvedAt and VerifiedAt times
#rnx

diff --git a/Engine/Source/Programs/Horde/Horde.Build/Issues/Impl/IssueCollection.cs b/Engine/Source/Programs/Horde/Horde.Build/Issues/Impl/IssueCollection.cs
--- a/Engine/Source/Programs/Horde/Horde.Build/Issues/Impl/IssueCollection.cs
+++ b/Engine/Source/Programs/Horde/Horde.Build/Issues/Impl/IssueCollection.cs
@@ -572,6 +572,31 @@ namespace Horde.Build.Collections.Impl
 				}
 			}
 
+			if (newIssue.ResolvedAt != oldIssue.ResolvedAt)
+			{
+				if (newIssue.ResolvedAt == null)
+				{
+					issueLogger.LogInformation("Clearing resolved at time");
+				}
+				else
+				{
+					issueLogger.LogInformation("Setting resolved at time");
+				}
+			}
+
+			if (newIssue.VerifiedAt != oldIssue.VerifiedAt)
+			{
+				if (newIssue.VerifiedAt == null)
+				{
+					issueLogger.LogInformation("Clearing verified at time");
+				}
+				else
+				{
+					issueLogger.LogInformation("Setting verified at time");
+				}
+			}
+
+
 			if (newIssue.ExternalIssueKey != oldIssue.ExternalIssueKey)
 			{
 				if (newIssue.ExternalIssueKey != null)
