Horde Storage - Added POST endpoint to compressed-blobs that allows you to send a CompressedBuffer without knowing its uncompressed content hash (as we will calculate this anyway and the header alread contains it).

diff --git a/Engine/Source/Programs/Horde/HordeStorage/Horde.Storage/Controllers/CompressedBlobController.cs b/Engine/Source/Programs/Horde/HordeStorage/Horde.Storage/Controllers/CompressedBlobController.cs
--- a/Engine/Source/Programs/Horde/HordeStorage/Horde.Storage/Controllers/CompressedBlobController.cs
+++ b/Engine/Source/Programs/Horde/HordeStorage/Horde.Storage/Controllers/CompressedBlobController.cs
@@ -286,15 +286,64 @@ namespace Horde.Storage.Controllers
             }
         }
 
-        private async Task<ContentId> PutImpl(NamespaceId ns, ContentId id, IBufferedPayload payload)
+        [HttpPost("{ns}")]
+        [Authorize("Storage.write")]
+        [DisableRequestSizeLimit]
+        [RequiredContentType(CustomMediaTypeNames.UnrealCompressedBuffer)]
+        public async Task<IActionResult> Post(
+            [Required] NamespaceId ns)
+        {
+            ActionResult? result = await _requestHelper.HasAccessToNamespace(User, Request, ns);
+            if (result != null)
+            {
+                return result;
+            }
+
+            _diagnosticContext.Set("Content-Length", Request.ContentLength ?? -1);
+
+            using IBufferedPayload payload = await _bufferedPayloadFactory.CreateFromRequest(Request);
+
+            try
+            {
+                ContentId identifier = await PutImpl(ns, null, payload);
+
+                return Ok(new
+                {
+                    Identifier = identifier.ToString()
+                });
+            }
+            catch (HashMismatchException e)
+            {
+                return BadRequest(new ProblemDetails
+                {
+                    Title = $"Incorrect hash, got hash \"{e.SuppliedHash}\" but hash of content was determined to be \"{e.ContentHash}\""
+                });
+            }
+        }
+
+        private async Task<ContentId> PutImpl(NamespaceId ns, ContentId? id, IBufferedPayload payload)
         {
-            // decompress the content and generate a identifier from it to verify the identifier we got
+            // decompress the content and generate a identifier from it to verify the identifier we got if we got one
             await using Stream decompressStream = payload.GetStream();
             // TODO: we should add a overload for decompress content that can work on streams, otherwise we are still limited to 2GB compressed blobs
             byte[] decompressedContent = _compressedBufferUtils.DecompressContent(await decompressStream.ToByteArray());
 
             await using MemoryStream decompressedStream = new MemoryStream(decompressedContent);
-            ContentId identifierDecompressedPayload = ContentId.FromContentHash(await _storage.VerifyContentMatchesHash(decompressedStream, id));
+            ContentId identifierDecompressedPayload;
+            if (id != null)
+            {
+                identifierDecompressedPayload = ContentId.FromContentHash(await _storage.VerifyContentMatchesHash(decompressedStream, id));
+            }
+            else
+            {
+                ContentHash blobHash;
+                {
+                    using IScope _ = Tracer.Instance.StartActive("web.hash");
+                    blobHash = await BlobIdentifier.FromStream(decompressedStream);
+                }
+
+                identifierDecompressedPayload = ContentId.FromContentHash(blobHash);
+            }
 
             BlobIdentifier identifierCompressedPayload;
             {
diff --git a/Engine/Source/Programs/Horde/HordeStorage/Horde.Storage/Implementation/ContentId/CompressedBufferUtils.cs b/Engine/Source/Programs/Horde/HordeStorage/Horde.Storage/Implementation/ContentId/CompressedBufferUtils.cs
--- a/Engine/Source/Programs/Horde/HordeStorage/Horde.Storage/Implementation/ContentId/CompressedBufferUtils.cs
+++ b/Engine/Source/Programs/Horde/HordeStorage/Horde.Storage/Implementation/ContentId/CompressedBufferUtils.cs
@@ -3,6 +3,7 @@
 using System;
 using System.Buffers.Binary;
 using System.IO;
+using Datadog.Trace;
 using Force.Crc32;
 using Jupiter.Implementation;
 using Jupiter.Utils;
@@ -193,6 +194,8 @@ namespace Horde.Storage.Implementation
             }
 
             {
+                using IScope _ = Tracer.Instance.StartActive("web.hash");
+
                 // only read the first 20 bytes of the hash field as IoHashes are 20 bytes and not 32 bytes
                 byte[] slicedHash = new byte[20];
                 Array.Copy(header.RawHash, 0, slicedHash, 0, 20);
diff --git a/Engine/Source/Programs/Horde/HordeStorage/Horde.Storage/Tests/Functional/ContentId/CompressedBlobTests.cs b/Engine/Source/Programs/Horde/HordeStorage/Horde.Storage/Tests/Functional/ContentId/CompressedBlobTests.cs
--- a/Engine/Source/Programs/Horde/HordeStorage/Horde.Storage/Tests/Functional/ContentId/CompressedBlobTests.cs
+++ b/Engine/Source/Programs/Horde/HordeStorage/Horde.Storage/Tests/Functional/ContentId/CompressedBlobTests.cs
@@ -203,6 +203,27 @@ namespace Horde.Storage.FunctionalTests.CompressedBlobs
             }
         }
 
+        [TestMethod]
+        public async Task PostNoIdentifier()
+        {
+            byte[] texturePayload = await File.ReadAllBytesAsync("ContentId/Payloads/UncompressedTexture_CAS_dea81b6c3b565bb5089695377c98ce0f1c13b0c3.udd");
+            BlobIdentifier compressedPayloadIdentifier = BlobIdentifier.FromBlob(texturePayload);
+            BlobIdentifier uncompressedPayloadIdentifier = new BlobIdentifier("DEA81B6C3B565BB5089695377C98CE0F1C13B0C3");
+            {
+                using ByteArrayContent content = new(texturePayload);
+                content.Headers.ContentType = new MediaTypeHeaderValue(CustomMediaTypeNames.UnrealCompressedBuffer);
+                // we purposefully use the compressed identifier here which is not what is expected
+                HttpResponseMessage result = await Client!.PostAsync(new Uri($"api/v1/compressed-blobs/{TestNamespace}", UriKind.Relative), content);
+
+                result.EnsureSuccessStatusCode();
+
+                InsertResponse response = await result.Content.ReadAsAsync<InsertResponse>();
+                Assert.IsNotNull(response.Identifier);
+                Assert.AreEqual(uncompressedPayloadIdentifier, response.Identifier);
+                Assert.AreNotEqual(compressedPayloadIdentifier, response.Identifier);
+            }
+        }
+
         [TestMethod]
         public async Task GetUncompressedContent()
         {
