Horde: Guard against an empty device constraint
#rnx

diff --git a/Engine/Source/Programs/Horde/Horde.Build/Controllers/DevicesController.cs b/Engine/Source/Programs/Horde/Horde.Build/Controllers/DevicesController.cs
--- a/Engine/Source/Programs/Horde/Horde.Build/Controllers/DevicesController.cs
+++ b/Engine/Source/Programs/Horde/Horde.Build/Controllers/DevicesController.cs
@@ -739,6 +739,10 @@ namespace Horde.Build.Controllers
 					string[] tokens = deviceType.Split(":");
 					platformName = tokens[0];
 					constraint = tokens[1];
+					if (constraint.Length == 0)
+					{
+						constraint = null;
+					}
 				}
 
 				DevicePlatformId platformId;
