Horde: Backout device controller changes for now
#rnx

diff --git a/Engine/Source/Programs/Horde/Horde.Build/Controllers/DevicesController.cs b/Engine/Source/Programs/Horde/Horde.Build/Controllers/DevicesController.cs
--- a/Engine/Source/Programs/Horde/Horde.Build/Controllers/DevicesController.cs
+++ b/Engine/Source/Programs/Horde/Horde.Build/Controllers/DevicesController.cs
@@ -726,25 +726,23 @@ namespace Horde.Build.Controllers
 
 			List<DeviceRequestData> requestedDevices = new List<DeviceRequestData>();
 
-			HashSet<string> perfSpecs = new HashSet<string> { "Unspecified", "Minimum", "Recommended", "High" };
+			List<string> perfSpecs = new List<string>();
 
 			foreach (string deviceType in request.DeviceTypes)
 			{
 
 				string platformName = deviceType;
-				string? constraint = null;
+				string perfSpecName = "Unspecified";
 
 				if (deviceType.Contains(':', StringComparison.Ordinal))
 				{
 					string[] tokens = deviceType.Split(":");
 					platformName = tokens[0];
-					constraint = tokens[1];
-					if (constraint.Length == 0)
-					{
-						constraint = null;
-					}
+					perfSpecName = tokens[1];
 				}
 
+				perfSpecs.Add(perfSpecName);
+
 				DevicePlatformId platformId;
 
 				DevicePlatformMapV1 mapV1 = await _deviceService.GetPlatformMapV1();
@@ -770,33 +768,22 @@ namespace Horde.Build.Controllers
 				List<string> includeModels = new List<string>();
 				List<string> excludeModels = new List<string>();
 
-				if (constraint != null && perfSpecs.Contains(constraint))
+				if (perfSpecName == "High")
 				{
-					if (constraint == "High")
-					{
-						string? model = null;
-						if (mapV1.PerfSpecHighMap.TryGetValue(platformId, out model))
-						{
-							includeModels.Add(model);
-						}
-					}
-					else if (constraint == "Minimum" || constraint == "Recommended")
+					string? model = null;
+					if (mapV1.PerfSpecHighMap.TryGetValue(platformId, out model))
 					{
-						string? model = null;
-						if (mapV1.PerfSpecHighMap.TryGetValue(platformId, out model))
-						{
-							excludeModels.Add(model);
-						}
+						includeModels.Add(model);
 					}
 				}
-				else if (constraint != null)
+
+				if (perfSpecName == "Minimum" || perfSpecName == "Recommended")
 				{
-					string? model = platform.Models?.FirstOrDefault(x => x == constraint);
-					if (model == null)
+					string? model = null;
+					if (mapV1.PerfSpecHighMap.TryGetValue(platformId, out model))
 					{
-						return NotFound($"Platform {platform.Id} has no model {model}");
+						excludeModels.Add(model);
 					}
-					includeModels.Add(constraint);
 				}
 
 				requestedDevices.Add(new DeviceRequestData(platformId, platformName, includeModels, excludeModels));
