Horde: Output more information about exceptions when parsing JSON log text.

diff --git a/Engine/Source/Programs/Horde/Horde.Build/Logs/Impl/LogFileService.cs b/Engine/Source/Programs/Horde/Horde.Build/Logs/Impl/LogFileService.cs
--- a/Engine/Source/Programs/Horde/Horde.Build/Logs/Impl/LogFileService.cs
+++ b/Engine/Source/Programs/Horde/Horde.Build/Logs/Impl/LogFileService.cs
@@ -160,5 +160,15 @@
 		Task<Stream> OpenRawStreamAsync(ILogFile logFile, long offset, long length);
 
 		/// <summary>
+		/// Parses a stream of json text and outputs plain text
+		/// </summary>
+		/// <param name="logFile">The log file to query</param>
+		/// <param name="offset">Offset within the data to copy from</param>
+		/// <param name="length">Length of the data to copy</param>
+		/// <param name="outputStream">Output stream to receive the text data</param>
+		/// <returns>Async text</returns>
+		Task CopyPlainTextStreamAsync(ILogFile logFile, long offset, long length, Stream outputStream);
+
+		/// <summary>
 		/// Gets the offset of the given line number
 		/// </summary>
@@ -200,76 +210,6 @@
 				await stream.CopyToAsync(outputStream);
 			}
 		}
-
-		/// <summary>
-		/// Parses a stream of json text and outputs plain text
-		/// </summary>
-		/// <param name="logFileService">The log file service</param>
-		/// <param name="logFile">The log file to query</param>
-		/// <param name="offset">Offset within the data to copy from</param>
-		/// <param name="length">Length of the data to copy</param>
-		/// <param name="outputStream">Output stream to receive the text data</param>
-		/// <returns>Async text</returns>
-		public static async Task CopyPlainTextStreamAsync(this ILogFileService logFileService, ILogFile logFile, long offset, long length, Stream outputStream)
-		{
-			using (Stream stream = await logFileService.OpenRawStreamAsync(logFile, 0, Int64.MaxValue))
-			{
-				byte[] readBuffer = new byte[4096];
-				int readBufferLength = 0;
-
-				byte[] writeBuffer = new byte[4096];
-				int writeBufferLength = 0;
-
-				while (length > 0)
-				{
-					// Add more data to the buffer
-					int readBytes = await stream.ReadAsync(readBuffer.AsMemory(readBufferLength, readBuffer.Length - readBufferLength));
-					readBufferLength += readBytes;
-
-					// Copy as many lines as possible to the output
-					int convertedBytes = 0;
-					for (int endIdx = 1; endIdx < readBufferLength; endIdx++)
-					{
-						if (readBuffer[endIdx] == '\n')
-						{
-							writeBufferLength = LogText.ConvertToPlainText(readBuffer.AsSpan(convertedBytes, endIdx - convertedBytes), writeBuffer, writeBufferLength);
-							convertedBytes = endIdx + 1;
-						}
-					}
-
-					// If there's anything in the write buffer, write it out
-					if (writeBufferLength > 0)
-					{
-						if (offset < writeBufferLength)
-						{
-							int writeLength = (int)Math.Min((long)writeBufferLength - offset, length);
-							await outputStream.WriteAsync(writeBuffer.AsMemory((int)offset, writeLength));
-							length -= writeLength;
-						}
-						offset = Math.Max(offset - writeBufferLength, 0);
-						writeBufferLength = 0;
-					}
-
-					// If we were able to read something, shuffle down the rest of the buffer. Otherwise expand the read buffer.
-					if (convertedBytes > 0)
-					{
-						Buffer.BlockCopy(readBuffer, convertedBytes, readBuffer, 0, readBufferLength - convertedBytes);
-						readBufferLength -= convertedBytes;
-					}
-					else if (readBufferLength > 0)
-					{
-						Array.Resize(ref readBuffer, readBuffer.Length + 128);
-						writeBuffer = new byte[readBuffer.Length];
-					}
-
-					// Exit if we didn't read anything in this iteration
-					if (readBytes == 0)
-					{
-						break;
-					}
-				}
-			}
-		}
 	}
 	
 	/// <summary>
@@ -940,6 +880,85 @@
 			}
 		}
 
+		/// <summary>
+		/// Helper method for catching exceptions in <see cref="LogText.ConvertToPlainText(ReadOnlySpan{Byte}, Byte[], Int32)"/>
+		/// </summary>
+		public static int GuardedConvertToPlainText(ReadOnlySpan<byte> input, byte[] output, int outputOffset, ILogger Logger)
+		{
+			try
+			{
+				return LogText.ConvertToPlainText(input, output, outputOffset);
+			}
+			catch (Exception Ex)
+			{
+				Logger.LogWarning(Ex, "Unable to convert log line to plain text: {Line}", Encoding.UTF8.GetString(input));
+				output[outputOffset] = (byte)'\n';
+				return outputOffset + 1;
+			}
+		}
+
 		/// <inheritdoc/>
+		public async Task CopyPlainTextStreamAsync(ILogFile logFile, long offset, long length, Stream outputStream)
+		{
+			using (Stream stream = await OpenRawStreamAsync(logFile, 0, Int64.MaxValue))
+			{
+				byte[] readBuffer = new byte[4096];
+				int readBufferLength = 0;
+
+				byte[] writeBuffer = new byte[4096];
+				int writeBufferLength = 0;
+
+				while (length > 0)
+				{
+					// Add more data to the buffer
+					int readBytes = await stream.ReadAsync(readBuffer.AsMemory(readBufferLength, readBuffer.Length - readBufferLength));
+					readBufferLength += readBytes;
+
+					// Copy as many lines as possible to the output
+					int convertedBytes = 0;
+					for (int endIdx = 1; endIdx < readBufferLength; endIdx++)
+					{
+						if (readBuffer[endIdx] == '\n')
+						{
+							writeBufferLength = GuardedConvertToPlainText(readBuffer.AsSpan(convertedBytes, endIdx - convertedBytes), writeBuffer, writeBufferLength, _logger);
+							convertedBytes = endIdx + 1;
+						}
+					}
+
+					// If there's anything in the write buffer, write it out
+					if (writeBufferLength > 0)
+					{
+						if (offset < writeBufferLength)
+						{
+							int writeLength = (int)Math.Min((long)writeBufferLength - offset, length);
+							await outputStream.WriteAsync(writeBuffer.AsMemory((int)offset, writeLength));
+							length -= writeLength;
+						}
+						offset = Math.Max(offset - writeBufferLength, 0);
+						writeBufferLength = 0;
+					}
+
+					// If we were able to read something, shuffle down the rest of the buffer. Otherwise expand the read buffer.
+					if (convertedBytes > 0)
+					{
+						Buffer.BlockCopy(readBuffer, convertedBytes, readBuffer, 0, readBufferLength - convertedBytes);
+						readBufferLength -= convertedBytes;
+					}
+					else if (readBufferLength > 0)
+					{
+						Array.Resize(ref readBuffer, readBuffer.Length + 128);
+						writeBuffer = new byte[readBuffer.Length];
+					}
+
+					// Exit if we didn't read anything in this iteration
+					if (readBytes == 0)
+					{
+						break;
+					}
+				}
+			}
+		}
+
+		/// <inheritdoc/>
 		public async Task<(int, long)> GetLineOffsetAsync(ILogFile logFile, int lineIdx)
 		{
