Horde: Rename environment variable to enable structured logging to UE_LOG_JSON_TO_STDOUT.

diff --git a/Engine/Source/Programs/Horde/Horde.Agent/Execution/BuildGraphExecutor.cs b/Engine/Source/Programs/Horde/Horde.Agent/Execution/BuildGraphExecutor.cs
--- a/Engine/Source/Programs/Horde/Horde.Agent/Execution/BuildGraphExecutor.cs
+++ b/Engine/Source/Programs/Horde/Horde.Agent/Execution/BuildGraphExecutor.cs
@@ -796,11 +796,8 @@ namespace Horde.Agent.Execution
 			newEnvironment["UE_HORDE_BATCHID"] = _batchId;
 			newEnvironment["UE_HORDE_STEPID"] = step.StepId;
 
-			// Request structured JSON output directly from tools
-			newEnvironment["UE_STDOUT_JSON"] = "1";
-
 			// Enable structured logging output
-			newEnvironment["UE_LOG_JSON"] = "1";
+			newEnvironment["UE_LOG_JSON_TO_STDOUT"] = "1";
 
 			// Pass the location of the cleanup script to the job
 			FileReference cleanupScript = GetCleanupScript(workspaceDir);
