Horde: Do not notify the owner of an issue if there's a triage thread about it.

diff --git a/Engine/Source/Programs/Horde/Horde.Build/Notifications/Impl/SlackNotificationSink.cs b/Engine/Source/Programs/Horde/Horde.Build/Notifications/Impl/SlackNotificationSink.cs
--- a/Engine/Source/Programs/Horde/Horde.Build/Notifications/Impl/SlackNotificationSink.cs
+++ b/Engine/Source/Programs/Horde/Horde.Build/Notifications/Impl/SlackNotificationSink.cs
@@ -626,6 +626,23 @@
 
 			IIssueDetails details = await _issueService.GetIssueDetailsAsync(issue);
 
+			WorkflowConfig? workflow = null;
+			if (details.Spans.Count > 0)
+			{
+				IIssueSpan span = details.Spans[0];
+
+				WorkflowId? workflowId = span.LastFailure.Annotations.WorkflowId;
+				if (workflowId != null)
+				{
+					IStream? stream = await _streamService.GetStreamAsync(span.StreamId);
+					if (stream != null && stream.Config.TryGetWorkflow(workflowId.Value, out workflow))
+					{
+						await CreateOrUpdateWorkflowThreadAsync(issue, span, details.Spans, workflow);
+					}
+				}
+			}
+
+			bool notifyOwner = workflow?.TriageChannel == null;
 			bool notifySuspects = issue.Promoted || details.Spans.Any(x => x.LastFailure.Annotations.NotifySubmitters ?? false);
 
 			HashSet<UserId> userIds = new HashSet<UserId>();
@@ -633,7 +650,7 @@
 			{
 				userIds.UnionWith(details.Suspects.Select(x => x.AuthorId));
 			}
-			if (issue.OwnerId.HasValue)
+			if (notifyOwner && issue.OwnerId.HasValue)
 			{
 				userIds.Add(issue.OwnerId.Value);
 			}
@@ -653,21 +670,6 @@
 				}
 			}
 
-			if (details.Spans.Count > 0)
-			{
-				IIssueSpan span = details.Spans[0];
-
-				WorkflowId? workflowId = span.LastFailure.Annotations.WorkflowId;
-				if (workflowId != null)
-				{
-					IStream? stream = await _streamService.GetStreamAsync(span.StreamId);
-					if(stream != null && stream.Config.TryGetWorkflow(workflowId.Value, out WorkflowConfig? workflow))
-					{
-						await CreateOrUpdateWorkflowThreadAsync(issue, span, details.Spans, workflow);
-					}
-				}
-			}
-
 			foreach (IIssueSpan span in details.Spans)
 			{
 				IStream? stream = await _streamService.GetCachedStream(span.StreamId);
