Horde: Add a more descriptive error for failing to parse log text.

diff --git a/Engine/Source/Programs/Horde/Horde.Build/Logs/Data/LogText.cs b/Engine/Source/Programs/Horde/Horde.Build/Logs/Data/LogText.cs
--- a/Engine/Source/Programs/Horde/Horde.Build/Logs/Data/LogText.cs
+++ b/Engine/Source/Programs/Horde/Horde.Build/Logs/Data/LogText.cs
@@ -11,6 +11,20 @@
 namespace Horde.Build.Logs
 {
 	/// <summary>
+	/// Exception thrown trying to parse log output
+	/// </summary>
+	public class LogTextParseException : Exception
+	{
+		/// <summary>
+		/// Constructor
+		/// </summary>
+		public LogTextParseException(string message, Exception? ex)
+			: base(message, ex)
+		{
+		}
+	}
+
+	/// <summary>
 	/// Interface for log data
 	/// </summary>
 	public interface ILogText
@@ -237,7 +251,14 @@
 				}
 
 				// Convert the line to plain text
-				Length = ConvertToPlainText(inputLine, _internalData, Length);
+				try
+				{
+					Length = ConvertToPlainText(inputLine, _internalData, Length);
+				}
+				catch (Exception ex)
+				{
+					throw new LogTextParseException($"Exception while parsing log text as JSON. Line: {Encoding.UTF8.GetString(inputLine)}", ex);
+				}
 				_internalLineOffsets.Add(Length);
 			}
 		}
