Horde: Add *.inl as a code extension.

diff --git a/Engine/Source/Programs/Horde/Horde.Build/Issues/IIssueHandler.cs b/Engine/Source/Programs/Horde/Horde.Build/Issues/IIssueHandler.cs
--- a/Engine/Source/Programs/Horde/Horde.Build/Issues/IIssueHandler.cs
+++ b/Engine/Source/Programs/Horde/Horde.Build/Issues/IIssueHandler.cs
@@ -22,6 +22,7 @@
 			".c",
 			".cc",
 			".cpp",
+			".inl",
 			".m",
 			".mm",
 			".rc",
diff --git a/Engine/Source/Programs/Horde/Horde.Build/Perforce/Change.cs b/Engine/Source/Programs/Horde/Horde.Build/Perforce/Change.cs
--- a/Engine/Source/Programs/Horde/Horde.Build/Perforce/Change.cs
+++ b/Engine/Source/Programs/Horde/Horde.Build/Perforce/Change.cs
@@ -139,6 +139,7 @@
 			".c",
 			".cc",
 			".cpp",
+			".inl",
 			".m",
 			".mm",
 			".rc",
