Horde: Prevent BlockKit messages being truncated and losing formatting.

diff --git a/Engine/Source/Programs/Horde/Horde.Build/Notifications/Impl/SlackNotificationSink.cs b/Engine/Source/Programs/Horde/Horde.Build/Notifications/Impl/SlackNotificationSink.cs
--- a/Engine/Source/Programs/Horde/Horde.Build/Notifications/Impl/SlackNotificationSink.cs
+++ b/Engine/Source/Programs/Horde/Horde.Build/Notifications/Impl/SlackNotificationSink.cs
@@ -524,10 +524,10 @@
 				if (errors.Any())
 				{
 					string errorSummary = errors.Count > MaxJobStepEvents ? $"*Errors (First {MaxJobStepEvents} shown)*" : $"*Errors*";
-					eventStrings.Add(errorSummary);
+					attachment.Blocks.Add(new SectionBlock(errorSummary));
 					foreach (ILogEventData error in errors.Take(MaxJobStepEvents))
 					{
-						eventStrings.Add($"```{error.Message}```");
+						attachment.Blocks.Add(QuoteBlock(error.Message));
 					}
 				}
 				else if (warnings.Any())
@@ -536,11 +536,10 @@
 					eventStrings.Add(warningSummary);
 					foreach (ILogEventData warning in warnings.Take(MaxJobStepEvents))
 					{
-						eventStrings.Add($"```{warning.Message}```");
+						attachment.Blocks.Add(QuoteBlock(warning.Message));
 					}
 				}
 
-				attachment.Blocks.Add(new SectionBlock(String.Join("\n", eventStrings)));
 				attachment.Blocks.Add(new SectionBlock($"<{jobStepLogLink}|View Job Step Log>"));
 			}
 
@@ -960,18 +959,15 @@
 						events.RemoveAll(x => x.Severity == EventSeverity.Warning);
 					}
 
-					List<string> eventStrings = new List<string>();
 					for (int idx = 0; idx < Math.Min(events.Count, 3); idx++)
 					{
 						ILogEventData data = await _logFileService.GetEventDataAsync(logFile, events[idx].LineIndex, events[idx].LineCount);
-						eventStrings.Add($"```{data.Message}```");
+						attachment.Blocks.Add(QuoteBlock(data.Message));
 					}
 					if (events.Count > 3)
 					{
-						eventStrings.Add("```...```");
+						attachment.Blocks.Add(new SectionBlock("```...```"));
 					}
-
-					attachment.Blocks.Add(new SectionBlock(String.Join("\n", eventStrings)));
 				}
 			}
 
@@ -1346,7 +1342,7 @@
 			attachment.Blocks.Add(new HeaderBlock($"Config Update Failure :rip:", true));
 
 			attachment.Blocks.Add(new SectionBlock($"Horde was unable to update {fileName}"));
-			attachment.Blocks.Add(new SectionBlock($"```{errorMessage}```"));
+			attachment.Blocks.Add(QuoteBlock(errorMessage));
 			if (change != null)
 			{
 				if (author != null)
@@ -1359,7 +1355,7 @@
 				}
 				if (description != null)
 				{
-					attachment.Blocks.Add(new SectionBlock($"```{description}```"));
+					attachment.Blocks.Add(QuoteBlock(description));
 				}
 			}
 
@@ -1396,7 +1392,10 @@
 			attachment.Blocks.Add(new HeaderBlock($"Stream Update Failure :rip:", true));
 
 			attachment.Blocks.Add(new SectionBlock($"<!here> Horde was unable to update {file.DepotPath}"));
-			attachment.Blocks.Add(new SectionBlock($"```{file.Error}```"));
+			if (file.Error != null)
+			{
+				attachment.Blocks.Add(QuoteBlock(file.Error));
+			}
 
 			return SendMessageAsync(recipient, attachments: new[] { attachment });
 		}
@@ -1489,6 +1488,26 @@
 
 		#endregion
 
+		static SectionBlock QuoteBlock(string text)
+		{
+			const int MaxLength = TextObject.MaxLength - 6;
+
+			if (text.Length > MaxLength)
+			{
+				int length = text.LastIndexOf('\n', MaxLength);
+				if (length == 0)
+				{
+					text = text.Substring(0, MaxLength - 7) + "...\n...";
+				}
+				else
+				{
+					text = text.Substring(0, length + 1) + "...";
+				}
+			}
+
+			return new SectionBlock($"```{text}```");
+		}
+
 		const int MaxJobStepEvents = 5;
 
 		static string GetJobChangeText(IJob job)
diff --git a/Engine/Source/Programs/Horde/Horde.Build/Utilities/BlockKit/TextObject.cs b/Engine/Source/Programs/Horde/Horde.Build/Utilities/BlockKit/TextObject.cs
--- a/Engine/Source/Programs/Horde/Horde.Build/Utilities/BlockKit/TextObject.cs
+++ b/Engine/Source/Programs/Horde/Horde.Build/Utilities/BlockKit/TextObject.cs
@@ -1,5 +1,6 @@
 // Copyright Epic Games, Inc. All Rights Reserved.
 
+using System;
 using System.Text.Json;
 
 namespace Horde.Build.Utilities.Slack.BlockKit
@@ -9,6 +10,11 @@ namespace Horde.Build.Utilities.Slack.BlockKit
 	/// </summary>
 	public class TextObject
 	{
+		/// <summary>
+		/// Maximum length of a text object before truncation
+		/// </summary>
+		public const int MaxLength = 2048;
+
 		/// <summary>
 		/// Constructor
 		/// </summary>
@@ -47,8 +53,6 @@ namespace Horde.Build.Utilities.Slack.BlockKit
 			writer.WriteStartObject("text");
 			writer.WriteString("type", IsMarkdown ? "mrkdwn" : "plain_text");
 
-			const int MaxLength = 2048;
-
 			string truncatedText = Text;
 			if (truncatedText.Length > MaxLength)
 			{
