Horde: Annotate multi-line messages correctly.

diff --git a/Engine/Source/Programs/Horde/Horde.Agent/Utility/JsonRpcLogger.cs b/Engine/Source/Programs/Horde/Horde.Agent/Utility/JsonRpcLogger.cs
--- a/Engine/Source/Programs/Horde/Horde.Agent/Utility/JsonRpcLogger.cs
+++ b/Engine/Source/Programs/Horde/Horde.Agent/Utility/JsonRpcLogger.cs
@@ -285,10 +285,10 @@ namespace Horde.Agent.Parser
 							AddEvent(jsonLogEvent.Data.Span, lineIndex, EventSeverity.Error, events);
 						}
 
-						writer.Write(jsonLogEvent.Data.Span);
-						writer.Write(newline);
+						int lineCount = jsonLogEvent.GetMessageLineCount();
+						WriteEvent(jsonLogEvent.Data.Span, lineCount, writer);
 
-						lineIndex++;
+						lineIndex += lineCount;
 					}
 					else
 					{
@@ -362,6 +362,42 @@ namespace Horde.Agent.Parser
 			}
 		}
 
+		static readonly Utf8String s_newline = "\n";
+		static readonly Utf8String s_comma = ",";
+
+		public static void WriteEvent(ReadOnlySpan<byte> span, int lineCount, ArrayBufferWriter<byte> writer)
+		{
+			int insertIdx = span.Length;
+			if(lineCount > 1)
+			{
+				int endIdx = span.Length - 1;
+				while (endIdx > 0 && span[endIdx] == ' ')
+				{
+					endIdx--;
+				}
+				if(span[endIdx] == '}')
+				{
+					insertIdx = endIdx;
+				}
+			}
+
+			for (int idx = 0; idx < lineCount; idx++)
+			{
+				writer.Write(span.Slice(0, insertIdx));
+				if (insertIdx < span.Length)
+				{
+					writer.Write(s_comma);
+					using (Utf8JsonWriter jsonWriter = new Utf8JsonWriter(writer, new JsonWriterOptions { SkipValidation = true }))
+					{
+						jsonWriter.WriteNumber(LogEventPropertyName.Line, idx);
+						jsonWriter.WriteNumber(LogEventPropertyName.LineCount, lineCount);
+					}
+					writer.Write(span.Slice(insertIdx));
+				}
+				writer.Write(s_newline);
+			}
+		}
+
 		void AddEvent(ReadOnlySpan<byte> span, int lineIndex, EventSeverity severity, List<CreateEventRequest> events)
 		{
 			try
diff --git a/Engine/Source/Programs/Shared/EpicGames.Core/JsonLogEvent.cs b/Engine/Source/Programs/Shared/EpicGames.Core/JsonLogEvent.cs
--- a/Engine/Source/Programs/Shared/EpicGames.Core/JsonLogEvent.cs
+++ b/Engine/Source/Programs/Shared/EpicGames.Core/JsonLogEvent.cs
@@ -135,6 +135,58 @@ namespace EpicGames.Core
 			return (LogLevel)result;
 		}
 
+		static readonly Utf8String s_newlineEscaped = "\\n";
+
+		/// <summary>
+		/// Count the number of lines in the message field of a log event
+		/// </summary>
+		/// <returns>Number of lines in the message</returns>
+		public int GetMessageLineCount()
+		{
+			if (Data.Span.IndexOf(s_newlineEscaped.Span) != -1)
+			{
+				Utf8JsonReader reader = new Utf8JsonReader(Data.Span);
+				if (reader.Read() && reader.TokenType == JsonTokenType.StartObject)
+				{
+					while (reader.Read() && reader.TokenType == JsonTokenType.PropertyName)
+					{
+						ReadOnlySpan<byte> propertyName = reader.ValueSpan;
+						if (!reader.Read())
+						{
+							break;
+						}
+						else if (propertyName.SequenceEqual(LogEventPropertyName.Message) && reader.TokenType == JsonTokenType.String)
+						{
+							return CountLines(reader.ValueSpan);
+						}
+					}
+				}
+			}
+			return 1;
+		}
+		
+		/// <summary>
+		/// Counts the number of newlines in an escaped JSON string
+		/// </summary>
+		/// <param name="str">The escaped string</param>
+		/// <returns></returns>
+		static int CountLines(ReadOnlySpan<byte> str)
+		{
+			int lines = 1;
+			for (int idx = 0; idx < str.Length - 1; idx++)
+			{
+				if (str[idx] == '\\')
+				{
+					if (str[idx + 1] == 'n')
+					{
+						lines++;
+					}
+					idx++;
+				}
+			}
+			return lines;
+		}
+
 		/// <summary>
 		/// Formats an event as a string
 		/// </summary>
