Horde: Prevent exception when trying to assign lease escaping the scope where we have another lease ready to assign.

diff --git a/Engine/Source/Programs/Horde/Horde.Build/Agents/Impl/AgentService.cs b/Engine/Source/Programs/Horde/Horde.Build/Agents/Impl/AgentService.cs
--- a/Engine/Source/Programs/Horde/Horde.Build/Agents/Impl/AgentService.cs
+++ b/Engine/Source/Programs/Horde/Horde.Build/Agents/Impl/AgentService.cs
@@ -480,7 +480,7 @@ namespace Horde.Build.Services
 			catch (Exception ex)
 			{
 				_logger.LogError(ex, "Exception while trying to assign lease");
-				return Task.FromException<(ITaskSource, AgentLease)?>(ex);
+				return Task.FromResult<(ITaskSource, AgentLease)?>(null);
 			}
 		}
 
