Horde: Prevent updating Slack messages which haven't changed.

diff --git a/Engine/Source/Programs/Horde/Horde.Build/Notifications/Impl/SlackNotificationSink.cs b/Engine/Source/Programs/Horde/Horde.Build/Notifications/Impl/SlackNotificationSink.cs
--- a/Engine/Source/Programs/Horde/Horde.Build/Notifications/Impl/SlackNotificationSink.cs
+++ b/Engine/Source/Programs/Horde/Horde.Build/Notifications/Impl/SlackNotificationSink.cs
@@ -1706,6 +1706,12 @@ namespace Horde.Build.Notifications.Impl
 
 			string requestDigest = ContentHash.MD5(JsonSerializer.Serialize(message, new JsonSerializerOptions { DefaultIgnoreCondition = JsonIgnoreCondition.WhenWritingNull })).ToString();
 
+			MessageStateDocument? prevState = await GetMessageStateAsync(recipient, eventId);
+			if (prevState != null && prevState.Digest == requestDigest)
+			{
+				return (prevState, false);
+			}
+
 			(MessageStateDocument state, bool isNew) = await AddOrUpdateMessageStateAsync(recipient, eventId, userId, requestDigest);
 			if (isNew)
 			{
