Horde: Add additional context for log parsing errors.

diff --git a/Engine/Source/Programs/Horde/Horde.Build/Logs/Builder/Impl/RedisLogBuilder.cs b/Engine/Source/Programs/Horde/Horde.Build/Logs/Builder/Impl/RedisLogBuilder.cs
--- a/Engine/Source/Programs/Horde/Horde.Build/Logs/Builder/Impl/RedisLogBuilder.cs
+++ b/Engine/Source/Programs/Horde/Horde.Build/Logs/Builder/Impl/RedisLogBuilder.cs
@@ -153,7 +153,16 @@ namespace Horde.Build.Logs.Builder
 
 				ReadOnlyLogText subChunkText = new ReadOnlyLogText(subChunkTextValue);
 				LogSubChunkData subChunkData = new LogSubChunkData(type, offset + length - subChunkText.Data.Length, lineIndex - subChunkText.LineCount, subChunkText);
-				byte[] subChunkDataBytes = subChunkData.ToByteArray();
+
+				byte[] subChunkDataBytes;
+				try
+				{
+					subChunkDataBytes = subChunkData.ToByteArray();
+				}
+				catch (Exception ex)
+				{
+					throw new LogTextParseException($"Error while parsing sub-chunk for {logId} offset {offset}", ex);
+				}
 
 				ITransaction writeTransaction = redisDb.CreateTransaction();
 
