Horde.Storage: Add CbField serializer to ExistCheckMultipleContentIdResponse.Needs

#rnx

diff --git a/Engine/Source/Programs/Horde/HordeStorage/Horde.Storage/Controllers/CompressedBlobController.cs b/Engine/Source/Programs/Horde/HordeStorage/Horde.Storage/Controllers/CompressedBlobController.cs
--- a/Engine/Source/Programs/Horde/HordeStorage/Horde.Storage/Controllers/CompressedBlobController.cs
+++ b/Engine/Source/Programs/Horde/HordeStorage/Horde.Storage/Controllers/CompressedBlobController.cs
@@ -10,6 +10,7 @@ using System.Net.Mime;
 using System.Threading.Tasks;
 using Datadog.Trace;
 using EpicGames.Horde.Storage;
+using EpicGames.Serialization;
 using Horde.Storage.Implementation;
 using Jupiter;
 using Jupiter.Common.Implementation;
@@ -410,6 +411,7 @@ namespace Horde.Storage.Controllers
 
     public class ExistCheckMultipleContentIdResponse
     {
-        public ContentId[] Needs { get; set; } = null!;
+		[CbField("needs")]
+		public ContentId[] Needs { get; set; } = null!;
     }
 }
