Horde: Fix another case of text not being truncated to fit within message limits.

diff --git a/Engine/Source/Programs/Horde/Horde.Build/Notifications/Impl/SlackNotificationSink.cs b/Engine/Source/Programs/Horde/Horde.Build/Notifications/Impl/SlackNotificationSink.cs
--- a/Engine/Source/Programs/Horde/Horde.Build/Notifications/Impl/SlackNotificationSink.cs
+++ b/Engine/Source/Programs/Horde/Horde.Build/Notifications/Impl/SlackNotificationSink.cs
@@ -778,7 +778,7 @@ namespace Horde.Build.Notifications.Impl
 					summary.AppendLine($"From {FormatJobStep(span.FirstFailure, span.NodeName)}:");
 					foreach (ILogEventData eventDataItem in eventDataItems)
 					{
-						summary.AppendLine($"```{eventDataItem.Message}```");
+						summary.AppendLine(QuoteText(eventDataItem.Message, TextObject.MaxLength));
 					}
 					if (events.Count > eventDataItems.Count)
 					{
@@ -1490,22 +1490,25 @@ namespace Horde.Build.Notifications.Impl
 
 		static SectionBlock QuoteBlock(string text)
 		{
-			const int MaxLength = TextObject.MaxLength - 6;
+			return new SectionBlock(QuoteText(text, TextObject.MaxLength));
+		}
 
-			if (text.Length > MaxLength)
+		static string QuoteText(string text, int maxLength)
+		{
+			maxLength -= 6;
+			if (text.Length > maxLength)
 			{
-				int length = text.LastIndexOf('\n', MaxLength);
+				int length = text.LastIndexOf('\n', maxLength);
 				if (length == 0)
 				{
-					text = text.Substring(0, MaxLength - 7) + "...\n...";
+					text = text.Substring(0, maxLength - 7) + "...\n...";
 				}
 				else
 				{
 					text = text.Substring(0, length + 1) + "...";
 				}
 			}
-
-			return new SectionBlock($"```{text}```");
+			return $"```{text}```";
 		}
 
 		const int MaxJobStepEvents = 5;
