Horde: Suppress errors about malformed log data; output warnings instead. This data was produced by bad code that no longer exists, but throwing an exception prevents indexing from completing.

diff --git a/Engine/Source/Programs/Horde/Horde.Build/Logs/Builder/Impl/RedisLogBuilder.cs b/Engine/Source/Programs/Horde/Horde.Build/Logs/Builder/Impl/RedisLogBuilder.cs
--- a/Engine/Source/Programs/Horde/Horde.Build/Logs/Builder/Impl/RedisLogBuilder.cs
+++ b/Engine/Source/Programs/Horde/Horde.Build/Logs/Builder/Impl/RedisLogBuilder.cs
@@ -157,7 +157,7 @@ namespace Horde.Build.Logs.Builder
 				byte[] subChunkDataBytes;
 				try
 				{
-					subChunkDataBytes = subChunkData.ToByteArray();
+					subChunkDataBytes = subChunkData.ToByteArray(_logger);
 				}
 				catch (Exception ex)
 				{
diff --git a/Engine/Source/Programs/Horde/Horde.Build/Logs/Data/LogChunkData.cs b/Engine/Source/Programs/Horde/Horde.Build/Logs/Data/LogChunkData.cs
--- a/Engine/Source/Programs/Horde/Horde.Build/Logs/Data/LogChunkData.cs
+++ b/Engine/Source/Programs/Horde/Horde.Build/Logs/Data/LogChunkData.cs
@@ -5,6 +5,7 @@ using System.Collections.Generic;
 using System.Linq;
 using EpicGames.Core;
 using Horde.Build.Api;
+using Microsoft.Extensions.Logging;
 
 namespace Horde.Build.Logs
 {
@@ -214,26 +215,27 @@ namespace Horde.Build.Logs
 		/// Write the chunk data to a stream
 		/// </summary>
 		/// <returns>Serialized data</returns>
-		public void Write(MemoryWriter writer)
+		public void Write(MemoryWriter writer, ILogger logger)
 		{
 			writer.WriteInt32(CurrentSignature | (1 << 24));
 
 			writer.WriteInt32(SubChunks.Count);
 			foreach (LogSubChunkData subChunk in SubChunks)
 			{
-				writer.WriteLogSubChunkData(subChunk);
+				writer.WriteLogSubChunkData(subChunk, logger);
 			}
 		}
 
 		/// <summary>
 		/// Construct an object from flat memory buffer
 		/// </summary>
+		/// <param name="logger">Logger for output</param>
 		/// <returns>Log chunk data</returns>
-		public byte[] ToByteArray()
+		public byte[] ToByteArray(ILogger logger)
 		{
-			byte[] data = new byte[GetSerializedSize()];
+			byte[] data = new byte[GetSerializedSize(logger)];
 			MemoryWriter writer = new MemoryWriter(data);
-			Write(writer);
+			Write(writer, logger);
 			writer.CheckOffset(data.Length);
 			return data;
 		}
@@ -241,9 +243,9 @@ namespace Horde.Build.Logs
 		/// <summary>
 		/// Determines the size of the serialized buffer
 		/// </summary>
-		public int GetSerializedSize()
+		public int GetSerializedSize(ILogger logger)
 		{
-			return sizeof(int) + (sizeof(int) + SubChunks.Sum(x => x.GetSerializedSize()));
+			return sizeof(int) + (sizeof(int) + SubChunks.Sum(x => x.GetSerializedSize(logger)));
 		}
 	}
 
@@ -268,9 +270,9 @@ namespace Horde.Build.Logs
 		/// Write the chunk data to a stream
 		/// </summary>
 		/// <returns>Serialized data</returns>
-		public static void WriteLogChunkData(this MemoryWriter writer, LogChunkData chunkData)
+		public static void WriteLogChunkData(this MemoryWriter writer, LogChunkData chunkData, ILogger logger)
 		{
-			chunkData.Write(writer);
+			chunkData.Write(writer, logger);
 		}
 	}
 }
diff --git a/Engine/Source/Programs/Horde/Horde.Build/Logs/Data/LogSubChunkData.cs b/Engine/Source/Programs/Horde/Horde.Build/Logs/Data/LogSubChunkData.cs
--- a/Engine/Source/Programs/Horde/Horde.Build/Logs/Data/LogSubChunkData.cs
+++ b/Engine/Source/Programs/Horde/Horde.Build/Logs/Data/LogSubChunkData.cs
@@ -4,6 +4,7 @@ using System;
 using EpicGames.Core;
 using Horde.Build.Api;
 using Horde.Build.Utilities;
+using Microsoft.Extensions.Logging;
 
 namespace Horde.Build.Logs
 {
@@ -79,14 +80,14 @@ namespace Horde.Build.Logs
 		/// <summary>
 		/// Index for tokens in this chunk
 		/// </summary>
-		public LogIndexData BuildIndex()
+		public LogIndexData BuildIndex(ILogger logger)
 		{
 			if(_indexInternal == null)
 			{
 				ILogText plainText = InflateText();
 				if (Type != LogType.Text)
 				{
-					plainText = plainText.ToPlainText();
+					plainText = plainText.ToPlainText(logger);
 				}
 
 				LogIndexBlock[] blocks = new LogIndexBlock[1];
@@ -181,7 +182,8 @@ namespace Horde.Build.Logs
 		/// Serializes the sub-chunk to a stream
 		/// </summary>
 		/// <param name="writer">Writer to output to</param>
-		public void Write(MemoryWriter writer)
+		/// <param name="logger">Logger for output</param>
+		public void Write(MemoryWriter writer, ILogger logger)
 		{
 			writer.WriteInt32(2); // Version placeholder
 
@@ -190,19 +192,19 @@ namespace Horde.Build.Logs
 			writer.WriteInt32(LineCount);
 
 			writer.WriteVariableLengthBytes(DeflateText().Span);
-			writer.WriteLogIndexData(BuildIndex());
+			writer.WriteLogIndexData(BuildIndex(logger));
 		}
 
 		/// <summary>
 		/// Serializes this object to a byte array
 		/// </summary>
 		/// <returns>Byte array</returns>
-		public byte[] ToByteArray()
+		public byte[] ToByteArray(ILogger logger)
 		{
-			byte[] data = new byte[GetSerializedSize()];
+			byte[] data = new byte[GetSerializedSize(logger)];
 
 			MemoryWriter writer = new MemoryWriter(data);
-			Write(writer);
+			Write(writer, logger);
 			writer.CheckOffset(data.Length);
 
 			return data;
@@ -211,9 +213,9 @@ namespace Horde.Build.Logs
 		/// <summary>
 		/// Determines the size of serialized data
 		/// </summary>
-		public int GetSerializedSize()
+		public int GetSerializedSize(ILogger logger)
 		{
-			return sizeof(int) + sizeof(int) + sizeof(int) + sizeof(int) + (sizeof(int) + DeflateText().Length) + BuildIndex().GetSerializedSize();
+			return sizeof(int) + sizeof(int) + sizeof(int) + sizeof(int) + (sizeof(int) + DeflateText().Length) + BuildIndex(logger).GetSerializedSize();
 		}
 	}
 
@@ -239,9 +241,10 @@ namespace Horde.Build.Logs
 		/// </summary>
 		/// <param name="writer">Writer to output to</param>
 		/// <param name="subChunkData">The sub-chunk data to write</param>
-		public static void WriteLogSubChunkData(this MemoryWriter writer, LogSubChunkData subChunkData)
+		/// <param name="logger">Logger for output</param>
+		public static void WriteLogSubChunkData(this MemoryWriter writer, LogSubChunkData subChunkData, ILogger logger)
 		{
-			subChunkData.Write(writer);
+			subChunkData.Write(writer, logger);
 		}
 	}
 }
diff --git a/Engine/Source/Programs/Horde/Horde.Build/Logs/Data/LogText.cs b/Engine/Source/Programs/Horde/Horde.Build/Logs/Data/LogText.cs
--- a/Engine/Source/Programs/Horde/Horde.Build/Logs/Data/LogText.cs
+++ b/Engine/Source/Programs/Horde/Horde.Build/Logs/Data/LogText.cs
@@ -7,6 +7,7 @@
 using System.Text;
 using System.Text.Json;
 using EpicGames.Core;
+using Microsoft.Extensions.Logging;
 
 namespace Horde.Build.Logs
 {
@@ -233,7 +234,7 @@
 		/// <summary>
 		/// Updates the plain text representation of this chunk
 		/// </summary>
-		public void AppendPlainText(ILogText srcText, int srcLineIndex, int srcLineCount)
+		public void AppendPlainText(ILogText srcText, int srcLineIndex, int srcLineCount, ILogger logger)
 		{
 			// Convert all the lines to plain text
 			for (int idx = 0; idx < srcLineCount; idx++)
@@ -257,7 +258,8 @@
 				}
 				catch (Exception ex)
 				{
-					throw new LogTextParseException($"Exception while parsing log text as JSON. Line: {Encoding.UTF8.GetString(inputLine)}", ex);
+					inputLine.CopyTo(_internalData.AsSpan(Length));
+					logger.LogWarning(ex, "Exception while attempting to parse log text as JSON. Line: \"{Line}\"", Encoding.UTF8.GetString(inputLine).Trim());
 				}
 				_internalLineOffsets.Add(Length);
 			}
@@ -488,11 +490,12 @@
 		/// Converts a log text instance to plain text
 		/// </summary>
 		/// <param name="logText">The text to convert</param>
+		/// <param name="logger">Logger for conversion warnings</param>
 		/// <returns>The plain text instance</returns>
-		public static ILogText ToPlainText(this ILogText logText)
+		public static ILogText ToPlainText(this ILogText logText, ILogger logger)
 		{
 			LogText other = new LogText();
-			other.AppendPlainText(logText, 0, logText.LineCount);
+			other.AppendPlainText(logText, 0, logText.LineCount, logger);
 			return other;
 		}
 	}
diff --git a/Engine/Source/Programs/Horde/Horde.Build/Logs/Impl/LogFileService.cs b/Engine/Source/Programs/Horde/Horde.Build/Logs/Impl/LogFileService.cs
--- a/Engine/Source/Programs/Horde/Horde.Build/Logs/Impl/LogFileService.cs
+++ b/Engine/Source/Programs/Horde/Horde.Build/Logs/Impl/LogFileService.cs
@@ -1265,7 +1265,7 @@ namespace Horde.Build.Services
 					{
 						try
 						{
-							indexes.Add(subChunkData.BuildIndex());
+							indexes.Add(subChunkData.BuildIndex(_logger));
 						}
 						catch (Exception ex)
 						{
@@ -1559,7 +1559,7 @@ namespace Horde.Build.Services
 						if (firstLine < subChunkData.LineIndex + subChunkData.LineCount)
 						{
 							// Create an index containing just this sub-chunk
-							LogIndexData index = subChunkData.BuildIndex();
+							LogIndexData index = subChunkData.BuildIndex(_logger);
 							foreach (int lineIndex in index.Search(firstLine, searchText, searchStats))
 							{
 								yield return lineIndex;
diff --git a/Engine/Source/Programs/Horde/Horde.Build/Logs/Storage/Impl/PersistentLogStorage.cs b/Engine/Source/Programs/Horde/Horde.Build/Logs/Storage/Impl/PersistentLogStorage.cs
--- a/Engine/Source/Programs/Horde/Horde.Build/Logs/Storage/Impl/PersistentLogStorage.cs
+++ b/Engine/Source/Programs/Horde/Horde.Build/Logs/Storage/Impl/PersistentLogStorage.cs
@@ -96,9 +96,9 @@ namespace Horde.Build.Logs.Readers
 			_logger.LogDebug("Writing log {LogId} chunk offset {Offset} to persistent storage", logId, offset);
 
 			string path = $"{logId}/offset_{offset}";
-			byte[] data = new byte[chunkData.GetSerializedSize()];
+			byte[] data = new byte[chunkData.GetSerializedSize(_logger)];
 			MemoryWriter writer = new MemoryWriter(data);
-			writer.WriteLogChunkData(chunkData);
+			writer.WriteLogChunkData(chunkData, _logger);
 			writer.CheckOffset(data.Length);
 
 			return _storageProvider.WriteBytesAsync(path, data);
diff --git a/Engine/Source/Programs/Horde/Horde.Build/Logs/Storage/Impl/RedisLogStorage.cs b/Engine/Source/Programs/Horde/Horde.Build/Logs/Storage/Impl/RedisLogStorage.cs
--- a/Engine/Source/Programs/Horde/Horde.Build/Logs/Storage/Impl/RedisLogStorage.cs
+++ b/Engine/Source/Programs/Horde/Horde.Build/Logs/Storage/Impl/RedisLogStorage.cs
@@ -154,7 +154,7 @@ namespace Horde.Build.Logs.Storage.Impl
 				chunkData = await _inner.ReadChunkAsync(logId, offset, lineIndex);
 				if (chunkData != null)
 				{
-					await AddAsync(key, chunkData.ToByteArray());
+					await AddAsync(key, chunkData.ToByteArray(_logger));
 				}
 			}
 			return chunkData;
