Horde: Fix test compilation.

diff --git a/Engine/Source/Programs/Horde/Horde.Build.Tests/LogChunkTests.cs b/Engine/Source/Programs/Horde/Horde.Build.Tests/LogChunkTests.cs
--- a/Engine/Source/Programs/Horde/Horde.Build.Tests/LogChunkTests.cs
+++ b/Engine/Source/Programs/Horde/Horde.Build.Tests/LogChunkTests.cs
@@ -5,6 +5,7 @@ using System.Collections.Generic;
 using System.Text;
 using Horde.Build.Api;
 using Horde.Build.Logs;
+using Microsoft.Extensions.Logging.Abstractions;
 using Microsoft.VisualStudio.TestTools.UnitTesting;
 
 namespace Horde.Build.Tests
@@ -29,7 +30,7 @@ namespace Horde.Build.Tests
 			subChunks.Add(new LogSubChunkData(LogType.Json, textData.Length, 3, new LogText(textData, textData.Length)));
 
 			LogChunkData oldChunkData = new LogChunkData(0, 0, subChunks);
-			byte[] data = oldChunkData.ToByteArray();
+			byte[] data = oldChunkData.ToByteArray(NullLogger.Instance);
 			LogChunkData newChunkData = LogChunkData.FromMemory(data, 0, 0);
 
 			Assert.AreEqual(oldChunkData.Length, newChunkData.Length);
