Horde: Support for splitting log events with format strings and values that span multiple lines.

diff --git a/Engine/Source/Programs/Horde/Horde.Agent.Tests/LogParserTests.cs b/Engine/Source/Programs/Horde/Horde.Agent.Tests/LogParserTests.cs
--- a/Engine/Source/Programs/Horde/Horde.Agent.Tests/LogParserTests.cs
+++ b/Engine/Source/Programs/Horde/Horde.Agent.Tests/LogParserTests.cs
@@ -1,6 +1,7 @@
 // Copyright Epic Games, Inc. All Rights Reserved.
 
 using System;
+using System.Buffers;
 using System.Collections.Generic;
 using System.IO;
 using System.Linq;
@@ -77,6 +78,30 @@ namespace Horde.Agent.Tests
 			Assert.AreEqual(impl._lines[1].GetProperty("Text").ToString(), "world");
 		}
 
+		static string[] SplitMultiLineMessage(string message)
+		{
+			ArrayBufferWriter<byte> writer = new ArrayBufferWriter<byte>();
+			int count = JsonRpcLogger.WriteEvent(Encoding.UTF8.GetBytes(message), writer);
+			string[] result = Encoding.UTF8.GetString(writer.WrittenSpan).Split('\n', StringSplitOptions.RemoveEmptyEntries);
+			Assert.AreEqual(count, result.Length);
+			return result;
+		}
+
+		[TestMethod]
+		public void MultiLineTest()
+		{
+			const string msg = @"{""message"":""ignored"",""format"":""This\nis a\nmulti-{Line}-end\nlog message {Var1} {Var2}"",""properties"":{""Line"":""line\nsplit\nin\nfour"",""Var1"":123,""Var2"":{""$type"":""SourceFile"",""$text"":""D:\\build\\\u002B\u002BUE5\\Sync\\GenerateProjectFiles.bat"",""relativePath"":""GenerateProjectFiles.bat"",""depotPath"":""//UE5/Main/GenerateProjectFiles.bat@20392842""}}}";
+
+			string[] result = SplitMultiLineMessage(msg);
+			Assert.AreEqual(@"{""message"":""This"",""format"":""This"",""properties"":{""Line"":""line\nsplit\nin\nfour"",""Var1"":123,""Var2"":{""$type"":""SourceFile"",""$text"":""D:\\build\\\u002B\u002BUE5\\Sync\\GenerateProjectFiles.bat"",""relativePath"":""GenerateProjectFiles.bat"",""depotPath"":""//UE5/Main/GenerateProjectFiles.bat@20392842""},""Line$0"":""line"",""Line$1"":""split"",""Line$2"":""in"",""Line$3"":""four""},""line"":0,""lineCount"":7}", result[0]);
+			Assert.AreEqual(@"{""message"":""is a"",""format"":""is a"",""properties"":{""Line"":""line\nsplit\nin\nfour"",""Var1"":123,""Var2"":{""$type"":""SourceFile"",""$text"":""D:\\build\\\u002B\u002BUE5\\Sync\\GenerateProjectFiles.bat"",""relativePath"":""GenerateProjectFiles.bat"",""depotPath"":""//UE5/Main/GenerateProjectFiles.bat@20392842""},""Line$0"":""line"",""Line$1"":""split"",""Line$2"":""in"",""Line$3"":""four""},""line"":1,""lineCount"":7}", result[1]);
+			Assert.AreEqual(@"{""message"":""multi-line"",""format"":""multi-{Line$0}"",""properties"":{""Line"":""line\nsplit\nin\nfour"",""Var1"":123,""Var2"":{""$type"":""SourceFile"",""$text"":""D:\\build\\\u002B\u002BUE5\\Sync\\GenerateProjectFiles.bat"",""relativePath"":""GenerateProjectFiles.bat"",""depotPath"":""//UE5/Main/GenerateProjectFiles.bat@20392842""},""Line$0"":""line"",""Line$1"":""split"",""Line$2"":""in"",""Line$3"":""four""},""line"":2,""lineCount"":7}", result[2]);
+			Assert.AreEqual(@"{""message"":""split"",""format"":""{Line$1}"",""properties"":{""Line"":""line\nsplit\nin\nfour"",""Var1"":123,""Var2"":{""$type"":""SourceFile"",""$text"":""D:\\build\\\u002B\u002BUE5\\Sync\\GenerateProjectFiles.bat"",""relativePath"":""GenerateProjectFiles.bat"",""depotPath"":""//UE5/Main/GenerateProjectFiles.bat@20392842""},""Line$0"":""line"",""Line$1"":""split"",""Line$2"":""in"",""Line$3"":""four""},""line"":3,""lineCount"":7}", result[3]);
+			Assert.AreEqual(@"{""message"":""in"",""format"":""{Line$2}"",""properties"":{""Line"":""line\nsplit\nin\nfour"",""Var1"":123,""Var2"":{""$type"":""SourceFile"",""$text"":""D:\\build\\\u002B\u002BUE5\\Sync\\GenerateProjectFiles.bat"",""relativePath"":""GenerateProjectFiles.bat"",""depotPath"":""//UE5/Main/GenerateProjectFiles.bat@20392842""},""Line$0"":""line"",""Line$1"":""split"",""Line$2"":""in"",""Line$3"":""four""},""line"":4,""lineCount"":7}", result[4]);
+			Assert.AreEqual(@"{""message"":""four-end"",""format"":""{Line$3}-end"",""properties"":{""Line"":""line\nsplit\nin\nfour"",""Var1"":123,""Var2"":{""$type"":""SourceFile"",""$text"":""D:\\build\\\u002B\u002BUE5\\Sync\\GenerateProjectFiles.bat"",""relativePath"":""GenerateProjectFiles.bat"",""depotPath"":""//UE5/Main/GenerateProjectFiles.bat@20392842""},""Line$0"":""line"",""Line$1"":""split"",""Line$2"":""in"",""Line$3"":""four""},""line"":5,""lineCount"":7}", result[5]);
+			Assert.AreEqual(@"{""message"":""log message 123 D:\\build\\\u002B\u002BUE5\\Sync\\GenerateProjectFiles.bat"",""format"":""log message {Var1} {Var2}"",""properties"":{""Line"":""line\nsplit\nin\nfour"",""Var1"":123,""Var2"":{""$type"":""SourceFile"",""$text"":""D:\\build\\\u002B\u002BUE5\\Sync\\GenerateProjectFiles.bat"",""relativePath"":""GenerateProjectFiles.bat"",""depotPath"":""//UE5/Main/GenerateProjectFiles.bat@20392842""},""Line$0"":""line"",""Line$1"":""split"",""Line$2"":""in"",""Line$3"":""four""},""line"":6,""lineCount"":7}", result[6]);
+		}
+
 		[TestMethod]
 		public void StructuredOutputMatcher()
 		{
diff --git a/Engine/Source/Programs/Horde/Horde.Agent/Utility/JsonRpcLogger.cs b/Engine/Source/Programs/Horde/Horde.Agent/Utility/JsonRpcLogger.cs
--- a/Engine/Source/Programs/Horde/Horde.Agent/Utility/JsonRpcLogger.cs
+++ b/Engine/Source/Programs/Horde/Horde.Agent/Utility/JsonRpcLogger.cs
@@ -6,6 +6,7 @@ using System.Collections.Generic;
 using System.Linq;
 using System.Text;
 using System.Text.Json;
+using System.Text.Json.Nodes;
 using System.Threading;
 using System.Threading.Channels;
 using System.Threading.Tasks;
@@ -17,6 +18,8 @@ using Microsoft.Extensions.Logging;
 
 namespace Horde.Agent.Parser
 {
+	using JsonObject = System.Text.Json.Nodes.JsonObject;
+
 	interface IJsonRpcLogSink
 	{
 		Task WriteEventsAsync(List<CreateEventRequest> events);
@@ -285,9 +288,7 @@ namespace Horde.Agent.Parser
 							AddEvent(jsonLogEvent.Data.Span, lineIndex, EventSeverity.Error, events);
 						}
 
-						int lineCount = jsonLogEvent.GetMessageLineCount();
-						WriteEvent(jsonLogEvent.Data.Span, lineCount, writer);
-
+						int lineCount = WriteEvent(jsonLogEvent.Data.Span, writer);
 						lineIndex += lineCount;
 					}
 					else
@@ -362,40 +363,120 @@ namespace Horde.Agent.Parser
 			}
 		}
 
+		static readonly string s_messagePropertyName = LogEventPropertyName.Message.ToString();
+		static readonly string s_formatPropertyName = LogEventPropertyName.Format.ToString();
+		static readonly string s_linePropertyName = LogEventPropertyName.Line.ToString();
+		static readonly string s_lineCountPropertyName = LogEventPropertyName.LineCount.ToString();
+
 		static readonly Utf8String s_newline = "\n";
-		static readonly Utf8String s_comma = ",";
+		static readonly Utf8String s_escapedNewline = "\\n";
 
-		public static void WriteEvent(ReadOnlySpan<byte> span, int lineCount, ArrayBufferWriter<byte> writer)
+		public static int WriteEvent(ReadOnlySpan<byte> logEvent, IBufferWriter<byte> writer)
 		{
-			int insertIdx = span.Length;
-			if(lineCount > 1)
+			if (logEvent.IndexOf(s_escapedNewline) == -1)
+			{
+				writer.Write(logEvent);
+				writer.Write(s_newline);
+				return 1;
+			}
+
+			JsonObject obj = (JsonObject)JsonNode.Parse(logEvent)!;
+
+			string format = (obj["format"] as JsonValue)?.ToString() ?? String.Empty;
+			IEnumerable<KeyValuePair<string, object?>> propertyValueList = Enumerable.Empty<KeyValuePair<string, object?>>();
+
+			// Split all the multi-line properties into separate properties
+			JsonObject? properties = obj["properties"] as JsonObject;
+			if (properties != null)
 			{
-				int endIdx = span.Length - 1;
-				while (endIdx > 0 && span[endIdx] == ' ')
+				// Get all the current property values
+				Dictionary<string, string> propertyValues = new Dictionary<string, string>(StringComparer.Ordinal);
+				foreach ((string name, JsonNode? node) in properties)
 				{
-					endIdx--;
+					string value = String.Empty;
+					if (node != null)
+					{
+						if (node is JsonObject valueObject)
+						{
+							value = valueObject["$text"]?.ToString() ?? String.Empty;
+						}
+						else
+						{
+							value = node.ToString();
+						}
+					}
+					propertyValues[name] = value;
 				}
-				if(span[endIdx] == '}')
+
+				// Split all the multi-line properties into separate things
+				int nameStart = -1;
+				for (int idx = 0; idx < format.Length; idx++)
 				{
-					insertIdx = endIdx;
+					if (format[idx] == '{')
+					{
+						nameStart = idx + 1;
+					}
+					else if (format[idx] == '}' && nameStart != -1)
+					{
+						string name = format.Substring(nameStart, idx - nameStart);
+						if (propertyValues.TryGetValue(name, out string? text))
+						{
+							int textLineEnd = text.IndexOf('\n', StringComparison.Ordinal);
+							if (textLineEnd != -1)
+							{
+								int lineNum = 0;
+
+								StringBuilder builder = new StringBuilder();
+								builder.Append(format, 0, nameStart - 1);
+
+								string delimiter = String.Empty;
+								for (int textLineStart = 0; textLineStart < text.Length;)
+								{
+									string newName = $"{name}${lineNum++}";
+									string newLine = text.Substring(textLineStart, textLineEnd - textLineStart);
+
+									// Insert this line
+									builder.Append($"{delimiter}{{{newName}}}");
+									properties![newName] = newLine;
+									propertyValues[newName] = newLine;
+									delimiter = "\n";
+
+									// Move to the next line
+									textLineStart = ++textLineEnd;
+									while (textLineEnd < text.Length && text[textLineEnd] != '\n')
+									{
+										textLineEnd++;
+									}
+								}
+
+								builder.Append(format, idx + 1, format.Length - (idx + 1));
+								format = builder.ToString();
+							}
+						}
+					}
 				}
+
+				// Get the enumerable property list for formatting
+				propertyValueList = propertyValues.Select(x => new KeyValuePair<string, object?>(x.Key, x.Value));
 			}
 
-			for (int idx = 0; idx < lineCount; idx++)
+			// Finally split the format string into multiple lines
+			string[] lines = format.Split('\n');
+			for (int idx = 0; idx < lines.Length; idx++)
 			{
-				writer.Write(span.Slice(0, insertIdx));
-				if (insertIdx < span.Length)
+				obj[s_messagePropertyName] = MessageTemplate.Render(lines[idx], propertyValueList);
+				obj[s_formatPropertyName] = lines[idx];
+				obj[s_linePropertyName] = idx;
+				obj[s_lineCountPropertyName] = lines.Length;
+
+				using (Utf8JsonWriter jsonWriter = new Utf8JsonWriter(writer))
 				{
-					writer.Write(s_comma);
-					using (Utf8JsonWriter jsonWriter = new Utf8JsonWriter(writer, new JsonWriterOptions { SkipValidation = true }))
-					{
-						jsonWriter.WriteNumber(LogEventPropertyName.Line, idx);
-						jsonWriter.WriteNumber(LogEventPropertyName.LineCount, lineCount);
-					}
-					writer.Write(span.Slice(insertIdx));
+					obj.WriteTo(jsonWriter);
 				}
-				writer.Write(s_newline);
+
+				writer.Write(s_newline.Span);
 			}
+			return lines.Length;
 		}
 
 		void AddEvent(ReadOnlySpan<byte> span, int lineIndex, EventSeverity severity, List<CreateEventRequest> events)
