Horde: Fix line count for multi-line log events.

diff --git a/Engine/Source/Programs/Horde/Horde.Agent/Utility/JsonRpcLogger.cs b/Engine/Source/Programs/Horde/Horde.Agent/Utility/JsonRpcLogger.cs
--- a/Engine/Source/Programs/Horde/Horde.Agent/Utility/JsonRpcLogger.cs
+++ b/Engine/Source/Programs/Horde/Horde.Agent/Utility/JsonRpcLogger.cs
@@ -256,9 +256,6 @@ namespace Horde.Agent.Parser
 			ArrayBufferWriter<byte> writer = new ArrayBufferWriter<byte>();
 			List<CreateEventRequest> events = new List<CreateEventRequest>();
 
-			// Line separator for JSON events
-			byte[] newline = { (byte)'\n' };
-
 			// The current jobstep outcome
 			JobStepOutcome postedOutcome = JobStepOutcome.Success;
 
@@ -278,17 +275,17 @@ namespace Horde.Agent.Parser
 					JsonLogEvent jsonLogEvent;
 					if (_dataChannel.Reader.TryRead(out jsonLogEvent))
 					{
-						// If we want an event for this log event, create one now
+						int lineCount = WriteEvent(jsonLogEvent.Data.Span, writer);
+
 						if (jsonLogEvent.Level == LogLevel.Warning && ++numWarnings <= MaxWarnings)
 						{
-							AddEvent(jsonLogEvent.Data.Span, lineIndex, EventSeverity.Warning, events);
+							AddEvent(jsonLogEvent.Data.Span, lineIndex, lineCount, EventSeverity.Warning, events);
 						}
 						else if ((jsonLogEvent.Level == LogLevel.Error || jsonLogEvent.Level == LogLevel.Critical) && ++numErrors <= MaxErrors)
 						{
-							AddEvent(jsonLogEvent.Data.Span, lineIndex, EventSeverity.Error, events);
+							AddEvent(jsonLogEvent.Data.Span, lineIndex, lineCount, EventSeverity.Error, events);
 						}
 
-						int lineCount = WriteEvent(jsonLogEvent.Data.Span, writer);
 						lineIndex += lineCount;
 					}
 					else
@@ -479,15 +476,11 @@ namespace Horde.Agent.Parser
 			return lines.Length;
 		}
 
-		void AddEvent(ReadOnlySpan<byte> span, int lineIndex, EventSeverity severity, List<CreateEventRequest> events)
+		void AddEvent(ReadOnlySpan<byte> span, int lineIndex, int lineCount, EventSeverity severity, List<CreateEventRequest> events)
 		{
 			try
 			{
-				int lineCount = GetEventLineCount(span);
-				if (lineCount > 0)
-				{
-					events.Add(new CreateEventRequest(severity, _logId, lineIndex, lineCount));
-				}
+				events.Add(new CreateEventRequest(severity, _logId, lineIndex, lineCount));
 			}
 			catch (Exception ex)
 			{
