Horde: Add metadata to issue fingerprints, and support a "CompileType" annotation on compile issues to allow refining issue messages.

diff --git a/Engine/Source/Programs/Horde/Horde.Build.Tests/IssueServiceTests.cs b/Engine/Source/Programs/Horde/Horde.Build.Tests/IssueServiceTests.cs
--- a/Engine/Source/Programs/Horde/Horde.Build.Tests/IssueServiceTests.cs
+++ b/Engine/Source/Programs/Horde/Horde.Build.Tests/IssueServiceTests.cs
@@ -159,24 +159,28 @@ namespace Horde.Build.Tests
 			_perforce.AddChange(MainStreamName, 125, chris, "Description", new string[] { "a/g.cpp" });
 
 			List<INode> nodes = new List<INode>();
-			nodes.Add(MockNode("Update Version Files"));
-			nodes.Add(MockNode("Compile UnrealHeaderTool Win64"));
-			nodes.Add(MockNode("Compile ShooterGameEditor Win64"));
-			nodes.Add(MockNode("Cook ShooterGame Win64"));
+			nodes.Add(MockNode("Update Version Files", NodeAnnotations.Empty));
+			nodes.Add(MockNode("Compile UnrealHeaderTool Win64", NodeAnnotations.Empty));
+			nodes.Add(MockNode("Compile ShooterGameEditor Win64", NodeAnnotations.Empty));
+			nodes.Add(MockNode("Cook ShooterGame Win64", NodeAnnotations.Empty));
 
-			Mock<INodeGroup> @group = new Mock<INodeGroup>(MockBehavior.Strict);
-			@group.SetupGet(x => x.Nodes).Returns(nodes);
+			NodeAnnotations staticAnalysisAnnotations = new NodeAnnotations();
+			staticAnalysisAnnotations.Add("CompileType", "Static analysis");
+			nodes.Add(MockNode("Static Analysis Win64", staticAnalysisAnnotations));
+
+			Mock<INodeGroup> grp = new Mock<INodeGroup>(MockBehavior.Strict);
+			grp.SetupGet(x => x.Nodes).Returns(nodes);
 
 			Mock<IGraph> graphMock = new Mock<IGraph>(MockBehavior.Strict);
-			graphMock.SetupGet(x => x.Groups).Returns(new List<INodeGroup> { @group.Object });
+			graphMock.SetupGet(x => x.Groups).Returns(new List<INodeGroup> { grp.Object });
 			_graph = graphMock.Object;
 		}
 
-		public static INode MockNode(string name)
+		public static INode MockNode(string name, IReadOnlyNodeAnnotations annotations)
 		{
 			Mock<INode> node = new Mock<INode>(MockBehavior.Strict);
 			node.SetupGet(x => x.Name).Returns(name);
-			node.SetupGet(x => x.Annotations).Returns(NodeAnnotations.Empty);
+			node.SetupGet(x => x.Annotations).Returns(annotations);
 			return node.Object;
 		}
 
@@ -735,6 +739,30 @@ namespace Horde.Build.Tests
 			}
 		}
 
+		[TestMethod]
+		public async Task CompileTypeTest()
+		{
+			string[] lines =
+			{
+				FileReference.Combine(_workspaceDir, "FOO.CPP").FullName + @"(170) : warning C6011: Dereferencing NULL pointer 'CurrentProperty'. : Lines: 159, 162, 163, 169, 170, 174, 176, 159, 162, 163, 169, 170",
+				FileReference.Combine(_workspaceDir, "foo.cpp").FullName + @"(170) : warning C6011: Dereferencing NULL pointer 'CurrentProperty'. : Lines: 159, 162, 163, 169, 170, 174, 176, 159, 162, 163, 169, 170"
+			};
+
+			IJob job = CreateJob(_mainStreamId, 120, "Compile Test", _graph);
+			await UpdateCompleteStep(job, 0, 0, JobStepOutcome.Success);
+			await UpdateCompleteStep(job, 0, 1, JobStepOutcome.Success);
+			await UpdateCompleteStep(job, 0, 2, JobStepOutcome.Success);
+			await UpdateCompleteStep(job, 0, 3, JobStepOutcome.Success);
+			await ParseEventsAsync(job, 0, 4, lines);
+			await UpdateCompleteStep(job, 0, 4, JobStepOutcome.Failure);
+
+			List<IIssue> issues = await IssueService.FindIssuesAsync();
+			Assert.AreEqual(1, issues.Count);
+
+			IIssue issue = issues[0];
+			Assert.AreEqual("Static analysis warnings in FOO.CPP", issue.Summary);
+		}
+
 		[TestMethod]
 		public async Task CompileIssueTest()
 		{
diff --git a/Engine/Source/Programs/Horde/Horde.Build/Issues/IIssueCollection.cs b/Engine/Source/Programs/Horde/Horde.Build/Issues/IIssueCollection.cs
--- a/Engine/Source/Programs/Horde/Horde.Build/Issues/IIssueCollection.cs
+++ b/Engine/Source/Programs/Horde/Horde.Build/Issues/IIssueCollection.cs
@@ -17,7 +17,6 @@ namespace Horde.Build.Collections
 	using StreamId = StringId<IStream>;
 	using TemplateRefId = StringId<TemplateRef>;
 	using UserId = ObjectId<IUser>;
-	using WorkflowId = StringId<WorkflowConfig>;
 
 	/// <summary>
 	/// Fingerprint for an issue
@@ -33,13 +32,17 @@ namespace Horde.Build.Collections
 		/// <inheritdoc/>
 		public CaseInsensitiveStringSet? RejectKeys { get; set; }
 
+		/// <inheritdoc/>
+		public CaseInsensitiveStringSet? Metadata { get; set; }
+
 		/// <summary>
 		/// Constructor
 		/// </summary>
 		/// <param name="type">The type of issue</param>
 		/// <param name="keys">Keys which uniquely identify this issue</param>
 		/// <param name="rejectKeys">Keys which should not match with this issue</param>
-		public NewIssueFingerprint(string type, IEnumerable<string> keys, IEnumerable<string>? rejectKeys)
+		/// <param name="metadata">Additional metadata added by the issue handler</param>
+		public NewIssueFingerprint(string type, IEnumerable<string> keys, IEnumerable<string>? rejectKeys, IEnumerable<string>? metadata)
 		{
 			Type = type;
 			Keys = new CaseInsensitiveStringSet(keys);
@@ -48,6 +51,10 @@ namespace Horde.Build.Collections
 			{
 				RejectKeys = new CaseInsensitiveStringSet(rejectKeys);
 			}
+			if (metadata != null && metadata.Any())
+			{
+				Metadata = new CaseInsensitiveStringSet(metadata);
+			}
 		}
 
 		/// <summary>
@@ -55,7 +62,7 @@ namespace Horde.Build.Collections
 		/// </summary>
 		/// <param name="other">The fingerprint to copy from</param>
 		public NewIssueFingerprint(IIssueFingerprint other)
-			: this(other.Type, other.Keys, other.RejectKeys)
+			: this(other.Type, other.Keys, other.RejectKeys, other.Metadata)
 		{
 		}
 
@@ -63,21 +70,19 @@ namespace Horde.Build.Collections
 		/// Merges another fingerprint into this one
 		/// </summary>
 		/// <param name="other">The other fingerprint to merge with</param>
-		public NewIssueFingerprint MergeWith(IIssueFingerprint other)
+		public void MergeWith(IIssueFingerprint other)
 		{
-			NewIssueFingerprint result = new NewIssueFingerprint(this);
-
-			result.Keys.UnionWith(other.Keys);
+			Keys.UnionWith(other.Keys);
 			if (other.RejectKeys != null)
 			{
-				if (result.RejectKeys == null)
-				{
-					result.RejectKeys = new CaseInsensitiveStringSet();
-				}
-				result.RejectKeys.UnionWith(other.RejectKeys);
+				RejectKeys ??= new CaseInsensitiveStringSet();
+				RejectKeys.UnionWith(other.RejectKeys);
+			}
+			if (other.Metadata != null)
+			{
+				Metadata ??= new CaseInsensitiveStringSet();
+				Metadata.UnionWith(other.Metadata);
 			}
-
-			return result;
 		}
 
 		/// <inheritdoc/>
@@ -90,20 +95,13 @@ namespace Horde.Build.Collections
 			{
 				return false;
 			}
-
-			if (RejectKeys == null)
+			if (!ContentsEqual(RejectKeys, otherFingerprint.RejectKeys))
 			{
-				if (otherFingerprint.RejectKeys != null && otherFingerprint.RejectKeys.Count > 0)
-				{
-					return false;
-				}
+				return false;
 			}
-			else
+			if (!ContentsEqual(Metadata, otherFingerprint.Metadata))
 			{
-				if (otherFingerprint.RejectKeys == null || !RejectKeys.SetEquals(otherFingerprint.RejectKeys))
-				{
-					return false;
-				}
+				return false;
 			}
 
 			return true;
@@ -116,18 +114,43 @@ namespace Horde.Build.Collections
 		public override int GetHashCode()
 		{
 			int result = StringComparer.Ordinal.GetHashCode(Type);
-			foreach (string key in Keys)
+			return HashCode.Combine(result, GetContentsHash(Keys), GetContentsHash(RejectKeys), GetContentsHash(Metadata));
+		}
+
+		/// <summary>
+		/// Checks if the contents of two sets are equal
+		/// </summary>
+		/// <param name="setA"></param>
+		/// <param name="setB"></param>
+		/// <returns></returns>
+		static bool ContentsEqual(CaseInsensitiveStringSet? setA, CaseInsensitiveStringSet? setB)
+		{
+			if (setA == null || setA.Count == 0)
 			{
-				result = HashCode.Combine(result, key);
+				return setB == null || setB.Count == 0;
 			}
-			if (RejectKeys != null)
+			else
+			{
+				return setB != null && setA.SetEquals(setB);
+			}
+		}
+
+		/// <summary>
+		/// Gets the hash of the contents of a case insensitive set
+		/// </summary>
+		/// <param name="set"></param>
+		/// <returns></returns>
+		static int GetContentsHash(CaseInsensitiveStringSet? set)
+		{
+			int value = 0;
+			if (set != null)
 			{
-				foreach (string rejectKey in RejectKeys)
+				foreach (string element in set)
 				{
-					result = HashCode.Combine(result, rejectKey);
+					value = HashCode.Combine(value, StringComparer.OrdinalIgnoreCase.GetHashCode(element));
 				}
 			}
-			return result;
+			return value;
 		}
 
 		/// <summary>
@@ -139,17 +162,7 @@ namespace Horde.Build.Collections
 		public static NewIssueFingerprint Merge(IIssueFingerprint a, IIssueFingerprint b)
 		{
 			NewIssueFingerprint newFingerprint = new NewIssueFingerprint(a);
-			newFingerprint.Keys.UnionWith(b.Keys);
-
-			if (b.RejectKeys != null)
-			{
-				if (newFingerprint.RejectKeys == null)
-				{
-					newFingerprint.RejectKeys = new CaseInsensitiveStringSet();
-				}
-				newFingerprint.RejectKeys.UnionWith(b.RejectKeys);
-			}
-
+			newFingerprint.MergeWith(b);
 			return newFingerprint;
 		}
 
diff --git a/Engine/Source/Programs/Horde/Horde.Build/Issues/IIssueFingerprint.cs b/Engine/Source/Programs/Horde/Horde.Build/Issues/IIssueFingerprint.cs
--- a/Engine/Source/Programs/Horde/Horde.Build/Issues/IIssueFingerprint.cs
+++ b/Engine/Source/Programs/Horde/Horde.Build/Issues/IIssueFingerprint.cs
@@ -1,6 +1,7 @@
 // Copyright Epic Games, Inc. All Rights Reserved.
 
 using System;
+using System.Collections.Generic;
 using System.Linq;
 using Horde.Build.Utilities;
 
@@ -12,7 +13,7 @@ namespace Horde.Build.Models
 	public interface IIssueFingerprint
 	{
 		/// <summary>
-		/// The type of issue
+		/// The type of issue, which defines the handler to use for it
 		/// </summary>
 		public string Type { get; }
 
@@ -25,6 +26,11 @@ namespace Horde.Build.Models
 		/// Set of keys which should trigger a negative match
 		/// </summary>
 		public CaseInsensitiveStringSet? RejectKeys { get; }
+
+		/// <summary>
+		/// Collection of additional metadata added by the handler
+		/// </summary>
+		public CaseInsensitiveStringSet? Metadata { get; }
 	}
 
 	/// <summary>
@@ -83,5 +89,25 @@ namespace Horde.Build.Models
 			}
 			return true;
 		}
+
+		/// <summary>
+		/// Gets all the metadata values with a given key
+		/// </summary>
+		/// <param name="fingerprint">Fingerprint to find values for</param>
+		/// <param name="key">Key name to search for</param>
+		/// <returns>All values with the given key</returns>
+		public static IEnumerable<string> GetMetadataValues(this IIssueFingerprint fingerprint, string key)
+		{
+			if (fingerprint.Metadata != null)
+			{
+				foreach (string element in fingerprint.Metadata)
+				{
+					if (element.StartsWith(key, StringComparison.OrdinalIgnoreCase) && element.Length > key.Length && element[key.Length] == '=')
+					{
+						yield return element.Substring(key.Length + 1);
+					}
+				}
+			}
+		}
 	}
 }
diff --git a/Engine/Source/Programs/Horde/Horde.Build/Issues/IIssueHandler.cs b/Engine/Source/Programs/Horde/Horde.Build/Issues/IIssueHandler.cs
--- a/Engine/Source/Programs/Horde/Horde.Build/Issues/IIssueHandler.cs
+++ b/Engine/Source/Programs/Horde/Horde.Build/Issues/IIssueHandler.cs
@@ -137,10 +137,11 @@
 		/// </summary>
 		/// <param name="job">The job that spawned the event</param>
 		/// <param name="node">Node that was executed</param>
+		/// <param name="annotations"></param>
 		/// <param name="eventData">The event data</param>
 		/// <param name="fingerprint">Receives the fingerprint on success</param>
 		/// <returns>True if the match is successful</returns>
-		bool TryGetFingerprint(IJob job, INode node, ILogEventData eventData, [NotNullWhen(true)] out NewIssueFingerprint? fingerprint);
+		bool TryGetFingerprint(IJob job, INode node, IReadOnlyNodeAnnotations annotations, ILogEventData eventData, [NotNullWhen(true)] out NewIssueFingerprint? fingerprint);
 
 		/// <summary>
 		/// Rank all the suspect changes for a given fingerprint
diff --git a/Engine/Source/Programs/Horde/Horde.Build/Issues/Impl/Handlers/CompileIssueHandler.cs b/Engine/Source/Programs/Horde/Horde.Build/Issues/Impl/Handlers/CompileIssueHandler.cs
--- a/Engine/Source/Programs/Horde/Horde.Build/Issues/Impl/Handlers/CompileIssueHandler.cs
+++ b/Engine/Source/Programs/Horde/Horde.Build/Issues/Impl/Handlers/CompileIssueHandler.cs
@@ -24,6 +24,11 @@ namespace Horde.Build.IssueHandlers.Impl
 		/// </summary>
 		const string NotePrefix = "note:";
 
+		/// <summary>
+		/// Annotation describing the compile type
+		/// </summary>
+		const string CompileTypeAnnotation = "CompileType";
+
 		/// <inheritdoc/>
 		public string Type => "Compile";
 
@@ -98,7 +103,7 @@ namespace Horde.Build.IssueHandlers.Impl
 		}
 
 		/// <inheritdoc/>
-		public bool TryGetFingerprint(IJob job, INode node, ILogEventData eventData, [NotNullWhen(true)] out NewIssueFingerprint? fingerprint)
+		public bool TryGetFingerprint(IJob job, INode node, IReadOnlyNodeAnnotations annotations, ILogEventData eventData, [NotNullWhen(true)] out NewIssueFingerprint? fingerprint)
 		{
 			if (!IsMatchingEventId(eventData.EventId))
 			{
@@ -108,16 +113,28 @@ namespace Horde.Build.IssueHandlers.Impl
 
 			List<string> newFileNames = new List<string>();
 			GetSourceFiles(eventData, newFileNames);
-			fingerprint = new NewIssueFingerprint(Type, newFileNames, null);
+
+			string? compileType;
+			if (!annotations.TryGetValue(CompileTypeAnnotation, out compileType))
+			{
+				compileType = "Compile";
+			}
+
+			List<string> newMetadata = new List<string>();
+			newMetadata.Add($"{CompileTypeAnnotation}={compileType}");
+
+			fingerprint = new NewIssueFingerprint(Type, newFileNames, null, newMetadata);
 			return true;
 		}
 
 		/// <inheritdoc/>
 		public string GetSummary(IIssueFingerprint fingerprint, IssueSeverity severity)
 		{
-			string type = (severity == IssueSeverity.Warning) ? "Compile warnings" : "Compile errors";
+			List<string> types = fingerprint.GetMetadataValues(CompileTypeAnnotation).ToList();
+			string type = (types.Count == 1) ? types[0] : "Compile";
+			string level = (severity == IssueSeverity.Warning) ? "warnings" : "errors";
 			string list = StringUtils.FormatList(fingerprint.Keys.Where(x => !x.StartsWith(NotePrefix, StringComparison.Ordinal)).ToArray(), 2);
-			return $"{type} in {list}";
+			return $"{type} {level} in {list}";
 		}
 
 		/// <inheritdoc/>
diff --git a/Engine/Source/Programs/Horde/Horde.Build/Issues/Impl/Handlers/ContentIssueHandler.cs b/Engine/Source/Programs/Horde/Horde.Build/Issues/Impl/Handlers/ContentIssueHandler.cs
--- a/Engine/Source/Programs/Horde/Horde.Build/Issues/Impl/Handlers/ContentIssueHandler.cs
+++ b/Engine/Source/Programs/Horde/Horde.Build/Issues/Impl/Handlers/ContentIssueHandler.cs
@@ -54,7 +54,7 @@ namespace Horde.Build.IssueHandlers.Impl
 		}
 
 		/// <inheritdoc/>
-		public bool TryGetFingerprint(IJob job, INode node, ILogEventData eventData, [NotNullWhen(true)] out NewIssueFingerprint? fingerprint)
+		public bool TryGetFingerprint(IJob job, INode node, IReadOnlyNodeAnnotations annotations, ILogEventData eventData, [NotNullWhen(true)] out NewIssueFingerprint? fingerprint)
 		{
 			if (!IsMatchingEventId(eventData.EventId))
 			{
@@ -64,7 +64,7 @@ namespace Horde.Build.IssueHandlers.Impl
 
 			HashSet<string> newAssetNames = new HashSet<string>();
 			GetAssetNames(eventData, newAssetNames);
-			fingerprint = new NewIssueFingerprint(Type, newAssetNames, null);
+			fingerprint = new NewIssueFingerprint(Type, newAssetNames, null, null);
 			return true;
 		}
 
diff --git a/Engine/Source/Programs/Horde/Horde.Build/Issues/Impl/Handlers/CopyrightIssueHandler.cs b/Engine/Source/Programs/Horde/Horde.Build/Issues/Impl/Handlers/CopyrightIssueHandler.cs
--- a/Engine/Source/Programs/Horde/Horde.Build/Issues/Impl/Handlers/CopyrightIssueHandler.cs
+++ b/Engine/Source/Programs/Horde/Horde.Build/Issues/Impl/Handlers/CopyrightIssueHandler.cs
@@ -59,7 +59,7 @@ namespace Horde.Build.IssueHandlers.Impl
 		}
 
 		/// <inheritdoc/>
-		public bool TryGetFingerprint(IJob job, INode node, ILogEventData eventData, [NotNullWhen(true)] out NewIssueFingerprint? fingerprint)
+		public bool TryGetFingerprint(IJob job, INode node, IReadOnlyNodeAnnotations annotations, ILogEventData eventData, [NotNullWhen(true)] out NewIssueFingerprint? fingerprint)
 		{
 			if (!IsMatchingEventId(eventData.EventId))
 			{
@@ -69,7 +69,7 @@ namespace Horde.Build.IssueHandlers.Impl
 
 			HashSet<string> newFileNames = new HashSet<string>(StringComparer.OrdinalIgnoreCase);
 			GetSourceFiles(eventData, newFileNames);
-			fingerprint = new NewIssueFingerprint(Type, newFileNames, null);
+			fingerprint = new NewIssueFingerprint(Type, newFileNames, null, null);
 			return true;
 		}
 
diff --git a/Engine/Source/Programs/Horde/Horde.Build/Issues/Impl/Handlers/DefaultIssueHandler.cs b/Engine/Source/Programs/Horde/Horde.Build/Issues/Impl/Handlers/DefaultIssueHandler.cs
--- a/Engine/Source/Programs/Horde/Horde.Build/Issues/Impl/Handlers/DefaultIssueHandler.cs
+++ b/Engine/Source/Programs/Horde/Horde.Build/Issues/Impl/Handlers/DefaultIssueHandler.cs
@@ -25,9 +25,9 @@ namespace Horde.Build.IssueHandlers.Impl
 		public int Priority => 0;
 
 		/// <inheritdoc/>
-		public bool TryGetFingerprint(IJob job, INode node, ILogEventData eventData, [NotNullWhen(true)] out NewIssueFingerprint? fingerprint)
+		public bool TryGetFingerprint(IJob job, INode node, IReadOnlyNodeAnnotations annotations, ILogEventData eventData, [NotNullWhen(true)] out NewIssueFingerprint? fingerprint)
 		{
-			fingerprint = new NewIssueFingerprint(Type, new[] { node.Name }, null);
+			fingerprint = new NewIssueFingerprint(Type, new[] { node.Name }, null, null);
 			return true;
 		}
 
diff --git a/Engine/Source/Programs/Horde/Horde.Build/Issues/Impl/Handlers/GauntletIssueHandler.cs b/Engine/Source/Programs/Horde/Horde.Build/Issues/Impl/Handlers/GauntletIssueHandler.cs
--- a/Engine/Source/Programs/Horde/Horde.Build/Issues/Impl/Handlers/GauntletIssueHandler.cs
+++ b/Engine/Source/Programs/Horde/Horde.Build/Issues/Impl/Handlers/GauntletIssueHandler.cs
@@ -97,7 +97,7 @@ namespace Horde.Build.IssueHandlers.Impl
 		}
 
 		/// <inheritdoc/>
-		public bool TryGetFingerprint(IJob job, INode node, ILogEventData eventData, [NotNullWhen(true)] out NewIssueFingerprint? fingerprint)
+		public bool TryGetFingerprint(IJob job, INode node, IReadOnlyNodeAnnotations annotations, ILogEventData eventData, [NotNullWhen(true)] out NewIssueFingerprint? fingerprint)
 		{
 			if(!IsMatchingEventId(eventData.EventId))
 			{
@@ -115,7 +115,7 @@ namespace Horde.Build.IssueHandlers.Impl
 				return false;
 			}
 
-			fingerprint = new NewIssueFingerprint(Type, keys, null);
+			fingerprint = new NewIssueFingerprint(Type, keys, null, null);
 			return true;
 		}
 
diff --git a/Engine/Source/Programs/Horde/Horde.Build/Issues/Impl/Handlers/LocalizationIssueHandler.cs b/Engine/Source/Programs/Horde/Horde.Build/Issues/Impl/Handlers/LocalizationIssueHandler.cs
--- a/Engine/Source/Programs/Horde/Horde.Build/Issues/Impl/Handlers/LocalizationIssueHandler.cs
+++ b/Engine/Source/Programs/Horde/Horde.Build/Issues/Impl/Handlers/LocalizationIssueHandler.cs
@@ -59,7 +59,7 @@ namespace Horde.Build.IssueHandlers.Impl
 		}
 
 		/// <inheritdoc/>
-		public bool TryGetFingerprint(IJob job, INode node, ILogEventData eventData, [NotNullWhen(true)] out NewIssueFingerprint? fingerprint)
+		public bool TryGetFingerprint(IJob job, INode node, IReadOnlyNodeAnnotations annotations, ILogEventData eventData, [NotNullWhen(true)] out NewIssueFingerprint? fingerprint)
 		{
 			if (!IsMatchingEventId(eventData.EventId))
 			{
@@ -69,7 +69,7 @@ namespace Horde.Build.IssueHandlers.Impl
 
 			HashSet<string> newFileNames = new HashSet<string>(StringComparer.OrdinalIgnoreCase);
 			GetSourceFiles(eventData, newFileNames);
-			fingerprint = new NewIssueFingerprint(Type, newFileNames, null);
+			fingerprint = new NewIssueFingerprint(Type, newFileNames, null, null);
 			return true;
 		}
 
diff --git a/Engine/Source/Programs/Horde/Horde.Build/Issues/Impl/Handlers/SymbolIssueHandler.cs b/Engine/Source/Programs/Horde/Horde.Build/Issues/Impl/Handlers/SymbolIssueHandler.cs
--- a/Engine/Source/Programs/Horde/Horde.Build/Issues/Impl/Handlers/SymbolIssueHandler.cs
+++ b/Engine/Source/Programs/Horde/Horde.Build/Issues/Impl/Handlers/SymbolIssueHandler.cs
@@ -83,7 +83,7 @@ namespace Horde.Build.IssueHandlers.Impl
 			}
 		}
 
-		public bool TryGetFingerprint(IJob job, INode node, ILogEventData eventData, [NotNullWhen(true)] out NewIssueFingerprint? fingerprint)
+		public bool TryGetFingerprint(IJob job, INode node, IReadOnlyNodeAnnotations annotations, ILogEventData eventData, [NotNullWhen(true)] out NewIssueFingerprint? fingerprint)
 		{
 			if (!IsMatchingEventId(eventData.EventId))
 			{
@@ -94,7 +94,7 @@ namespace Horde.Build.IssueHandlers.Impl
 			SortedSet<string> symbolNames = new SortedSet<string>();
 			GetSymbolNames(eventData, symbolNames);
 
-			fingerprint = new NewIssueFingerprint(Type, symbolNames, null);
+			fingerprint = new NewIssueFingerprint(Type, symbolNames, null, null);
 			return true;
 		}
 	}
diff --git a/Engine/Source/Programs/Horde/Horde.Build/Issues/Impl/Handlers/SystemicIssueHandler.cs b/Engine/Source/Programs/Horde/Horde.Build/Issues/Impl/Handlers/SystemicIssueHandler.cs
--- a/Engine/Source/Programs/Horde/Horde.Build/Issues/Impl/Handlers/SystemicIssueHandler.cs
+++ b/Engine/Source/Programs/Horde/Horde.Build/Issues/Impl/Handlers/SystemicIssueHandler.cs
@@ -47,7 +47,7 @@ namespace Horde.Build.IssueHandlers.Impl
 		}
 
 		/// <inheritdoc/>
-		public bool TryGetFingerprint(IJob job, INode node, ILogEventData eventData, [NotNullWhen(true)] out NewIssueFingerprint? fingerprint)
+		public bool TryGetFingerprint(IJob job, INode node, IReadOnlyNodeAnnotations annotations, ILogEventData eventData, [NotNullWhen(true)] out NewIssueFingerprint? fingerprint)
 		{
 			fingerprint = null;
 
@@ -63,7 +63,7 @@ namespace Horde.Build.IssueHandlers.Impl
 					{
 						if (message.Contains(systemicExitMessages[j], StringComparison.InvariantCultureIgnoreCase))
 						{
-							fingerprint = new NewIssueFingerprint(Type, new[] { node.Name }, null);
+							fingerprint = new NewIssueFingerprint(Type, new[] { node.Name }, null, null);
 							return true;
 						}
 					}
@@ -75,7 +75,7 @@ namespace Horde.Build.IssueHandlers.Impl
 				return false;
 			}
 
-			fingerprint = new NewIssueFingerprint(Type, new[] { node.Name }, null);
+			fingerprint = new NewIssueFingerprint(Type, new[] { node.Name }, null, null);
 			return true;
 		}
 
diff --git a/Engine/Source/Programs/Horde/Horde.Build/Issues/Impl/IssueCollection.cs b/Engine/Source/Programs/Horde/Horde.Build/Issues/Impl/IssueCollection.cs
--- a/Engine/Source/Programs/Horde/Horde.Build/Issues/Impl/IssueCollection.cs
+++ b/Engine/Source/Programs/Horde/Horde.Build/Issues/Impl/IssueCollection.cs
@@ -223,6 +223,7 @@ namespace Horde.Build.Collections.Impl
 			public string Type { get; set; }
 			public CaseInsensitiveStringSet Keys { get; set; }
 			public CaseInsensitiveStringSet? RejectKeys { get; set; }
+			public CaseInsensitiveStringSet? Metadata { get; set; }
 
 			[BsonConstructor]
 			private IssueFingerprint()
@@ -236,6 +237,7 @@ namespace Horde.Build.Collections.Impl
 				Type = fingerprint.Type;
 				Keys = fingerprint.Keys;
 				RejectKeys = fingerprint.RejectKeys;
+				Metadata = fingerprint.Metadata;
 			}
 		}
 
diff --git a/Engine/Source/Programs/Horde/Horde.Build/Issues/Impl/IssueService.cs b/Engine/Source/Programs/Horde/Horde.Build/Issues/Impl/IssueService.cs
--- a/Engine/Source/Programs/Horde/Horde.Build/Issues/Impl/IssueService.cs
+++ b/Engine/Source/Programs/Horde/Horde.Build/Issues/Impl/IssueService.cs
@@ -536,7 +536,7 @@ namespace Horde.Build.Services.Impl
 			/// <returns>A new group combining both groups</returns>
 			public NewEventGroup MergeWith(NewEventGroup otherGroup)
 			{
-				NewEventGroup newGroup = new NewEventGroup(Fingerprint.MergeWith(otherGroup.Fingerprint));
+				NewEventGroup newGroup = new NewEventGroup(NewIssueFingerprint.Merge(Fingerprint, otherGroup.Fingerprint));
 				newGroup.Events.AddRange(Events);
 				newGroup.Events.AddRange(otherGroup.Events);
 				return newGroup;
@@ -617,7 +617,7 @@ namespace Horde.Build.Services.Impl
 			}
 
 			// Gets the events for this step grouped by fingerprint
-			HashSet<NewEventGroup> eventGroups = await GetEventGroupsForStepAsync(job, batch, step, node);
+			HashSet<NewEventGroup> eventGroups = await GetEventGroupsForStepAsync(job, batch, step, node, annotations);
 
 			// Try to update all the events. We may need to restart this due to optimistic transactions, so keep track of any existing spans we do not need to check against.
 			await using(IAsyncDisposable issueLock = await _issueCollection.EnterCriticalSectionAsync())
@@ -660,8 +660,9 @@ namespace Horde.Build.Services.Impl
 		/// <param name="batch">Unique id of the batch</param>
 		/// <param name="step">Unique id of the step</param>
 		/// <param name="node">The node corresponding to the step</param>
+		/// <param name="annotations">Annotations for this node</param>
 		/// <returns>Set of new events</returns>
-		async Task<HashSet<NewEventGroup>> GetEventGroupsForStepAsync(IJob job, IJobStepBatch batch, IJobStep step, INode node)
+		async Task<HashSet<NewEventGroup>> GetEventGroupsForStepAsync(IJob job, IJobStepBatch batch, IJobStep step, INode node, IReadOnlyNodeAnnotations annotations)
 		{
 			// Make sure the step has a log file
 			if (step.LogId == null)
@@ -691,7 +692,7 @@ namespace Horde.Build.Services.Impl
 				NewIssueFingerprint? fingerprint = null;
 				foreach (IIssueHandler matcher in _matchers)
 				{
-					if (matcher.TryGetFingerprint(job, node, stepEventData, out fingerprint))
+					if (matcher.TryGetFingerprint(job, node, annotations, stepEventData, out fingerprint))
 					{
 						if (fingerprint.Type != "Systemic")
 						{
