Horde: Fix error matchers not skipping over log timestamps from the engine.

diff --git a/Engine/Source/Programs/AutomationTool/AutomationUtils/Matchers/ContentEventMatcher.cs b/Engine/Source/Programs/AutomationTool/AutomationUtils/Matchers/ContentEventMatcher.cs
--- a/Engine/Source/Programs/AutomationTool/AutomationUtils/Matchers/ContentEventMatcher.cs
+++ b/Engine/Source/Programs/AutomationTool/AutomationUtils/Matchers/ContentEventMatcher.cs
@@ -14,13 +14,14 @@ namespace AutomationUtils.Matchers
 	/// </summary>
 	class ContentEventMatcher : ILogEventMatcher
 	{
-		const string Pattern = 
-			@"^\s*" + 
+		static readonly Regex s_pattern = new Regex(
+			@"^\s*" +
+			@"(?:\[[\d\.\-: ]+\])*" +
 			@"(?<channel>[a-zA-Z0-9]+):\s+" + 
 			@"(?<severity>Error:|Warning:)\s+" + 
 			@"(?:\[AssetLog\] )?" + 
 			@"(?<asset>(?:[a-zA-Z]:)?[^:]+(?:.uasset|.umap)):\s*" + 
-			@"(?<message>.*)";
+			@"(?<message>.*)");
 
 		/// <inheritdoc/>
 		public LogEventMatch? Match(ILogCursor input)
@@ -29,8 +30,8 @@ namespace AutomationUtils.Matchers
 			// produce many false positives, making them very slow to execute.
 			if (input.IsMatch("Error:|Warning:"))
 			{
-				Match? match;
-				if(input.TryMatch(Pattern, out match))
+				Match? match = s_pattern.Match(input.CurrentLine!);
+				if (match.Success)
 				{
 					LogEventBuilder builder = new LogEventBuilder(input);
 
diff --git a/Engine/Source/Programs/AutomationTool/AutomationUtils/Matchers/LocalizationEventMatcher.cs b/Engine/Source/Programs/AutomationTool/AutomationUtils/Matchers/LocalizationEventMatcher.cs
--- a/Engine/Source/Programs/AutomationTool/AutomationUtils/Matchers/LocalizationEventMatcher.cs
+++ b/Engine/Source/Programs/AutomationTool/AutomationUtils/Matchers/LocalizationEventMatcher.cs
@@ -13,8 +13,9 @@ namespace AutomationUtils.Matchers
 	/// </summary>
 	class LocalizationEventMatcher : ILogEventMatcher
 	{
-		readonly Regex s_pattern = new Regex(
+		static readonly Regex s_pattern = new Regex(
 			@"^(\s*)" +
+			@"(?:\[[\d\.\-: ]+\])*" +
 			@"(?<channel>LogLocTextHelper|LogGatherTextFromSourceCommandlet):\s*" +
 			@"(?<severity>Error|Warning|Display):\s+" +
 			@"(?<file>([a-zA-Z]:)?[^:/\\]*[/\\][^:]+[^\)])" +
diff --git a/Engine/Source/Programs/AutomationTool/AutomationUtils/Matchers/LogChannelEventMatcher.cs b/Engine/Source/Programs/AutomationTool/AutomationUtils/Matchers/LogChannelEventMatcher.cs
--- a/Engine/Source/Programs/AutomationTool/AutomationUtils/Matchers/LogChannelEventMatcher.cs
+++ b/Engine/Source/Programs/AutomationTool/AutomationUtils/Matchers/LogChannelEventMatcher.cs
@@ -13,11 +13,18 @@ namespace AutomationUtils.Matchers
 	/// </summary>
 	class LogChannelEventMatcher : ILogEventMatcher
 	{
+		readonly static Regex s_pattern = new Regex(
+			@"^(\s*)" +
+			@"(?:\[[\d\.\-: ]+\])*" +
+			@"(?<channel>[a-zA-Z_][a-zA-Z0-9_]*):\s*" +
+			@"(?<severity>Error|Warning|Display): "
+		);
+
 		/// <inheritdoc/>
 		public LogEventMatch? Match(ILogCursor input)
 		{
-			Match? match;
-			if(input.TryMatch(@"^(\s*)(?<channel>[a-zA-Z_][a-zA-Z0-9_]*):\s*(?<severity>Error|Warning|Display): ", out match))
+			Match? match = s_pattern.Match(input.CurrentLine!);
+			if (match.Success)
 			{
 				LogEventBuilder builder = new LogEventBuilder(input);
 				builder.Annotate(match.Groups["channel"], LogEventMarkup.Channel);
diff --git a/Engine/Source/Programs/Horde/Horde.Agent.Tests/LogParserTests.cs b/Engine/Source/Programs/Horde/Horde.Agent.Tests/LogParserTests.cs
--- a/Engine/Source/Programs/Horde/Horde.Agent.Tests/LogParserTests.cs
+++ b/Engine/Source/Programs/Horde/Horde.Agent.Tests/LogParserTests.cs
@@ -762,7 +762,7 @@ namespace Horde.Agent.Tests
 		{
 			string[] lines =
 			{
-				@"LogLocTextHelper: Warning: Source/Runtime/Engine/Private/StaticMesh.cpp(5649): Text conflict from LOCTEXT macro for namespace"
+				@"[2022.05.31-04.25.40:235][  0]LogGatherTextFromSourceCommandlet: Warning: Plugins/Enterprise/VariantManager/Source/VariantManager/Private/SVariantManager.cpp(3717): LOCTEXT macro has an empty identifier and cannot be gathered."
 			};
 
 			{
@@ -771,10 +771,10 @@ namespace Horde.Agent.Tests
 				CheckEventGroup(logEvents, 0, 1, LogLevel.Warning, KnownLogEvents.Engine_Localization);
 
 				LogEvent logEvent = logEvents[0];
-				Assert.AreEqual("LogLocTextHelper", logEvent.GetProperty("channel").ToString());
+				Assert.AreEqual("LogGatherTextFromSourceCommandlet", logEvent.GetProperty("channel").ToString());
 				Assert.AreEqual("Warning", logEvent.GetProperty("severity").ToString());
-				Assert.AreEqual("Engine/Source/Runtime/Engine/Private/StaticMesh.cpp", logEvent.GetProperty<LogValue>("file")!.Properties!["relativePath"].ToString());
-				Assert.AreEqual("5649", logEvent.GetProperty("line").ToString());
+				Assert.AreEqual("Engine/Plugins/Enterprise/VariantManager/Source/VariantManager/Private/SVariantManager.cpp", logEvent.GetProperty<LogValue>("file")!.Properties!["relativePath"].ToString());
+				Assert.AreEqual("3717", logEvent.GetProperty("line").ToString());
 			}
 		}
 
