Horde: Fix allowMentions flag not being respected for slack notifications.

diff --git a/Engine/Source/Programs/Horde/Horde.Build/Notifications/Impl/SlackNotificationSink.cs b/Engine/Source/Programs/Horde/Horde.Build/Notifications/Impl/SlackNotificationSink.cs
--- a/Engine/Source/Programs/Horde/Horde.Build/Notifications/Impl/SlackNotificationSink.cs
+++ b/Engine/Source/Programs/Horde/Horde.Build/Notifications/Impl/SlackNotificationSink.cs
@@ -1134,7 +1134,7 @@ namespace Horde.Build.Notifications.Impl
 				return user.Login;
 			}
 
-			if (!_environment.IsProduction() && allowMentions)
+			if (!_environment.IsProduction() || !allowMentions)
 			{
 				return $"{user.Name} [{slackUserId}]";
 			}
