Horde: Add a matcher for inconsistent Perforce case events.

diff --git a/Engine/Source/Programs/AutomationTool/Scripts/CheckPerforceCase.cs b/Engine/Source/Programs/AutomationTool/Scripts/CheckPerforceCase.cs
--- a/Engine/Source/Programs/AutomationTool/Scripts/CheckPerforceCase.cs
+++ b/Engine/Source/Programs/AutomationTool/Scripts/CheckPerforceCase.cs
@@ -6,6 +6,8 @@ using System.Collections.Generic;
 using System.Linq;
 using System.Text;
 using System.Threading.Tasks;
+using Microsoft.Extensions.Logging;
+using EpicGames.Core;
 
 namespace BuildScripts.Automation
 {
@@ -64,7 +66,7 @@ namespace BuildScripts.Automation
 				}
 			}
 
-			public int PrintIssues()
+			public int PrintIssues(ILogger Logger)
 			{
 				int NumIssues = 0;
 				foreach (IGrouping<string, KeyValuePair<string, TreeNode>> Group in ChildNodes.GroupBy(x => x.Key, StringComparer.OrdinalIgnoreCase))
@@ -72,7 +74,6 @@ namespace BuildScripts.Automation
 					if (Group.Count() > 1)
 					{
 						StringBuilder Message = new StringBuilder();
-						Message.AppendFormat("Inconsistent casing for {0}:", Group.First().Value.Path);
 						foreach (KeyValuePair<string, TreeNode> Pair in Group)
 						{
 							Message.AppendFormat("\n  Could be '{0}':", Pair.Key, Pair.Value.Count, (Pair.Value.Count == 1)? "file" : "files");
@@ -89,13 +90,15 @@ namespace BuildScripts.Automation
 
 							Message.AppendFormat("\n    ({0} {1})", Pair.Value.Count, (Pair.Value.Count == 1) ? "file" : "files");
 						}
-						LogWarning(Message.ToString());
+
+						Logger.LogWarning(KnownLogEvents.AutomationTool_PerforceCase, "Inconsistent casing for {File}: {Message}", new LogValue(LogValueType.DepotPath, Group.First().Value.Path), Message.ToString());
+
 						NumIssues++;
 					}
 
 					foreach (KeyValuePair<string, TreeNode> Pair in Group)
 					{
-						NumIssues += Pair.Value.PrintIssues();
+						NumIssues += Pair.Value.PrintIssues(Logger);
 					}
 				}
 				return NumIssues;
@@ -127,7 +130,7 @@ namespace BuildScripts.Automation
 				}
 				RootNode.Add(File, 2);
 			}
-			RootNode.PrintIssues();
+			RootNode.PrintIssues(Logger);
 		}
 	}
 }
diff --git a/Engine/Source/Programs/Horde/Horde.Build/Issues/Impl/Handlers/PerforceCaseIssueHandler.cs b/Engine/Source/Programs/Horde/Horde.Build/Issues/Impl/Handlers/PerforceCaseIssueHandler.cs
new file mode 100644
--- /dev/null
+++ b/Engine/Source/Programs/Horde/Horde.Build/Issues/Impl/Handlers/PerforceCaseIssueHandler.cs
@@ -0,0 +1,92 @@
+// Copyright Epic Games, Inc. All Rights Reserved.
+
+using System.Collections.Generic;
+using System.Diagnostics.CodeAnalysis;
+using System.Linq;
+using System.Text.Json;
+using EpicGames.Core;
+using Horde.Build.Collections;
+using Horde.Build.Models;
+using Microsoft.Extensions.Logging;
+
+namespace Horde.Build.IssueHandlers.Impl
+{
+	/// <summary>
+	/// Instance of a Perforce case mismatch error
+	/// </summary>
+	class PerforceCaseIssueHandler : IIssueHandler
+	{
+		/// <inheritdoc/>
+		public string Type => "PerforceCase";
+
+		/// <inheritdoc/>
+		public int Priority => 10;
+
+		/// <summary>
+		/// Determines if the given event id matches
+		/// </summary>
+		/// <param name="eventId">The event id to compare</param>
+		/// <returns>True if the given event id matches</returns>
+		public static bool IsMatchingEventId(EventId? eventId)
+		{
+			return eventId == KnownLogEvents.AutomationTool_PerforceCase;
+		}
+
+		/// <inheritdoc/>
+		public void RankSuspects(IIssueFingerprint fingerprint, List<SuspectChange> changes)
+		{
+			foreach (SuspectChange change in changes)
+			{
+				if(change.Details.Files.Any(x => fingerprint.Keys.Contains(x.DepotPath)))
+				{
+					change.Rank += 30;
+				}
+			}
+		}
+
+		/// <inheritdoc/>
+		public string GetSummary(IIssueFingerprint fingerprint, IssueSeverity severity)
+		{
+			return $"Inconsistent case for {StringUtils.FormatList(fingerprint.Keys.Select(x => x.Substring(x.LastIndexOf('/') + 1)).ToArray(), 3)}";
+		}
+
+		/// <summary>
+		/// Extracts a list of source files from an event
+		/// </summary>
+		/// <param name="logEventData">The event data</param>
+		/// <param name="depotFiles">List of source files</param>
+		static void GetSourceFiles(ILogEventData logEventData, HashSet<string> depotFiles)
+		{
+			foreach (ILogEventLine line in logEventData.Lines)
+			{
+				JsonElement properties;
+				if (line.Data.TryGetProperty("properties", out properties) && properties.ValueKind == JsonValueKind.Object)
+				{
+					foreach (JsonProperty property in properties.EnumerateObject())
+					{
+						if (property.NameEquals("File") && property.Value.ValueKind == JsonValueKind.String)
+						{
+							depotFiles.Add(property.Value.GetString()!);
+						}
+					}
+				}
+			}
+		}
+
+		/// <inheritdoc/>
+		public bool TryGetFingerprint(IJob job, INode node, IReadOnlyNodeAnnotations annotations, ILogEventData eventData, [NotNullWhen(true)] out NewIssueFingerprint? fingerprint)
+		{
+			if (!IsMatchingEventId(eventData.EventId))
+			{
+				fingerprint = null;
+				return false;
+			}
+
+			HashSet<string> newFileNames = new HashSet<string>();
+			GetSourceFiles(eventData, newFileNames);
+
+			fingerprint = new NewIssueFingerprint(Type, newFileNames, null, null);
+			return true;
+		}
+	}
+}
diff --git a/Engine/Source/Programs/Shared/EpicGames.Core/KnownLogEvents.cs b/Engine/Source/Programs/Shared/EpicGames.Core/KnownLogEvents.cs
--- a/Engine/Source/Programs/Shared/EpicGames.Core/KnownLogEvents.cs
+++ b/Engine/Source/Programs/Shared/EpicGames.Core/KnownLogEvents.cs
@@ -100,6 +100,11 @@ namespace EpicGames.Core
 		/// </summary>
 		public static EventId AutomationTool_MissingCopyright { get; } = new EventId(404);
 
+		/// <summary>
+		/// UAT: Mismatched Perforce case
+		/// </summary>
+		public static EventId AutomationTool_PerforceCase { get; } = new EventId(405);
+
 		/// <summary>
 		/// MSBuild: Generic error
 		/// </summary>
diff --git a/Engine/Source/Programs/Shared/EpicGames.Core/LogValue.cs b/Engine/Source/Programs/Shared/EpicGames.Core/LogValue.cs
--- a/Engine/Source/Programs/Shared/EpicGames.Core/LogValue.cs
+++ b/Engine/Source/Programs/Shared/EpicGames.Core/LogValue.cs
@@ -24,6 +24,7 @@ namespace EpicGames.Core
 		public static readonly Utf8String ErrorCode = "ErrorCode";
 		public static readonly Utf8String ToolName = "ToolName";
 		public static readonly Utf8String ScreenshotTest = "ScreenshotTest";
+		public static readonly Utf8String DepotPath = "DepotPath";
 	}
 
 	/// <summary>
