Horde: Add server side of device reservation model constraints to v1 endpoint
#rnx

diff --git a/Engine/Source/Programs/Horde/Horde.Build/Controllers/DevicesController.cs b/Engine/Source/Programs/Horde/Horde.Build/Controllers/DevicesController.cs
--- a/Engine/Source/Programs/Horde/Horde.Build/Controllers/DevicesController.cs
+++ b/Engine/Source/Programs/Horde/Horde.Build/Controllers/DevicesController.cs
@@ -726,23 +726,23 @@ namespace Horde.Build.Controllers
 
 			List<DeviceRequestData> requestedDevices = new List<DeviceRequestData>();
 
+			HashSet<string> legacyPerfSpecs = new HashSet<string> { "Unspecified", "Minimum", "Recommended", "High" };
+
 			List<string> perfSpecs = new List<string>();
 
 			foreach (string deviceType in request.DeviceTypes)
 			{
 
 				string platformName = deviceType;
-				string perfSpecName = "Unspecified";
+				string? constraint = null;
 
 				if (deviceType.Contains(':', StringComparison.Ordinal))
 				{
 					string[] tokens = deviceType.Split(":");
 					platformName = tokens[0];
-					perfSpecName = tokens[1];
+					constraint = tokens[1];
 				}
 
-				perfSpecs.Add(perfSpecName);
-
 				DevicePlatformId platformId;
 
 				DevicePlatformMapV1 mapV1 = await _deviceService.GetPlatformMapV1();
@@ -768,24 +768,55 @@ namespace Horde.Build.Controllers
 				List<string> includeModels = new List<string>();
 				List<string> excludeModels = new List<string>();
 
-				if (perfSpecName == "High")
+				if (constraint == null)
+				{
+					perfSpecs.Add("Unspecified");
+				}
+				else if (legacyPerfSpecs.Contains(constraint))
 				{
-					string? model = null;
-					if (mapV1.PerfSpecHighMap.TryGetValue(platformId, out model))
+					perfSpecs.Add(constraint);
+
+					if (constraint == "High")
 					{
-						includeModels.Add(model);
+						string? model = null;
+						if (mapV1.PerfSpecHighMap.TryGetValue(platformId, out model))
+						{
+							includeModels.Add(model);
+						}
 					}
+					else if (constraint == "Minimum" || constraint == "Recommended")
+					{
+						string? model = null;
+						if (mapV1.PerfSpecHighMap.TryGetValue(platformId, out model))
+						{
+							excludeModels.Add(model);
+						}
+					}					
 				}
-
-				if (perfSpecName == "Minimum" || perfSpecName == "Recommended")
+				else
 				{
-					string? model = null;
-					if (mapV1.PerfSpecHighMap.TryGetValue(platformId, out model))
+					string? model = platform.Models?.FirstOrDefault(x => x == constraint);
+					if (model == null)
 					{
-						excludeModels.Add(model);
+						return NotFound($"Platform {platform.Id} has no model {model}");
 					}
-				}
 
+					string modelPerfSpec = "Minimum";
+					string? specModel = null;
+
+					if (mapV1.PerfSpecHighMap.TryGetValue(platform.Id, out specModel))
+					{
+						if (model == specModel)
+						{
+							modelPerfSpec = "High";
+						}
+					}
+
+					perfSpecs.Add(modelPerfSpec);
+
+					includeModels.Add(constraint);
+				}
+								
 				requestedDevices.Add(new DeviceRequestData(platformId, platformName, includeModels, excludeModels));
 			}
 
@@ -803,6 +834,7 @@ namespace Horde.Build.Controllers
 			response.Guid = reservation.LegacyGuid;
 			response.DeviceNames = devices.Select(x => x.Name).ToArray();
 			response.DevicePerfSpecs = perfSpecs.ToArray();
+			response.DeviceModels = devices.Select(x => x.ModelId ?? "Base").ToArray();
 			response.HostName = request.Hostname;
 			response.StartDateTime = reservation.CreateTimeUtc.ToString("O", System.Globalization.CultureInfo.InvariantCulture);
 			response.Duration = $"{request.Duration}";
@@ -925,8 +957,8 @@ namespace Horde.Build.Controllers
 			response.AvailableStartTime = "00:00:00";
 			response.AvailableEndTime = "00:00:00";
 			response.Enabled = true;
+			response.Model = device.ModelId ?? "Base";
 			response.DeviceData = "";
-
 			response.PerfSpec = "Minimum";
 
 			if (device.ModelId != null)
