Horde: Add missing file from previous submit
#rnx

diff --git a/Engine/Source/Programs/Horde/Horde.Build/Api/Devices.cs b/Engine/Source/Programs/Horde/Horde.Build/Api/Devices.cs
--- a/Engine/Source/Programs/Horde/Horde.Build/Api/Devices.cs
+++ b/Engine/Source/Programs/Horde/Horde.Build/Api/Devices.cs
@@ -611,6 +611,12 @@ namespace Horde.Build.Api
 		[Required]
 		public string[] DevicePerfSpecs { get; set; } = null!;
 
+		/// <summary>
+		/// The corresponding perf specs of the reserved devices
+		/// </summary>
+		[Required]
+		public string[] DeviceModels { get; set; } = null!;
+
 		/// <summary>
 		/// The host name of the machine making the reservation
 		/// </summary>
@@ -675,6 +681,12 @@ namespace Horde.Build.Api
 		[Required]
 		public string PerfSpec { get; set; } = null!;
 
+		/// <summary>
+		/// The device model information
+		/// </summary>
+		[Required]
+		public string Model { get; set; } = null!;
+
 		/// <summary>
 		/// The available start time which is parsed client side
 		/// </summary>
