Horde: Removing device model contraints from v1 endpoint for deployment
#rnx

diff --git a/Engine/Source/Programs/Horde/Horde.Build/Api/Devices.cs b/Engine/Source/Programs/Horde/Horde.Build/Api/Devices.cs
--- a/Engine/Source/Programs/Horde/Horde.Build/Api/Devices.cs
+++ b/Engine/Source/Programs/Horde/Horde.Build/Api/Devices.cs
@@ -611,12 +611,6 @@ namespace Horde.Build.Api
 		[Required]
 		public string[] DevicePerfSpecs { get; set; } = null!;
 
-		/// <summary>
-		/// The corresponding perf specs of the reserved devices
-		/// </summary>
-		[Required]
-		public string[] DeviceModels { get; set; } = null!;
-
 		/// <summary>
 		/// The host name of the machine making the reservation
 		/// </summary>
@@ -681,12 +675,6 @@ namespace Horde.Build.Api
 		[Required]
 		public string PerfSpec { get; set; } = null!;
 
-		/// <summary>
-		/// The device model information
-		/// </summary>
-		[Required]
-		public string Model { get; set; } = null!;
-
 		/// <summary>
 		/// The available start time which is parsed client side
 		/// </summary>
diff --git a/Engine/Source/Programs/Horde/Horde.Build/Controllers/DevicesController.cs b/Engine/Source/Programs/Horde/Horde.Build/Controllers/DevicesController.cs
--- a/Engine/Source/Programs/Horde/Horde.Build/Controllers/DevicesController.cs
+++ b/Engine/Source/Programs/Horde/Horde.Build/Controllers/DevicesController.cs
@@ -726,23 +726,23 @@ namespace Horde.Build.Controllers
 
 			List<DeviceRequestData> requestedDevices = new List<DeviceRequestData>();
 
-			HashSet<string> legacyPerfSpecs = new HashSet<string> { "Unspecified", "Minimum", "Recommended", "High" };
-
 			List<string> perfSpecs = new List<string>();
 
 			foreach (string deviceType in request.DeviceTypes)
 			{
 
 				string platformName = deviceType;
-				string? constraint = null;
+				string perfSpecName = "Unspecified";
 
 				if (deviceType.Contains(':', StringComparison.Ordinal))
 				{
 					string[] tokens = deviceType.Split(":");
 					platformName = tokens[0];
-					constraint = tokens[1];
+					perfSpecName = tokens[1];
 				}
 
+				perfSpecs.Add(perfSpecName);
+
 				DevicePlatformId platformId;
 
 				DevicePlatformMapV1 mapV1 = await _deviceService.GetPlatformMapV1();
@@ -768,55 +768,24 @@ namespace Horde.Build.Controllers
 				List<string> includeModels = new List<string>();
 				List<string> excludeModels = new List<string>();
 
-				if (constraint == null)
-				{
-					perfSpecs.Add("Unspecified");
-				}
-				else if (legacyPerfSpecs.Contains(constraint))
+				if (perfSpecName == "High")
 				{
-					perfSpecs.Add(constraint);
-
-					if (constraint == "High")
+					string? model = null;
+					if (mapV1.PerfSpecHighMap.TryGetValue(platformId, out model))
 					{
-						string? model = null;
-						if (mapV1.PerfSpecHighMap.TryGetValue(platformId, out model))
-						{
-							includeModels.Add(model);
-						}
+						includeModels.Add(model);
 					}
-					else if (constraint == "Minimum" || constraint == "Recommended")
-					{
-						string? model = null;
-						if (mapV1.PerfSpecHighMap.TryGetValue(platformId, out model))
-						{
-							excludeModels.Add(model);
-						}
-					}					
 				}
-				else
+
+				if (perfSpecName == "Minimum" || perfSpecName == "Recommended")
 				{
-					string? model = platform.Models?.FirstOrDefault(x => x == constraint);
-					if (model == null)
+					string? model = null;
+					if (mapV1.PerfSpecHighMap.TryGetValue(platformId, out model))
 					{
-						return NotFound($"Platform {platform.Id} has no model {model}");
+						excludeModels.Add(model);
 					}
-
-					string modelPerfSpec = "Minimum";
-					string? specModel = null;
-
-					if (mapV1.PerfSpecHighMap.TryGetValue(platform.Id, out specModel))
-					{
-						if (model == specModel)
-						{
-							modelPerfSpec = "High";
-						}
-					}
-
-					perfSpecs.Add(modelPerfSpec);
-
-					includeModels.Add(constraint);
 				}
-								
+
 				requestedDevices.Add(new DeviceRequestData(platformId, platformName, includeModels, excludeModels));
 			}
 
@@ -834,7 +803,6 @@ namespace Horde.Build.Controllers
 			response.Guid = reservation.LegacyGuid;
 			response.DeviceNames = devices.Select(x => x.Name).ToArray();
 			response.DevicePerfSpecs = perfSpecs.ToArray();
-			response.DeviceModels = devices.Select(x => x.ModelId ?? "Base").ToArray();
 			response.HostName = request.Hostname;
 			response.StartDateTime = reservation.CreateTimeUtc.ToString("O", System.Globalization.CultureInfo.InvariantCulture);
 			response.Duration = $"{request.Duration}";
@@ -957,8 +925,8 @@ namespace Horde.Build.Controllers
 			response.AvailableStartTime = "00:00:00";
 			response.AvailableEndTime = "00:00:00";
 			response.Enabled = true;
-			response.Model = device.ModelId ?? "Base";
 			response.DeviceData = "";
+
 			response.PerfSpec = "Minimum";
 
 			if (device.ModelId != null)
