Horde: Add issue matcher for unacceptable words.

diff --git a/Engine/Source/Programs/Horde/Horde.Build/Issues/Impl/Handlers/CompileIssueHandler.cs b/Engine/Source/Programs/Horde/Horde.Build/Issues/Impl/Handlers/CompileIssueHandler.cs
--- a/Engine/Source/Programs/Horde/Horde.Build/Issues/Impl/Handlers/CompileIssueHandler.cs
+++ b/Engine/Source/Programs/Horde/Horde.Build/Issues/Impl/Handlers/CompileIssueHandler.cs
@@ -4,11 +4,9 @@ using System;
 using System.Collections.Generic;
 using System.Diagnostics.CodeAnalysis;
 using System.Linq;
-using System.Text.Json;
 using EpicGames.Core;
 using Horde.Build.Collections;
 using Horde.Build.Models;
-using Horde.Build.Utilities;
 using Microsoft.Extensions.Logging;
 using MongoDB.Driver;
 
@@ -17,23 +15,18 @@ namespace Horde.Build.IssueHandlers.Impl
 	/// <summary>
 	/// Instance of a particular compile error
 	/// </summary>
-	class CompileIssueHandler : IIssueHandler
+	class CompileIssueHandler : SourceFileIssueHandler
 	{
-		/// <summary>
-		/// Prefix used to identify files that may match against modified files, but which are not the files failing to compile
-		/// </summary>
-		const string NotePrefix = "note:";
-
 		/// <summary>
 		/// Annotation describing the compile type
 		/// </summary>
 		const string CompileTypeAnnotation = "CompileType";
 
 		/// <inheritdoc/>
-		public string Type => "Compile";
+		public override string Type => "Compile";
 
 		/// <inheritdoc/>
-		public int Priority => 10;
+		public override int Priority => 10;
 
 		/// <summary>
 		/// Determines if the given event id matches
@@ -45,65 +38,8 @@ namespace Horde.Build.IssueHandlers.Impl
 			return eventId == KnownLogEvents.Compiler || eventId == KnownLogEvents.AutomationTool_SourceFileLine || eventId == KnownLogEvents.MSBuild;
 		}
 
-		/// <summary>
-		/// Extracts a list of source files from an event
-		/// </summary>
-		/// <param name="logEventData">The event data</param>
-		/// <param name="sourceFiles">List of source files</param>
-		static void GetSourceFiles(ILogEventData logEventData, List<string> sourceFiles)
-		{
-			foreach (ILogEventLine line in logEventData.Lines)
-			{
-				JsonElement properties;
-				if (line.Data.TryGetProperty("properties", out properties) && properties.ValueKind == JsonValueKind.Object)
-				{
-					string? prefix = null;
-
-					JsonElement noteElement;
-					if (properties.TryGetProperty("note", out noteElement) && noteElement.GetBoolean())
-					{
-						prefix = NotePrefix;
-					}
-
-					foreach (JsonProperty property in properties.EnumerateObject())
-					{
-						if (property.NameEquals("file") && property.Value.ValueKind == JsonValueKind.String)
-						{
-							AddSourceFile(sourceFiles, property.Value.GetString()!, prefix);
-						}
-						if (property.Value.HasStringProperty("$type", "SourceFile") && property.Value.TryGetStringProperty("relativePath", out string? value))
-						{
-							AddSourceFile(sourceFiles, value, prefix);
-						}
-					}
-				}
-			}
-		}
-
-		/// <summary>
-		/// Add a new source file to a list of unique source files
-		/// </summary>
-		/// <param name="sourceFiles">List of source files</param>
-		/// <param name="relativePath">File to add</param>
-		/// <param name="prefix">Prefix to insert at the start of the filename</param>
-		static void AddSourceFile(List<string> sourceFiles, string relativePath, string? prefix)
-		{
-			int endIdx = relativePath.LastIndexOfAny(new char[] { '/', '\\' }) + 1;
-
-			string fileName = relativePath.Substring(endIdx);
-			if (prefix != null)
-			{
-				fileName = prefix + fileName;
-			}
-
-			if (!sourceFiles.Any(x => x.Equals(fileName, StringComparison.OrdinalIgnoreCase)))
-			{
-				sourceFiles.Add(fileName);
-			}
-		}
-
 		/// <inheritdoc/>
-		public bool TryGetFingerprint(IJob job, INode node, IReadOnlyNodeAnnotations annotations, ILogEventData eventData, [NotNullWhen(true)] out NewIssueFingerprint? fingerprint)
+		public override bool TryGetFingerprint(IJob job, INode node, IReadOnlyNodeAnnotations annotations, ILogEventData eventData, [NotNullWhen(true)] out NewIssueFingerprint? fingerprint)
 		{
 			if (!IsMatchingEventId(eventData.EventId))
 			{
@@ -128,7 +64,7 @@ namespace Horde.Build.IssueHandlers.Impl
 		}
 
 		/// <inheritdoc/>
-		public string GetSummary(IIssueFingerprint fingerprint, IssueSeverity severity)
+		public override string GetSummary(IIssueFingerprint fingerprint, IssueSeverity severity)
 		{
 			List<string> types = fingerprint.GetMetadataValues(CompileTypeAnnotation).ToList();
 			string type = (types.Count == 1) ? types[0] : "Compile";
@@ -136,37 +72,5 @@ namespace Horde.Build.IssueHandlers.Impl
 			string list = StringUtils.FormatList(fingerprint.Keys.Where(x => !x.StartsWith(NotePrefix, StringComparison.Ordinal)).ToArray(), 2);
 			return $"{type} {level} in {list}";
 		}
-
-		/// <inheritdoc/>
-		public void RankSuspects(IIssueFingerprint fingerprint, List<SuspectChange> suspects)
-		{
-			List<string> fileNames = new List<string>();
-			foreach (string key in fingerprint.Keys)
-			{
-				if (key.StartsWith(NotePrefix, StringComparison.Ordinal))
-				{
-					fileNames.Add(key.Substring(NotePrefix.Length));
-				}
-				else
-				{
-					fileNames.Add(key);
-				}
-			}
-
-			foreach (SuspectChange change in suspects)
-			{
-				if (change.ContainsCode)
-				{
-					if (fileNames.Any(x => change.ModifiesFile(x)))
-					{
-						change.Rank += 20;
-					}
-					else
-					{
-						change.Rank += 10;
-					}
-				}
-			}
-		}
 	}
 }
diff --git a/Engine/Source/Programs/Horde/Horde.Build/Issues/Impl/Handlers/SourceFileIssueHandler.cs b/Engine/Source/Programs/Horde/Horde.Build/Issues/Impl/Handlers/SourceFileIssueHandler.cs
new file mode 100644
--- /dev/null
+++ b/Engine/Source/Programs/Horde/Horde.Build/Issues/Impl/Handlers/SourceFileIssueHandler.cs
@@ -0,0 +1,126 @@
+// Copyright Epic Games, Inc. All Rights Reserved.
+
+using System;
+using System.Collections.Generic;
+using System.Diagnostics.CodeAnalysis;
+using System.Linq;
+using System.Text.Json;
+using Horde.Build.Collections;
+using Horde.Build.Models;
+using Horde.Build.Utilities;
+using MongoDB.Driver;
+
+namespace Horde.Build.IssueHandlers.Impl
+{
+	/// <summary>
+	/// Instance of a particular compile error
+	/// </summary>
+	abstract class SourceFileIssueHandler : IIssueHandler
+	{
+		/// <summary>
+		/// Prefix used to identify files that may match against modified files, but which are not the files failing to compile
+		/// </summary>
+		protected const string NotePrefix = "note:";
+
+		/// <inheritdoc/>
+		public abstract string Type { get; }
+
+		/// <inheritdoc/>
+		public abstract int Priority { get; }
+
+		/// <summary>
+		/// Extracts a list of source files from an event
+		/// </summary>
+		/// <param name="logEventData">The event data</param>
+		/// <param name="sourceFiles">List of source files</param>
+		protected static void GetSourceFiles(ILogEventData logEventData, List<string> sourceFiles)
+		{
+			foreach (ILogEventLine line in logEventData.Lines)
+			{
+				JsonElement properties;
+				if (line.Data.TryGetProperty("properties", out properties) && properties.ValueKind == JsonValueKind.Object)
+				{
+					string? prefix = null;
+
+					JsonElement noteElement;
+					if (properties.TryGetProperty("note", out noteElement) && noteElement.GetBoolean())
+					{
+						prefix = NotePrefix;
+					}
+
+					foreach (JsonProperty property in properties.EnumerateObject())
+					{
+						if (property.NameEquals("file") && property.Value.ValueKind == JsonValueKind.String)
+						{
+							AddSourceFile(sourceFiles, property.Value.GetString()!, prefix);
+						}
+						if (property.Value.HasStringProperty("$type", "SourceFile") && property.Value.TryGetStringProperty("relativePath", out string? value))
+						{
+							AddSourceFile(sourceFiles, value, prefix);
+						}
+					}
+				}
+			}
+		}
+
+		/// <summary>
+		/// Add a new source file to a list of unique source files
+		/// </summary>
+		/// <param name="sourceFiles">List of source files</param>
+		/// <param name="relativePath">File to add</param>
+		/// <param name="prefix">Prefix to insert at the start of the filename</param>
+		static void AddSourceFile(List<string> sourceFiles, string relativePath, string? prefix)
+		{
+			int endIdx = relativePath.LastIndexOfAny(new char[] { '/', '\\' }) + 1;
+
+			string fileName = relativePath.Substring(endIdx);
+			if (prefix != null)
+			{
+				fileName = prefix + fileName;
+			}
+
+			if (!sourceFiles.Any(x => x.Equals(fileName, StringComparison.OrdinalIgnoreCase)))
+			{
+				sourceFiles.Add(fileName);
+			}
+		}
+
+		/// <inheritdoc/>
+		public abstract bool TryGetFingerprint(IJob job, INode node, IReadOnlyNodeAnnotations annotations, ILogEventData eventData, [NotNullWhen(true)] out NewIssueFingerprint? fingerprint);
+
+		/// <inheritdoc/>
+		public abstract string GetSummary(IIssueFingerprint fingerprint, IssueSeverity severity);
+
+		/// <inheritdoc/>
+		public void RankSuspects(IIssueFingerprint fingerprint, List<SuspectChange> suspects)
+		{
+			List<string> fileNames = new List<string>();
+			foreach (string key in fingerprint.Keys)
+			{
+				if (key.StartsWith(NotePrefix, StringComparison.Ordinal))
+				{
+					fileNames.Add(key.Substring(NotePrefix.Length));
+				}
+				else
+				{
+					fileNames.Add(key);
+				}
+			}
+
+			foreach (SuspectChange change in suspects)
+			{
+				if (change.ContainsCode)
+				{
+					if (fileNames.Any(x => change.ModifiesFile(x)))
+					{
+						change.Rank += 20;
+					}
+					else
+					{
+						change.Rank += 10;
+					}
+				}
+			}
+		}
+	}
+}
diff --git a/Engine/Source/Programs/Horde/Horde.Build/Issues/Impl/Handlers/UnacceptableWordsIssueHandler.cs b/Engine/Source/Programs/Horde/Horde.Build/Issues/Impl/Handlers/UnacceptableWordsIssueHandler.cs
new file mode 100644
--- /dev/null
+++ b/Engine/Source/Programs/Horde/Horde.Build/Issues/Impl/Handlers/UnacceptableWordsIssueHandler.cs
@@ -0,0 +1,56 @@
+// Copyright Epic Games, Inc. All Rights Reserved.
+
+using System.Collections.Generic;
+using System.Diagnostics.CodeAnalysis;
+using System.Linq;
+using EpicGames.Core;
+using Horde.Build.Collections;
+using Horde.Build.Models;
+using Microsoft.Extensions.Logging;
+
+namespace Horde.Build.IssueHandlers.Impl
+{
+	/// <summary>
+	/// Instance of a Perforce case mismatch error
+	/// </summary>
+	class UnacceptableWordsIssueHandler : SourceFileIssueHandler
+	{
+		/// <inheritdoc/>
+		public override string Type => "UnacceptableWords";
+
+		/// <inheritdoc/>
+		public override int Priority => 8;
+
+		/// <summary>
+		/// Determines if the given event id matches
+		/// </summary>
+		/// <param name="eventId">The event id to compare</param>
+		/// <returns>True if the given event id matches</returns>
+		public static bool IsMatchingEventId(EventId? eventId)
+		{
+			return eventId == KnownLogEvents.AutomationTool_UnacceptableWords;
+		}
+
+		/// <inheritdoc/>
+		public override string GetSummary(IIssueFingerprint fingerprint, IssueSeverity severity)
+		{
+			return $"Unacceptable words in {StringUtils.FormatList(fingerprint.Keys.Select(x => x.Substring(x.LastIndexOf('/') + 1)).ToArray(), 3)}";
+		}
+
+		/// <inheritdoc/>
+		public override bool TryGetFingerprint(IJob job, INode node, IReadOnlyNodeAnnotations annotations, ILogEventData eventData, [NotNullWhen(true)] out NewIssueFingerprint? fingerprint)
+		{
+			if (!IsMatchingEventId(eventData.EventId))
+			{
+				fingerprint = null;
+				return false;
+			}
+
+			List<string> newFileNames = new List<string>();
+			GetSourceFiles(eventData, newFileNames);
+
+			fingerprint = new NewIssueFingerprint(Type, newFileNames, null, null);
+			return true;
+		}
+	}
+}
diff --git a/Engine/Source/Programs/Shared/EpicGames.Core/KnownLogEvents.cs b/Engine/Source/Programs/Shared/EpicGames.Core/KnownLogEvents.cs
--- a/Engine/Source/Programs/Shared/EpicGames.Core/KnownLogEvents.cs
+++ b/Engine/Source/Programs/Shared/EpicGames.Core/KnownLogEvents.cs
@@ -105,6 +105,11 @@ namespace EpicGames.Core
 		/// </summary>
 		public static EventId AutomationTool_PerforceCase { get; } = new EventId(405);
 
+		/// <summary>
+		/// UAT: Unacceptable words
+		/// </summary>
+		public static EventId AutomationTool_UnacceptableWords { get; } = new EventId(406);
+
 		/// <summary>
 		/// MSBuild: Generic error
 		/// </summary>
