Horde: Add a systemic error code for missing XGEControlWorker, and generalize issue handler to match any event id within the range 700-799.

diff --git a/Engine/Source/Programs/AutomationTool/AutomationUtils/Matchers/SystemicEventMatcher.cs b/Engine/Source/Programs/AutomationTool/AutomationUtils/Matchers/SystemicEventMatcher.cs
--- a/Engine/Source/Programs/AutomationTool/AutomationUtils/Matchers/SystemicEventMatcher.cs
+++ b/Engine/Source/Programs/AutomationTool/AutomationUtils/Matchers/SystemicEventMatcher.cs
@@ -36,6 +36,11 @@ namespace AutomationUtils.Matchers
 			{
 				return new LogEventBuilder(cursor).ToMatch(LogEventPriority.Low, LogLevel.Information, KnownLogEvents.Systemic_RoboMergeGateLocked);
 			}
+			if (cursor.Contains("LogXGEController: Warning: XGEControlWorker.exe does not exist"))
+			{
+				return new LogEventBuilder(cursor).ToMatch(LogEventPriority.High, LogLevel.Information, KnownLogEvents.Systemic_MissingXgeControlWorker);
+			}
+
 			return null;
 		}
 	}
diff --git a/Engine/Source/Programs/Horde/Horde.Build/Issues/Impl/Handlers/SystemicIssueHandler.cs b/Engine/Source/Programs/Horde/Horde.Build/Issues/Impl/Handlers/SystemicIssueHandler.cs
--- a/Engine/Source/Programs/Horde/Horde.Build/Issues/Impl/Handlers/SystemicIssueHandler.cs
+++ b/Engine/Source/Programs/Horde/Horde.Build/Issues/Impl/Handlers/SystemicIssueHandler.cs
@@ -26,19 +26,16 @@ namespace Horde.Build.IssueHandlers.Impl
 		/// <summary>
 		///  Known systemic errors
 		/// </summary>
-		static readonly HashSet<EventId?> s_knownSystemic = new HashSet<EventId?> {  KnownLogEvents.Systemic, KnownLogEvents.Systemic_Xge, KnownLogEvents.Systemic_Xge_Standalone, 
-																		KnownLogEvents.Systemic_Xge_ServiceNotRunning, KnownLogEvents.Systemic_Xge_BuildFailed, 
-																		KnownLogEvents.Systemic_SlowDDC, KnownLogEvents.Systemic_Horde, KnownLogEvents.Systemic_Horde_ArtifactUpload,
-																		KnownLogEvents.Horde, KnownLogEvents.Horde_InvalidPreflight};
+		static readonly HashSet<EventId> s_knownSystemic = new HashSet<EventId> { KnownLogEvents.Horde, KnownLogEvents.Horde_InvalidPreflight };
 
 		/// <summary>
 		/// Determines if the given event id matches
 		/// </summary>
 		/// <param name="eventId">The event id to compare</param>
 		/// <returns>True if the given event id matches</returns>
-		public static bool IsMatchingEventId(EventId? eventId)
+		public static bool IsMatchingEventId(EventId eventId)
 		{
-			return s_knownSystemic.Contains(eventId);
+			return s_knownSystemic.Contains(eventId) || (eventId.Id >= KnownLogEvents.Systemic.Id && eventId.Id <= KnownLogEvents.Systemic_Max.Id);
 		}
 
 		/// <inheritdoc/>
@@ -70,7 +67,7 @@ namespace Horde.Build.IssueHandlers.Impl
 				}
 			}
 
-			if (!IsMatchingEventId(eventData.EventId))
+			if (eventData.EventId == null || !IsMatchingEventId(eventData.EventId.Value))
 			{				
 				return false;
 			}
diff --git a/Engine/Source/Programs/Shared/EpicGames.Core/KnownLogEvents.cs b/Engine/Source/Programs/Shared/EpicGames.Core/KnownLogEvents.cs
--- a/Engine/Source/Programs/Shared/EpicGames.Core/KnownLogEvents.cs
+++ b/Engine/Source/Programs/Shared/EpicGames.Core/KnownLogEvents.cs
@@ -185,8 +185,21 @@ namespace EpicGames.Core
 		/// </summary>
 		public static EventId Systemic_MSBuild { get; } = new EventId(750);
 
+		/// <summary>
+		/// Robomerge gate is locked
+		/// </summary>
 		public static EventId Systemic_RoboMergeGateLocked { get; } = new EventId(760);
 
+		/// <summary>
+		/// XGEControlWorker is missing
+		/// </summary>
+		public static EventId Systemic_MissingXgeControlWorker { get; } = new EventId(761);
+
+		/// <summary>
+		/// Maximum code for systemic events. Add new events in the 700-799 range.
+		/// </summary>
+		public static EventId Systemic_Max { get; } = new EventId(799);
+
 		/// <summary>
 		/// Horde error codes
 		/// </summary>
