Horde: Ignore *.manifest files when matching localization issues. These are not modified by users.

diff --git a/Engine/Source/Programs/Horde/Horde.Build/Issues/Impl/Handlers/LocalizationIssueHandler.cs b/Engine/Source/Programs/Horde/Horde.Build/Issues/Impl/Handlers/LocalizationIssueHandler.cs
--- a/Engine/Source/Programs/Horde/Horde.Build/Issues/Impl/Handlers/LocalizationIssueHandler.cs
+++ b/Engine/Source/Programs/Horde/Horde.Build/Issues/Impl/Handlers/LocalizationIssueHandler.cs
@@ -51,9 +51,12 @@ namespace Horde.Build.IssueHandlers.Impl
 				string? relativePath;
 				if (line.Data.TryGetNestedProperty("properties.file.relativePath", out relativePath) || line.Data.TryGetNestedProperty("properties.file", out relativePath))
 				{
-					int endIdx = relativePath.LastIndexOfAny(new char[] { '/', '\\' }) + 1;
-					string fileName = relativePath.Substring(endIdx);
-					sourceFiles.Add(fileName);
+					if (!relativePath.EndsWith(".manifest", StringComparison.OrdinalIgnoreCase))
+					{
+						int endIdx = relativePath.LastIndexOfAny(new char[] { '/', '\\' }) + 1;
+						string fileName = relativePath.Substring(endIdx);
+						sourceFiles.Add(fileName);
+					}
 				}
 			}
 		}
