Horde: Treat pre-formatted multi-line errors as a single issue, rather than parsing into multiple issues.

diff --git a/Engine/Source/Programs/Horde/Horde.Agent.Tests/LogParserTests.cs b/Engine/Source/Programs/Horde/Horde.Agent.Tests/LogParserTests.cs
--- a/Engine/Source/Programs/Horde/Horde.Agent.Tests/LogParserTests.cs
+++ b/Engine/Source/Programs/Horde/Horde.Agent.Tests/LogParserTests.cs
@@ -80,8 +80,10 @@ namespace Horde.Agent.Tests
 
 		static string[] SplitMultiLineMessage(string message)
 		{
+			JsonLogEvent jsonLogEvent;
+			Assert.IsTrue(JsonLogEvent.TryParse(Encoding.UTF8.GetBytes(message), out jsonLogEvent));
 			ArrayBufferWriter<byte> writer = new ArrayBufferWriter<byte>();
-			int count = JsonRpcLogger.WriteEvent(Encoding.UTF8.GetBytes(message), writer);
+			int count = JsonRpcLogger.WriteEvent(jsonLogEvent, writer);
 			string[] result = Encoding.UTF8.GetString(writer.WrittenSpan).Split('\n', StringSplitOptions.RemoveEmptyEntries);
 			Assert.AreEqual(count, result.Length);
 			return result;
diff --git a/Engine/Source/Programs/Horde/Horde.Agent/Utility/JsonRpcLogger.cs b/Engine/Source/Programs/Horde/Horde.Agent/Utility/JsonRpcLogger.cs
--- a/Engine/Source/Programs/Horde/Horde.Agent/Utility/JsonRpcLogger.cs
+++ b/Engine/Source/Programs/Horde/Horde.Agent/Utility/JsonRpcLogger.cs
@@ -184,36 +184,6 @@ namespace Horde.Agent.Parser
 			}
 		}
 
-		static int GetEventLineCount(ReadOnlySpan<byte> line)
-		{
-			Utf8JsonReader reader = new Utf8JsonReader(line);
-			if(!reader.Read() || reader.TokenType != JsonTokenType.StartObject)
-			{
-				return 0;
-			}
-
-			int lineCount = 1;
-			while(reader.Read() && reader.TokenType == JsonTokenType.PropertyName)
-			{
-				ReadOnlySpan<byte> propertyName = reader.ValueSpan;
-				if (!reader.Read())
-				{
-					break;
-				}
-				else if(propertyName.SequenceEqual(LogEventPropertyName.Line) && reader.TryGetInt32(out int lineIndex) && lineIndex > 0)
-				{
-					return 0;
-				}
-				else if(propertyName.SequenceEqual(LogEventPropertyName.LineCount) && reader.TryGetInt32(out int lineCountValue))
-				{
-					lineCount = lineCountValue;
-				}
-				reader.Skip();
-			}
-
-			return lineCount;
-		}
-
 		/// <summary>
 		/// Stops the log writer's background task
 		/// </summary>
@@ -275,17 +245,18 @@ namespace Horde.Agent.Parser
 					JsonLogEvent jsonLogEvent;
 					if (_dataChannel.Reader.TryRead(out jsonLogEvent))
 					{
-						int lineCount = WriteEvent(jsonLogEvent.Data.Span, writer);
-
-						if (jsonLogEvent.Level == LogLevel.Warning && ++numWarnings <= MaxWarnings)
-						{
-							AddEvent(jsonLogEvent.Data.Span, lineIndex, lineCount, EventSeverity.Warning, events);
-						}
-						else if ((jsonLogEvent.Level == LogLevel.Error || jsonLogEvent.Level == LogLevel.Critical) && ++numErrors <= MaxErrors)
+						int lineCount = WriteEvent(jsonLogEvent, writer);
+						if (jsonLogEvent.LineIndex == 0)
 						{
-							AddEvent(jsonLogEvent.Data.Span, lineIndex, lineCount, EventSeverity.Error, events);
+							if (jsonLogEvent.Level == LogLevel.Warning && ++numWarnings <= MaxWarnings)
+							{
+								AddEvent(jsonLogEvent.Data.Span, lineIndex, lineCount, EventSeverity.Warning, events);
+							}
+							else if ((jsonLogEvent.Level == LogLevel.Error || jsonLogEvent.Level == LogLevel.Critical) && ++numErrors <= MaxErrors)
+							{
+								AddEvent(jsonLogEvent.Data.Span, lineIndex, lineCount, EventSeverity.Error, events);
+							}
 						}
-
 						lineIndex += lineCount;
 					}
 					else
@@ -368,16 +339,17 @@ namespace Horde.Agent.Parser
 		static readonly Utf8String s_newline = "\n";
 		static readonly Utf8String s_escapedNewline = "\\n";
 
-		public static int WriteEvent(ReadOnlySpan<byte> logEvent, IBufferWriter<byte> writer)
+		public static int WriteEvent(JsonLogEvent jsonLogEvent, IBufferWriter<byte> writer)
 		{
-			if (logEvent.IndexOf(s_escapedNewline) == -1)
+			ReadOnlySpan<byte> span = jsonLogEvent.Data.Span;
+			if (jsonLogEvent.LineCount > 0 || span.IndexOf(s_escapedNewline) == -1)
 			{
-				writer.Write(logEvent);
+				writer.Write(span);
 				writer.Write(s_newline);
-				return 1;
+				return jsonLogEvent.LineCount;
 			}
 
-			JsonObject obj = (JsonObject)JsonNode.Parse(logEvent)!;
+			JsonObject obj = (JsonObject)JsonNode.Parse(span)!;
 
 			string format = (obj["format"] as JsonValue)?.ToString() ?? String.Empty;
 			IEnumerable<KeyValuePair<string, object?>> propertyValueList = Enumerable.Empty<KeyValuePair<string, object?>>();
diff --git a/Engine/Source/Programs/Horde/Horde.Agent/Utility/PerforceLogger.cs b/Engine/Source/Programs/Horde/Horde.Agent/Utility/PerforceLogger.cs
--- a/Engine/Source/Programs/Horde/Horde.Agent/Utility/PerforceLogger.cs
+++ b/Engine/Source/Programs/Horde/Horde.Agent/Utility/PerforceLogger.cs
@@ -67,7 +67,7 @@ namespace Horde.Agent.Utility
 		public void Log<TState>(LogLevel logLevel, EventId eventId, TState state, Exception exception, Func<TState, Exception?, string> formatter)
 		{
 			JsonLogEvent logEvent = JsonLogEvent.FromLoggerState(logLevel, eventId, state, exception, formatter);
-			logEvent = new JsonLogEvent(logLevel, eventId, Annotate(logEvent.Data));
+			logEvent = new JsonLogEvent(logLevel, eventId, logEvent.LineIndex, logEvent.LineCount, Annotate(logEvent.Data));
 			_inner.Log(logLevel, eventId, logEvent, null, JsonLogEvent.Format);
 		}
 
diff --git a/Engine/Source/Programs/Shared/EpicGames.Core/JsonLogEvent.cs b/Engine/Source/Programs/Shared/EpicGames.Core/JsonLogEvent.cs
--- a/Engine/Source/Programs/Shared/EpicGames.Core/JsonLogEvent.cs
+++ b/Engine/Source/Programs/Shared/EpicGames.Core/JsonLogEvent.cs
@@ -1,6 +1,7 @@
 // Copyright Epic Games, Inc. All Rights Reserved.
 
 using System;
+using System.Diagnostics.CodeAnalysis;
 using System.Text;
 using System.Text.Json;
 using Microsoft.Extensions.Logging;
@@ -22,6 +23,16 @@ namespace EpicGames.Core
 		/// </summary>
 		public EventId EventId { get; }
 
+		/// <summary>
+		/// Index of this line
+		/// </summary>
+		public int LineIndex { get; }
+
+		/// <summary>
+		/// Number of lines in a multi-line message
+		/// </summary>
+		public int LineCount { get; }
+
 		/// <summary>
 		/// The utf-8 encoded JSON event
 		/// </summary>
@@ -30,10 +41,12 @@ namespace EpicGames.Core
 		/// <summary>
 		/// Constructor
 		/// </summary>
-		public JsonLogEvent(LogLevel level, EventId eventId, ReadOnlyMemory<byte> data)
+		public JsonLogEvent(LogLevel level, EventId eventId, int lineIndex, int lineCount, ReadOnlyMemory<byte> data)
 		{
 			Level = level;
 			EventId = eventId;
+			LineIndex = lineIndex;
+			LineCount = lineCount;
 			Data = data;
 		}
 
@@ -54,7 +67,22 @@ namespace EpicGames.Core
 				logEvent = LogEvent.FromState(logLevel, eventId, state, exception, formatter);
 			}
 
-			return new JsonLogEvent(logLevel, eventId, logEvent.ToJsonBytes());
+			return new JsonLogEvent(logLevel, eventId, 0, 1, logEvent.ToJsonBytes());
+		}
+
+		/// <summary>
+		/// Parse an event from the given data
+		/// </summary>
+		/// <param name="data"></param>
+		/// <returns></returns>
+		public static JsonLogEvent Parse(ReadOnlyMemory<byte> data)
+		{
+			JsonLogEvent logEvent;
+			if (!TryParse(data, out logEvent))
+			{
+				throw new InvalidOperationException("Cannot parse string");
+			}
+			return logEvent;
 		}
 
 		/// <summary>
@@ -69,6 +97,8 @@ namespace EpicGames.Core
 			{
 				LogLevel level = LogLevel.None;
 				int eventId = 0;
+				int lineIndex = 0;
+				int lineCount = 1;
 
 				Utf8JsonReader reader = new Utf8JsonReader(data.Span);
 				if (reader.Read() && reader.TokenType == JsonTokenType.StartObject)
@@ -88,13 +118,21 @@ namespace EpicGames.Core
 						{
 							eventId = reader.GetInt32();
 						}
+						else if (propertyName.SequenceEqual(LogEventPropertyName.Line) && reader.TokenType == JsonTokenType.Number)
+						{
+							reader.TryGetInt32(out lineIndex);
+						}
+						else if (propertyName.SequenceEqual(LogEventPropertyName.LineCount) && reader.TokenType == JsonTokenType.Number)
+						{
+							reader.TryGetInt32(out lineCount);
+						}
 						reader.Skip();
 					}
 				}
 
 				if (reader.TokenType == JsonTokenType.EndObject && level != LogLevel.None && reader.BytesConsumed == data.Length)
 				{
-					logEvent = new JsonLogEvent(level, new EventId(eventId), data.ToArray());
+					logEvent = new JsonLogEvent(level, new EventId(eventId), lineIndex, lineCount, data.ToArray());
 					return true;
 				}
 			}
