Horde: Device model constraints for v1 reservation endpoint
#rnx

diff --git a/Engine/Source/Programs/Horde/Horde.Build/Api/Devices.cs b/Engine/Source/Programs/Horde/Horde.Build/Api/Devices.cs
--- a/Engine/Source/Programs/Horde/Horde.Build/Api/Devices.cs
+++ b/Engine/Source/Programs/Horde/Horde.Build/Api/Devices.cs
@@ -173,9 +173,9 @@ namespace Horde.Build.Api
 		{
 			Id = id;
 			Name = name;
-            PoolType = poolType;
+			PoolType = poolType;
 			WriteAccess = writeAccess;
-        }
+		}
 	}
 
 	// Devices 
@@ -248,22 +248,22 @@ namespace Horde.Build.Api
 		/// <summary>
 		/// The job id which utilized device
 		/// </summary>
-		public string? JobId { get; set;}
+		public string? JobId { get; set; }
 
 		/// <summary>
 		/// The job's step id
 		/// </summary>
-		public string? StepId { get; set;}
+		public string? StepId { get; set; }
 
 		/// <summary>
 		/// The time device was reserved
 		/// </summary>
-		public DateTime ReservationStartUtc { get; set;}
+		public DateTime ReservationStartUtc { get; set; }
 
 		/// <summary>
 		/// The time device was freed
 		/// </summary>
-		public DateTime? ReservationFinishUtc { get; set;}
+		public DateTime? ReservationFinishUtc { get; set; }
 
 		/// <summary>
 		/// Constructor
@@ -351,7 +351,7 @@ namespace Horde.Build.Api
 		/// <summary>
 		///  Device Utilization data
 		/// </summary>
-		public List<GetDeviceUtilizationResponse>? Utilization { get; set;}
+		public List<GetDeviceUtilizationResponse>? Utilization { get; set; }
 
 		/// <summary>
 		/// Device response constructor
@@ -365,14 +365,14 @@ namespace Horde.Build.Api
 			Enabled = enabled;
 			Address = address;
 			ModelId = modelId;
-            ModifiedByUser = modifiedByUser;
-            Notes = notes;
+			ModifiedByUser = modifiedByUser;
+			Notes = notes;
 			ProblemTime = problemTime;
 			MaintenanceTime = maintenanceTime;
 			Utilization = utilization?.Select(u => new GetDeviceUtilizationResponse(u)).ToList();
-            CheckedOutByUserId = checkedOutByUser;
-            CheckOutTime = checkOutTime;
-        }
+			CheckedOutByUserId = checkedOutByUser;
+			CheckOutTime = checkOutTime;
+		}
 	}
 
 	/// <summary>
@@ -427,10 +427,10 @@ namespace Horde.Build.Api
 	public class CheckoutDeviceRequest
 	{
 		/// <summary>
-        /// Whether to checkout or in the device
-        /// </summary>
-        public bool Checkout { get; set; }
-    }
+		/// Whether to checkout or in the device
+		/// </summary>
+		public bool Checkout { get; set; }
+	}
 
 	/// <summary>
 	/// Device reservation request object
@@ -611,6 +611,12 @@ namespace Horde.Build.Api
 		[Required]
 		public string[] DevicePerfSpecs { get; set; } = null!;
 
+		/// <summary>
+		/// The corresponding perf specs of the reserved devices
+		/// </summary>
+		[Required]
+		public string[] DeviceModels { get; set; } = null!;
+
 		/// <summary>
 		/// The host name of the machine making the reservation
 		/// </summary>
@@ -675,6 +681,12 @@ namespace Horde.Build.Api
 		[Required]
 		public string PerfSpec { get; set; } = null!;
 
+		/// <summary>
+		/// The device model information
+		/// </summary>
+		[Required]
+		public string Model { get; set; } = null!;
+
 		/// <summary>
 		/// The available start time which is parsed client side
 		/// </summary>
diff --git a/Engine/Source/Programs/Horde/Horde.Build/Controllers/DevicesController.cs b/Engine/Source/Programs/Horde/Horde.Build/Controllers/DevicesController.cs
--- a/Engine/Source/Programs/Horde/Horde.Build/Controllers/DevicesController.cs
+++ b/Engine/Source/Programs/Horde/Horde.Build/Controllers/DevicesController.cs
@@ -244,7 +244,7 @@ namespace Horde.Build.Controllers
 			string? name = update.Name?.Trim();
 			string? address = update.Address?.Trim();
 
-            await _deviceService.UpdateDeviceAsync(deviceIdValue, poolIdValue, name, address, modelIdValue, update.Notes, update.Enabled, update.Problem, update.Maintenance, internalUser.Id);
+			await _deviceService.UpdateDeviceAsync(deviceIdValue, poolIdValue, name, address, modelIdValue, update.Notes, update.Enabled, update.Problem, update.Maintenance, internalUser.Id);
 
 			return Ok();
 		}
@@ -288,9 +288,9 @@ namespace Horde.Build.Controllers
 					return BadRequest($"Already checked out by user {device.CheckedOutByUser}");
 				}
 
-                await _deviceService.CheckoutDeviceAsync(deviceIdValue, internalUser.Id);
+				await _deviceService.CheckoutDeviceAsync(deviceIdValue, internalUser.Id);
 
-            }
+			}
 			else
 			{
 				await _deviceService.CheckoutDeviceAsync(deviceIdValue, null);
@@ -456,7 +456,7 @@ namespace Horde.Build.Controllers
 			}
 
 			await _deviceService.UpdatePoolAsync(new DevicePoolId(request.Id), projectIds);
-			
+
 			return Ok();
 		}
 
@@ -484,7 +484,7 @@ namespace Horde.Build.Controllers
 				{
 					continue;
 				}
-				
+
 				responses.Add(new GetDevicePoolResponse(pool.Id.ToString(), pool.Name, pool.PoolType, auth.Write));
 			}
 
@@ -550,7 +550,7 @@ namespace Horde.Build.Controllers
 							return BadRequest($"Unknown model {model} for platform {deviceRequest.PlatformId} on device reservation request");
 						}
 					}
-				}					
+				}
 
 				requestedDevices.Add(new DeviceRequestData(platformIdValue, platformIdValue.ToString(), deviceRequest.IncludeModels, deviceRequest.ExcludeModels));
 			}
@@ -726,23 +726,23 @@ namespace Horde.Build.Controllers
 
 			List<DeviceRequestData> requestedDevices = new List<DeviceRequestData>();
 
+			HashSet<string> legacyPerfSpecs = new HashSet<string> { "Unspecified", "Minimum", "Recommended", "High" };
+
 			List<string> perfSpecs = new List<string>();
 
 			foreach (string deviceType in request.DeviceTypes)
 			{
 
 				string platformName = deviceType;
-				string perfSpecName = "Unspecified";
+				string? constraint = null;
 
 				if (deviceType.Contains(':', StringComparison.Ordinal))
 				{
 					string[] tokens = deviceType.Split(":");
 					platformName = tokens[0];
-					perfSpecName = tokens[1];
+					constraint = tokens[1];
 				}
 
-				perfSpecs.Add(perfSpecName);
-
 				DevicePlatformId platformId;
 
 				DevicePlatformMapV1 mapV1 = await _deviceService.GetPlatformMapV1();
@@ -768,22 +768,53 @@ namespace Horde.Build.Controllers
 				List<string> includeModels = new List<string>();
 				List<string> excludeModels = new List<string>();
 
-				if (perfSpecName == "High")
+				if (constraint == null)
+				{
+					perfSpecs.Add("Unspecified");
+				}
+				else if (legacyPerfSpecs.Contains(constraint))
 				{
-					string? model = null;
-					if (mapV1.PerfSpecHighMap.TryGetValue(platformId, out model))
+					perfSpecs.Add(constraint);
+
+					if (constraint == "High")
 					{
-						includeModels.Add(model);
+						string? model = null;
+						if (mapV1.PerfSpecHighMap.TryGetValue(platformId, out model))
+						{
+							includeModels.Add(model);
+						}
+					}
+					else if (constraint == "Minimum" || constraint == "Recommended")
+					{
+						string? model = null;
+						if (mapV1.PerfSpecHighMap.TryGetValue(platformId, out model))
+						{
+							excludeModels.Add(model);
+						}
 					}
 				}
-
-				if (perfSpecName == "Minimum" || perfSpecName == "Recommended")
+				else
 				{
-					string? model = null;
-					if (mapV1.PerfSpecHighMap.TryGetValue(platformId, out model))
+					string? model = platform.Models?.FirstOrDefault(x => x == constraint);
+					if (model == null)
+					{
+						return NotFound($"Platform {platform.Id} has no model {model}");
+					}
+
+					string modelPerfSpec = "Minimum";
+					string? specModel = null;
+
+					if (mapV1.PerfSpecHighMap.TryGetValue(platform.Id, out specModel))
 					{
-						excludeModels.Add(model);
+						if (model == specModel)
+						{
+							modelPerfSpec = "High";
+						}
 					}
+
+					perfSpecs.Add(modelPerfSpec);
+
+					includeModels.Add(constraint);
 				}
 
 				requestedDevices.Add(new DeviceRequestData(platformId, platformName, includeModels, excludeModels));
@@ -803,6 +834,7 @@ namespace Horde.Build.Controllers
 			response.Guid = reservation.LegacyGuid;
 			response.DeviceNames = devices.Select(x => x.Name).ToArray();
 			response.DevicePerfSpecs = perfSpecs.ToArray();
+			response.DeviceModels = devices.Select(x => x.ModelId ?? "Base").ToArray();
 			response.HostName = request.Hostname;
 			response.StartDateTime = reservation.CreateTimeUtc.ToString("O", System.Globalization.CultureInfo.InvariantCulture);
 			response.Duration = $"{request.Duration}";
@@ -905,14 +937,14 @@ namespace Horde.Build.Controllers
 			DevicePlatformMapV1 mapV1 = await _deviceService.GetPlatformMapV1();
 
 			if (String.IsNullOrEmpty(platformName))
-			{				
+			{
 				if (!mapV1.PlatformReverseMap.TryGetValue(device.PlatformId, out platformName))
 				{
 					return BadRequest($"Unable to map platform for {deviceName} : {device.PlatformId}");
 				}
 			}
 
-			if (String.IsNullOrEmpty(platformName)) 
+			if (String.IsNullOrEmpty(platformName))
 			{
 				return BadRequest($"Unable to get platform for {deviceName} from reservation or mapping : {device.PlatformId}");
 			}
@@ -925,8 +957,8 @@ namespace Horde.Build.Controllers
 			response.AvailableStartTime = "00:00:00";
 			response.AvailableEndTime = "00:00:00";
 			response.Enabled = true;
+			response.Model = device.ModelId ?? "Base";
 			response.DeviceData = "";
-
 			response.PerfSpec = "Minimum";
 
 			if (device.ModelId != null)
