Horde: Fix multi-line parsing test.

diff --git a/Engine/Source/Programs/Horde/Horde.Agent.Tests/LogParserTests.cs b/Engine/Source/Programs/Horde/Horde.Agent.Tests/LogParserTests.cs
--- a/Engine/Source/Programs/Horde/Horde.Agent.Tests/LogParserTests.cs
+++ b/Engine/Source/Programs/Horde/Horde.Agent.Tests/LogParserTests.cs
@@ -92,16 +92,16 @@ namespace Horde.Agent.Tests
 		[TestMethod]
 		public void MultiLineTest()
 		{
-			const string msg = @"{""message"":""ignored"",""format"":""This\nis a\nmulti-{Line}-end\nlog message {Var1} {Var2}"",""properties"":{""Line"":""line\nsplit\nin\nfour"",""Var1"":123,""Var2"":{""$type"":""SourceFile"",""$text"":""D:\\build\\\u002B\u002BUE5\\Sync\\GenerateProjectFiles.bat"",""relativePath"":""GenerateProjectFiles.bat"",""depotPath"":""//UE5/Main/GenerateProjectFiles.bat@20392842""}}}";
+			const string msg = @"{""level"":""Information"",""message"":""ignored"",""format"":""This\nis a\nmulti-{Line}-end\nlog message {Var1} {Var2}"",""properties"":{""Line"":""line\nsplit\nin\nfour"",""Var1"":123,""Var2"":{""$type"":""SourceFile"",""$text"":""D:\\build\\\u002B\u002BUE5\\Sync\\GenerateProjectFiles.bat"",""relativePath"":""GenerateProjectFiles.bat"",""depotPath"":""//UE5/Main/GenerateProjectFiles.bat@20392842""}}}";
 
 			string[] result = SplitMultiLineMessage(msg);
-			Assert.AreEqual(@"{""message"":""This"",""format"":""This"",""properties"":{""Line"":""line\nsplit\nin\nfour"",""Var1"":123,""Var2"":{""$type"":""SourceFile"",""$text"":""D:\\build\\\u002B\u002BUE5\\Sync\\GenerateProjectFiles.bat"",""relativePath"":""GenerateProjectFiles.bat"",""depotPath"":""//UE5/Main/GenerateProjectFiles.bat@20392842""},""Line$0"":""line"",""Line$1"":""split"",""Line$2"":""in"",""Line$3"":""four""},""line"":0,""lineCount"":7}", result[0]);
-			Assert.AreEqual(@"{""message"":""is a"",""format"":""is a"",""properties"":{""Line"":""line\nsplit\nin\nfour"",""Var1"":123,""Var2"":{""$type"":""SourceFile"",""$text"":""D:\\build\\\u002B\u002BUE5\\Sync\\GenerateProjectFiles.bat"",""relativePath"":""GenerateProjectFiles.bat"",""depotPath"":""//UE5/Main/GenerateProjectFiles.bat@20392842""},""Line$0"":""line"",""Line$1"":""split"",""Line$2"":""in"",""Line$3"":""four""},""line"":1,""lineCount"":7}", result[1]);
-			Assert.AreEqual(@"{""message"":""multi-line"",""format"":""multi-{Line$0}"",""properties"":{""Line"":""line\nsplit\nin\nfour"",""Var1"":123,""Var2"":{""$type"":""SourceFile"",""$text"":""D:\\build\\\u002B\u002BUE5\\Sync\\GenerateProjectFiles.bat"",""relativePath"":""GenerateProjectFiles.bat"",""depotPath"":""//UE5/Main/GenerateProjectFiles.bat@20392842""},""Line$0"":""line"",""Line$1"":""split"",""Line$2"":""in"",""Line$3"":""four""},""line"":2,""lineCount"":7}", result[2]);
-			Assert.AreEqual(@"{""message"":""split"",""format"":""{Line$1}"",""properties"":{""Line"":""line\nsplit\nin\nfour"",""Var1"":123,""Var2"":{""$type"":""SourceFile"",""$text"":""D:\\build\\\u002B\u002BUE5\\Sync\\GenerateProjectFiles.bat"",""relativePath"":""GenerateProjectFiles.bat"",""depotPath"":""//UE5/Main/GenerateProjectFiles.bat@20392842""},""Line$0"":""line"",""Line$1"":""split"",""Line$2"":""in"",""Line$3"":""four""},""line"":3,""lineCount"":7}", result[3]);
-			Assert.AreEqual(@"{""message"":""in"",""format"":""{Line$2}"",""properties"":{""Line"":""line\nsplit\nin\nfour"",""Var1"":123,""Var2"":{""$type"":""SourceFile"",""$text"":""D:\\build\\\u002B\u002BUE5\\Sync\\GenerateProjectFiles.bat"",""relativePath"":""GenerateProjectFiles.bat"",""depotPath"":""//UE5/Main/GenerateProjectFiles.bat@20392842""},""Line$0"":""line"",""Line$1"":""split"",""Line$2"":""in"",""Line$3"":""four""},""line"":4,""lineCount"":7}", result[4]);
-			Assert.AreEqual(@"{""message"":""four-end"",""format"":""{Line$3}-end"",""properties"":{""Line"":""line\nsplit\nin\nfour"",""Var1"":123,""Var2"":{""$type"":""SourceFile"",""$text"":""D:\\build\\\u002B\u002BUE5\\Sync\\GenerateProjectFiles.bat"",""relativePath"":""GenerateProjectFiles.bat"",""depotPath"":""//UE5/Main/GenerateProjectFiles.bat@20392842""},""Line$0"":""line"",""Line$1"":""split"",""Line$2"":""in"",""Line$3"":""four""},""line"":5,""lineCount"":7}", result[5]);
-			Assert.AreEqual(@"{""message"":""log message 123 D:\\build\\\u002B\u002BUE5\\Sync\\GenerateProjectFiles.bat"",""format"":""log message {Var1} {Var2}"",""properties"":{""Line"":""line\nsplit\nin\nfour"",""Var1"":123,""Var2"":{""$type"":""SourceFile"",""$text"":""D:\\build\\\u002B\u002BUE5\\Sync\\GenerateProjectFiles.bat"",""relativePath"":""GenerateProjectFiles.bat"",""depotPath"":""//UE5/Main/GenerateProjectFiles.bat@20392842""},""Line$0"":""line"",""Line$1"":""split"",""Line$2"":""in"",""Line$3"":""four""},""line"":6,""lineCount"":7}", result[6]);
+			Assert.AreEqual(@"{""level"":""Information"",""message"":""This"",""format"":""This"",""properties"":{""Line"":""line\nsplit\nin\nfour"",""Var1"":123,""Var2"":{""$type"":""SourceFile"",""$text"":""D:\\build\\\u002B\u002BUE5\\Sync\\GenerateProjectFiles.bat"",""relativePath"":""GenerateProjectFiles.bat"",""depotPath"":""//UE5/Main/GenerateProjectFiles.bat@20392842""},""Line$0"":""line"",""Line$1"":""split"",""Line$2"":""in"",""Line$3"":""four""},""line"":0,""lineCount"":7}", result[0]);
+			Assert.AreEqual(@"{""level"":""Information"",""message"":""is a"",""format"":""is a"",""properties"":{""Line"":""line\nsplit\nin\nfour"",""Var1"":123,""Var2"":{""$type"":""SourceFile"",""$text"":""D:\\build\\\u002B\u002BUE5\\Sync\\GenerateProjectFiles.bat"",""relativePath"":""GenerateProjectFiles.bat"",""depotPath"":""//UE5/Main/GenerateProjectFiles.bat@20392842""},""Line$0"":""line"",""Line$1"":""split"",""Line$2"":""in"",""Line$3"":""four""},""line"":1,""lineCount"":7}", result[1]);
+			Assert.AreEqual(@"{""level"":""Information"",""message"":""multi-line"",""format"":""multi-{Line$0}"",""properties"":{""Line"":""line\nsplit\nin\nfour"",""Var1"":123,""Var2"":{""$type"":""SourceFile"",""$text"":""D:\\build\\\u002B\u002BUE5\\Sync\\GenerateProjectFiles.bat"",""relativePath"":""GenerateProjectFiles.bat"",""depotPath"":""//UE5/Main/GenerateProjectFiles.bat@20392842""},""Line$0"":""line"",""Line$1"":""split"",""Line$2"":""in"",""Line$3"":""four""},""line"":2,""lineCount"":7}", result[2]);
+			Assert.AreEqual(@"{""level"":""Information"",""message"":""split"",""format"":""{Line$1}"",""properties"":{""Line"":""line\nsplit\nin\nfour"",""Var1"":123,""Var2"":{""$type"":""SourceFile"",""$text"":""D:\\build\\\u002B\u002BUE5\\Sync\\GenerateProjectFiles.bat"",""relativePath"":""GenerateProjectFiles.bat"",""depotPath"":""//UE5/Main/GenerateProjectFiles.bat@20392842""},""Line$0"":""line"",""Line$1"":""split"",""Line$2"":""in"",""Line$3"":""four""},""line"":3,""lineCount"":7}", result[3]);
+			Assert.AreEqual(@"{""level"":""Information"",""message"":""in"",""format"":""{Line$2}"",""properties"":{""Line"":""line\nsplit\nin\nfour"",""Var1"":123,""Var2"":{""$type"":""SourceFile"",""$text"":""D:\\build\\\u002B\u002BUE5\\Sync\\GenerateProjectFiles.bat"",""relativePath"":""GenerateProjectFiles.bat"",""depotPath"":""//UE5/Main/GenerateProjectFiles.bat@20392842""},""Line$0"":""line"",""Line$1"":""split"",""Line$2"":""in"",""Line$3"":""four""},""line"":4,""lineCount"":7}", result[4]);
+			Assert.AreEqual(@"{""level"":""Information"",""message"":""four-end"",""format"":""{Line$3}-end"",""properties"":{""Line"":""line\nsplit\nin\nfour"",""Var1"":123,""Var2"":{""$type"":""SourceFile"",""$text"":""D:\\build\\\u002B\u002BUE5\\Sync\\GenerateProjectFiles.bat"",""relativePath"":""GenerateProjectFiles.bat"",""depotPath"":""//UE5/Main/GenerateProjectFiles.bat@20392842""},""Line$0"":""line"",""Line$1"":""split"",""Line$2"":""in"",""Line$3"":""four""},""line"":5,""lineCount"":7}", result[5]);
+			Assert.AreEqual(@"{""level"":""Information"",""message"":""log message 123 D:\\build\\\u002B\u002BUE5\\Sync\\GenerateProjectFiles.bat"",""format"":""log message {Var1} {Var2}"",""properties"":{""Line"":""line\nsplit\nin\nfour"",""Var1"":123,""Var2"":{""$type"":""SourceFile"",""$text"":""D:\\build\\\u002B\u002BUE5\\Sync\\GenerateProjectFiles.bat"",""relativePath"":""GenerateProjectFiles.bat"",""depotPath"":""//UE5/Main/GenerateProjectFiles.bat@20392842""},""Line$0"":""line"",""Line$1"":""split"",""Line$2"":""in"",""Line$3"":""four""},""line"":6,""lineCount"":7}", result[6]);
 		}
 
 		[TestMethod]
diff --git a/Engine/Source/Programs/Horde/Horde.Agent/Utility/JsonRpcLogger.cs b/Engine/Source/Programs/Horde/Horde.Agent/Utility/JsonRpcLogger.cs
--- a/Engine/Source/Programs/Horde/Horde.Agent/Utility/JsonRpcLogger.cs
+++ b/Engine/Source/Programs/Horde/Horde.Agent/Utility/JsonRpcLogger.cs
@@ -342,7 +342,7 @@ namespace Horde.Agent.Parser
 		public static int WriteEvent(JsonLogEvent jsonLogEvent, IBufferWriter<byte> writer)
 		{
 			ReadOnlySpan<byte> span = jsonLogEvent.Data.Span;
-			if (jsonLogEvent.LineCount > 0 || span.IndexOf(s_escapedNewline) == -1)
+			if (jsonLogEvent.LineCount > 1 || span.IndexOf(s_escapedNewline) == -1)
 			{
 				writer.Write(span);
 				writer.Write(s_newline);
