Horde: Fix disparity between written line count and pre-rendered line count in logs.

diff --git a/Engine/Source/Programs/Horde/Horde.Agent/Utility/JsonRpcLogger.cs b/Engine/Source/Programs/Horde/Horde.Agent/Utility/JsonRpcLogger.cs
--- a/Engine/Source/Programs/Horde/Horde.Agent/Utility/JsonRpcLogger.cs
+++ b/Engine/Source/Programs/Horde/Horde.Agent/Utility/JsonRpcLogger.cs
@@ -250,11 +250,11 @@ namespace Horde.Agent.Parser
 						{
 							if (jsonLogEvent.Level == LogLevel.Warning && ++numWarnings <= MaxWarnings)
 							{
-								AddEvent(jsonLogEvent.Data.Span, lineIndex, lineCount, EventSeverity.Warning, events);
+								AddEvent(jsonLogEvent.Data.Span, lineIndex, Math.Max(lineCount, jsonLogEvent.LineCount), EventSeverity.Warning, events);
 							}
 							else if ((jsonLogEvent.Level == LogLevel.Error || jsonLogEvent.Level == LogLevel.Critical) && ++numErrors <= MaxErrors)
 							{
-								AddEvent(jsonLogEvent.Data.Span, lineIndex, lineCount, EventSeverity.Error, events);
+								AddEvent(jsonLogEvent.Data.Span, lineIndex, Math.Max(lineCount, jsonLogEvent.LineCount), EventSeverity.Error, events);
 							}
 						}
 						lineIndex += lineCount;
@@ -346,7 +346,7 @@ namespace Horde.Agent.Parser
 			{
 				writer.Write(span);
 				writer.Write(s_newline);
-				return jsonLogEvent.LineCount;
+				return 1;
 			}
 
 			JsonObject obj = (JsonObject)JsonNode.Parse(span)!;
