Horde: If we have to repair any log chunks, push the repaired chunks back into storage.

diff --git a/Engine/Source/Programs/Horde/Horde.Build/Logs/Impl/LogFileService.cs b/Engine/Source/Programs/Horde/Horde.Build/Logs/Impl/LogFileService.cs
--- a/Engine/Source/Programs/Horde/Horde.Build/Logs/Impl/LogFileService.cs
+++ b/Engine/Source/Programs/Horde/Horde.Build/Logs/Impl/LogFileService.cs
@@ -1335,17 +1335,17 @@ namespace Horde.Build.Services
 			if (chunkIdx + 1 < logFile.Chunks.Count)
 			{
 				ILogChunk nextChunk = logFile.Chunks[chunkIdx + 1];
-				chunkData = RepairChunkData(logFile, chunkIdx, chunkData, (int)(nextChunk.Offset - chunk.Offset), nextChunk.LineIndex - chunk.LineIndex);
+				chunkData = await RepairChunkDataAsync(logFile, chunkIdx, chunkData, (int)(nextChunk.Offset - chunk.Offset), nextChunk.LineIndex - chunk.LineIndex);
 			}
 			else
 			{
 				if (logFile.MaxLineIndex != null && chunk.Length != 0)
 				{
-					chunkData = RepairChunkData(logFile, chunkIdx, chunkData, chunk.Length, logFile.MaxLineIndex.Value - chunk.LineIndex);
+					chunkData = await RepairChunkDataAsync(logFile, chunkIdx, chunkData, chunk.Length, logFile.MaxLineIndex.Value - chunk.LineIndex);
 				}
 				else if(chunkData == null)
 				{
-					chunkData = RepairChunkData(logFile, chunkIdx, chunkData, 1024, 1);
+					chunkData = await RepairChunkDataAsync(logFile, chunkIdx, chunkData, 1024, 1);
 				}
 			}
 
@@ -1361,7 +1361,7 @@ namespace Horde.Build.Services
 		/// <param name="length">Expected length of the data</param>
 		/// <param name="lineCount">Expected number of lines in the data</param>
 		/// <returns>Repaired chunk data</returns>
-		LogChunkData RepairChunkData(ILogFile logFile, int chunkIdx, LogChunkData? chunkData, int length, int lineCount)
+		async Task<LogChunkData> RepairChunkDataAsync(ILogFile logFile, int chunkIdx, LogChunkData? chunkData, int length, int lineCount)
 		{
 			int currentLength = 0;
 			int currentLineCount = 0;
@@ -1387,6 +1387,15 @@ namespace Horde.Build.Services
 
 				ILogChunk chunk = logFile.Chunks[chunkIdx];
 				chunkData = new LogChunkData(chunk.Offset, chunk.LineIndex, subChunks);
+
+				try
+				{
+					await _storage.WriteChunkAsync(logFile.Id, chunk.Offset, chunkData);
+				}
+				catch (Exception ex)
+				{
+					_logger.LogWarning(ex, "Unable to put repaired log data for log {LogId} chunk {ChunkIdx}", logFile.Id, chunkIdx);
+				}
 			}
 			return chunkData;
 		}
