Horde: Automatically acknowledge an issue when marking it as resolved.

diff --git a/Engine/Source/Programs/Horde/Horde.Build/Issues/Impl/IssueCollection.cs b/Engine/Source/Programs/Horde/Horde.Build/Issues/Impl/IssueCollection.cs
--- a/Engine/Source/Programs/Horde/Horde.Build/Issues/Impl/IssueCollection.cs
+++ b/Engine/Source/Programs/Horde/Horde.Build/Issues/Impl/IssueCollection.cs
@@ -904,6 +904,11 @@ namespace Horde.Build.Collections.Impl
 			{
 				updates.Add(Builders<Issue>.Update.Set(x => x.ManuallyPromoted, newManuallyPromoted.Value));
 			}
+			if (newResolvedById != null)
+			{
+				newOwnerId ??= newResolvedById;
+				newAcknowledged ??= true;
+			}
 			if (newOwnerId != null)
 			{
 				if (newOwnerId.Value == UserId.Empty)
