Horde: Initial cut of agent pool view and editor
#skipci
#rnx

diff --git a/Engine/Source/Programs/Horde/HordeDashboard/src/backend/AgentStore.ts b/Engine/Source/Programs/Horde/HordeDashboard/src/backend/AgentStore.ts
--- a/Engine/Source/Programs/Horde/HordeDashboard/src/backend/AgentStore.ts
+++ b/Engine/Source/Programs/Horde/HordeDashboard/src/backend/AgentStore.ts
@@ -40,7 +40,7 @@ export class AgentStore {
     async update(slim = false): Promise<void> {
         return new Promise<void>((resolve, reject) => {
             const promises: any[] = [];
-            promises.push(backend.getAgents(this.modifiedAfterDate));
+            promises.push(backend.getAgents({modifiedAfter: this.modifiedAfterDate?.toISOString()}));
             if(!slim) {
                 promises.push(backend.getPools());
             }
diff --git a/Engine/Source/Programs/Horde/HordeDashboard/src/backend/Api.ts b/Engine/Source/Programs/Horde/HordeDashboard/src/backend/Api.ts
--- a/Engine/Source/Programs/Horde/HordeDashboard/src/backend/Api.ts
+++ b/Engine/Source/Programs/Horde/HordeDashboard/src/backend/Api.ts
@@ -73,8 +73,7 @@ export enum JobStepOutcome {
 }
 
 /// Systemic error codes for a job failing
-export enum JobStepError
-{
+export enum JobStepError {
 	/// No systemic error
 	None,
 	/// Step did not complete in the required amount of time
@@ -288,6 +287,11 @@ export type JobStreamQuery = {
 
 };
 
+export type AgentQuery = {
+	modifiedAfter?: string;
+	poolId?: string;
+}
+
 export type JobQuery = {
 	id?: string[],
 	filter?: string;
@@ -335,7 +339,7 @@ export type IssueQuery = {
 }
 
 export type IssueQueryV2 = {
-	id?: string[]; 
+	id?: string[];
 	streamId?: string;
 	minChange?: number;
 	maxChange?: number;
@@ -376,7 +380,7 @@ export enum IssueSeverity {
 
 /** Setting information required by dashboard */
 export type GetDashboardConfigResponse = {
-	
+
 	/** The name of the external issue service */
 	externalIssueServiceName?: string;
 
@@ -715,6 +719,47 @@ export type GetAgentResponse = {
 
 }
 
+export type Condition = {
+
+	/// The condition text
+	text: string;
+
+	/// Error produced when parsing the condition
+	error?: string;
+
+}
+
+/// Available pool sizing strategies
+export enum PoolSizeStrategy
+{
+	/// Strategy based on lease utilization
+	LeaseUtilization = "LeaseUtilization",
+	
+	/// Strategy based on size of job build queue
+	JobQueue = "JobQueue",
+	
+	/// No-op strategy used as fallback/default behavior
+	NoOp = "NoOp"
+}
+
+/** Here for API completeness, though empty class on server */
+export type LeaseUtilizationSettings = {
+
+}
+
+export type JobQueueSettings = {
+
+	/// Factor translating queue size to additional agents to grow the pool with
+	/// The result is always rounded up to nearest integer. 
+	/// Example: if there are 20 jobs in queue, a factor 0.25 will result in 5 new agents being added (20 * 0.25)
+	scaleOutFactor: number;
+
+	/// Factor by which to shrink the pool size with when queue is empty
+	/// The result is always rounded up to nearest integer.
+	/// Example: when the queue size is zero, a default value of 0.9 will shrink the pool by 10% (current agent count * 0.9)
+	scaleInFactor: number;
+}
+
 /**Information about an agent */
 export type GetPoolResponse = {
 
@@ -724,6 +769,36 @@ export type GetPoolResponse = {
 	/**Friendly name of the agent */
 	name: string;
 
+	/// Condition for agents to be auto-added to the pool
+	condition?: Condition;
+
+	/// Whether to enable autoscaling for this pool
+	enableAutoscaling: boolean;
+
+	/// Cooldown time between scale-out events in seconds
+	scaleOutCooldown?: number;
+
+	/// Cooldown time between scale-in events in seconds
+	scaleInCooldown?: number;
+
+	/// Pool sizing strategy to be used for this pool
+	sizeStrategy?: PoolSizeStrategy;
+
+	/// Settings for lease utilization pool sizing strategy (if used)
+	leaseUtilizationSettings?: LeaseUtilizationSettings;
+
+	/// Settings for job queue pool sizing strategy (if used) 
+	jobQueueSettings?: JobQueueSettings;
+
+	/// The minimum nunmber of agents to retain in this pool
+	minAgents?: number;
+
+	/// The minimum number of idle agents in this pool, if autoscaling is enabled
+	numReserveAgents?: number;
+
+	/// List of workspaces that this agent contains
+	workspaces: GetAgentWorkspaceResponse[];
+
 	/**Arbitrary properties associated with this event */
 	properties?: { [key: string]: string };
 
@@ -1797,7 +1872,7 @@ export type GetTemplateStepStateResponse = {
 
 /**  Updates an existing stream template ref */
 export type UpdateTemplateRefRequest = {
-	
+
 	/** Step states to update */
 	stepStates?: UpdateStepStateRequest[];
 
@@ -1846,13 +1921,13 @@ export type DefaultPreflightRequest = {
 	changeTemplateId?: string;
 }
 
-	
-/**  Configuration for an issue workflow */	
+
+/**  Configuration for an issue workflow */
 export type WorkflowConfig = {
 
 	/** Identifier for this workflow */
 	id: string;
-	
+
 	/** Name of the tab to post summary data to */
 	summaryTab?: string;
 
@@ -1871,10 +1946,10 @@ export type GetStreamResponse = {
 	name: string;
 
 	/**The config file path on the server*/
-	configPath?:string;
+	configPath?: string;
 
 	/**Revision of the config file */
-	configRevision?:string;
+	configRevision?: string;
 
 	/**List of tabs to display for this stream*/
 	tabs: GetStreamTabResponse[];
@@ -2813,7 +2888,7 @@ export type GetIssueResponse = {
 
 	/** User info for who quarantined issue */
 	quarantinedByUserInfo?: GetThinUserInfoResponse;
-	
+
 	/** The UTC time when the issue was quarantined */
 	quarantineTimeUtc?: Date | string;
 
@@ -3939,10 +4014,10 @@ export type GetAgentSoftwareChannelResponse = {
 export type GetExternalIssueResponse = {
 
 	/** The external issue key */
-	key: string; 
+	key: string;
 
 	/** The issue link on external tracking site */
-	link?: string; 
+	link?: string;
 
 	/** The issue status name, "To Do", "In Progress", etc*/
 	statusName?: string;
@@ -3965,7 +4040,7 @@ export type GetExternalIssueResponse = {
 
 /** Request an issue to be created on external issue tracking system */
 export type CreateExternalIssueRequest = {
-		
+
 	/** Horde issue which is linked to external issue */
 	issueId: number;
 
@@ -3993,7 +4068,7 @@ export type CreateExternalIssueRequest = {
 
 /** Response for externally created issue */
 export type CreateExternalIssueResponse = {
-	
+
 	/** External issue tracking key	 */
 	key: string;
 
@@ -4003,7 +4078,7 @@ export type CreateExternalIssueResponse = {
 
 /// External issue project information
 export type GetExternalIssueProjectResponse = {
-		
+
 	/// The project key	
 	projectKey: string;
 
@@ -4022,9 +4097,9 @@ export type GetExternalIssueProjectResponse = {
 
 /// Create a notice which will display on the dashboard
 export type CreateNoticeRequest = {
-		
+
 	/**  Message to display	*/
-	message: string;	
+	message: string;
 }
 
 /// Parameters required to update a notice
diff --git a/Engine/Source/Programs/Horde/HordeDashboard/src/backend/Backend.ts b/Engine/Source/Programs/Horde/HordeDashboard/src/backend/Backend.ts
--- a/Engine/Source/Programs/Horde/HordeDashboard/src/backend/Backend.ts
+++ b/Engine/Source/Programs/Horde/HordeDashboard/src/backend/Backend.ts
@@ -3,7 +3,7 @@
 import { configure } from "mobx";
 import { isNumber } from 'util';
 import templateCache from '../backend/TemplateCache';
-import { AgentData, ArtifactData, AuditLogEntry, AuditLogQuery, BatchUpdatePoolRequest, ChangeSummaryData, CreateDeviceRequest, CreateDeviceResponse, CreateExternalIssueRequest, CreateExternalIssueResponse, CreateJobRequest, CreateJobResponse, CreateNoticeRequest, CreatePoolRequest, CreateSoftwareResponse, CreateSubscriptionRequest, CreateSubscriptionResponse, DashboardPreference, EventData, FindIssueResponse, FindJobTimingsResponse, GetAgentSoftwareChannelResponse, GetArtifactZipRequest, GetDashboardConfigResponse, GetDevicePlatformResponse, GetDevicePoolResponse, GetDeviceReservationResponse, GetDeviceResponse, GetExternalIssueProjectResponse, GetExternalIssueResponse, GetGraphResponse, GetIssueStreamResponse, GetJobsTabResponse, GetJobStepRefResponse, GetJobStepTraceResponse, GetJobTimingResponse, GetLogEventResponse, GetNoticeResponse, GetNotificationResponse, GetPerforceServerStatusResponse, GetServerInfoResponse, GetServerSettingsResponse, GetSoftwareResponse, GetSubscriptionResponse, GetUserResponse, GetUtilizationTelemetryResponse, GlobalConfig, IssueData, IssueQuery, IssueQueryV2, JobData, JobQuery, JobsTabColumnType, JobStepOutcome, JobStreamQuery, JobTimingsQuery, LeaseData, LogData, LogLineData, PoolData, ProjectData, ScheduleData, ScheduleQuery, SearchLogFileResponse, ServerUpdateResponse, SessionData, StreamData, TabType, TemplateData, TestData, UpdateAgentRequest, UpdateDeviceRequest, UpdateGlobalConfigRequest, UpdateIssueRequest, UpdateJobRequest, UpdateLeaseRequest, UpdateNoticeRequest, UpdateNotificationsRequest, UpdatePoolRequest, UpdateServerSettingsRequest, UpdateStepRequest, UpdateStepResponse, UpdateTemplateRefRequest, UpdateUserRequest, UsersQuery } from './Api';
+import { AgentData, AgentQuery, ArtifactData, AuditLogEntry, AuditLogQuery, BatchUpdatePoolRequest, ChangeSummaryData, CreateDeviceRequest, CreateDeviceResponse, CreateExternalIssueRequest, CreateExternalIssueResponse, CreateJobRequest, CreateJobResponse, CreateNoticeRequest, CreatePoolRequest, CreateSoftwareResponse, CreateSubscriptionRequest, CreateSubscriptionResponse, DashboardPreference, EventData, FindIssueResponse, FindJobTimingsResponse, GetAgentSoftwareChannelResponse, GetArtifactZipRequest, GetDashboardConfigResponse, GetDevicePlatformResponse, GetDevicePoolResponse, GetDeviceReservationResponse, GetDeviceResponse, GetExternalIssueProjectResponse, GetExternalIssueResponse, GetGraphResponse, GetIssueStreamResponse, GetJobsTabResponse, GetJobStepRefResponse, GetJobStepTraceResponse, GetJobTimingResponse, GetLogEventResponse, GetNoticeResponse, GetNotificationResponse, GetPerforceServerStatusResponse, GetPoolResponse, GetServerInfoResponse, GetServerSettingsResponse, GetSoftwareResponse, GetSubscriptionResponse, GetUserResponse, GetUtilizationTelemetryResponse, GlobalConfig, IssueData, IssueQuery, IssueQueryV2, JobData, JobQuery, JobsTabColumnType, JobStepOutcome, JobStreamQuery, JobTimingsQuery, LeaseData, LogData, LogLineData, PoolData, ProjectData, ScheduleData, ScheduleQuery, SearchLogFileResponse, ServerUpdateResponse, SessionData, StreamData, TabType, TemplateData, TestData, UpdateAgentRequest, UpdateDeviceRequest, UpdateGlobalConfigRequest, UpdateIssueRequest, UpdateJobRequest, UpdateLeaseRequest, UpdateNoticeRequest, UpdateNotificationsRequest, UpdatePoolRequest, UpdateServerSettingsRequest, UpdateStepRequest, UpdateStepResponse, UpdateTemplateRefRequest, UpdateUserRequest, UsersQuery } from './Api';
 import dashboard from './Dashboard';
 import { ChallengeStatus, Fetch } from './Fetch';
 import graphCache, { GraphQuery } from './GraphCache';
@@ -111,6 +111,13 @@ export class Backend {
         });
     }
 
+    getPool(poolId: string): Promise<GetPoolResponse> {
+        return new Promise<GetPoolResponse>((resolve, reject) => {
+            this.backend.get(`/api/v1/pools/${poolId}`).then((response) => {
+                resolve(response.data as GetPoolResponse);
+            }).catch(reason => { reject(reason); });
+        });
+    }
 
     getPools(): Promise<PoolData[]> {
         return new Promise<PoolData[]>((resolve, reject) => {
@@ -139,10 +146,10 @@ export class Backend {
         });
     }
 
-    getAgents(modifiedAfter: Date): Promise<AgentData[]> {
+    getAgents(query: AgentQuery): Promise<AgentData[]> {
         return new Promise<AgentData[]>((resolve, reject) => {
             // for whenever we bring back modifiedDate, swap to this line
-            this.backend.get(`/api/v1/agents?ModifiedAfter=${modifiedAfter.toISOString()}`).then((response) => {
+            this.backend.get(`/api/v1/agents`, { params: query }).then((response) => {
                 //this.backend.get(`/api/v1/agents`).then((response) => {
                 const agents = response.data as AgentData[];
                 resolve(agents);
diff --git a/Engine/Source/Programs/Horde/HordeDashboard/src/base/utilities/timeUtils.ts b/Engine/Source/Programs/Horde/HordeDashboard/src/base/utilities/timeUtils.ts
--- a/Engine/Source/Programs/Horde/HordeDashboard/src/base/utilities/timeUtils.ts
+++ b/Engine/Source/Programs/Horde/HordeDashboard/src/base/utilities/timeUtils.ts
@@ -260,7 +260,7 @@ export const getElapsedString = (start: Moment, end: Moment, includeSeconds:bool
         duration += `${d.minutes()}m `;
     }
 
-    if (includeSeconds) {
+    if (includeSeconds || !duration) {
         if (d.seconds()) {
             duration += `${d.seconds()}s `;
         }    
@@ -313,7 +313,7 @@ export const getStepElapsed = (step: StepData | BatchData | undefined): string =
         end = moment(step.finishTime);
     }
 
-    return getElapsedString(start, end);
+    return getElapsedString(start, end, false);
 
 };
 
@@ -330,7 +330,7 @@ export const getLeaseElapsed = (lease: GetAgentLeaseResponse | undefined): strin
         end = moment(lease.finishTime);
     }
 
-    return getElapsedString(start, end);
+    return getElapsedString(start, end, false);
 
 };
 
