Horde: Change impersonated P4 workspace name to better sort
#rnx

diff --git a/Engine/Source/Programs/Horde/Horde.Build/Perforce/Impl/PerforceService.cs b/Engine/Source/Programs/Horde/Horde.Build/Perforce/Impl/PerforceService.cs
--- a/Engine/Source/Programs/Horde/Horde.Build/Perforce/Impl/PerforceService.cs
+++ b/Engine/Source/Programs/Horde/Horde.Build/Perforce/Impl/PerforceService.cs
@@ -1,4 +1,4 @@
-﻿// Copyright Epic Games, Inc. All Rights Reserved.
+// Copyright Epic Games, Inc. All Rights Reserved.
 
 using System;
 using System.Collections.Generic;
@@ -523,7 +523,7 @@ namespace Horde.Build.Services
 
 		static string GetClientName(string? serviceUserName, string stream, bool readOnly, bool createChange, string? username = null)
 		{
-			string clientName = $"horde-p4bridge-{Dns.GetHostName()}-{serviceUserName ?? "default"}-{stream.Replace("/", "+", StringComparison.OrdinalIgnoreCase)}";
+			string clientName = $"zzt-horde-p4bridge-{Dns.GetHostName()}-{serviceUserName ?? "default"}-{stream.Replace("/", "+", StringComparison.OrdinalIgnoreCase)}";
 
 			if (!readOnly)
 			{
