Horde: Restore seconds in elapsed time
#rnx
#skipci

diff --git a/Engine/Source/Programs/Horde/HordeDashboard/src/base/utilities/timeUtils.ts b/Engine/Source/Programs/Horde/HordeDashboard/src/base/utilities/timeUtils.ts
--- a/Engine/Source/Programs/Horde/HordeDashboard/src/base/utilities/timeUtils.ts
+++ b/Engine/Source/Programs/Horde/HordeDashboard/src/base/utilities/timeUtils.ts
@@ -313,7 +313,7 @@ export const getStepElapsed = (step: StepData | BatchData | undefined): string =
         end = moment(step.finishTime);
     }
 
-    return getElapsedString(start, end, false);
+    return getElapsedString(start, end);
 
 };
 
@@ -330,7 +330,7 @@ export const getLeaseElapsed = (lease: GetAgentLeaseResponse | undefined): strin
         end = moment(lease.finishTime);
     }
 
-    return getElapsedString(start, end, false);
+    return getElapsedString(start, end);
 
 };
 
