Horde Storage - Only passing BYPASS CACHE option to Scylla instances as its not a cassandra standard keyword so it confuses other cassandra dbs.

diff --git a/Engine/Source/Programs/Horde/HordeStorage/Horde.Storage/Horde.Storage.Startup.cs b/Engine/Source/Programs/Horde/HordeStorage/Horde.Storage/Horde.Storage.Startup.cs
--- a/Engine/Source/Programs/Horde/HordeStorage/Horde.Storage/Horde.Storage.Startup.cs
+++ b/Engine/Source/Programs/Horde/HordeStorage/Horde.Storage/Horde.Storage.Startup.cs
@@ -469,7 +469,8 @@ namespace Horde.Storage
             localSession.Execute(new SimpleStatement("CREATE TYPE IF NOT EXISTS blob_identifier (hash blob)"));
             localSession.UserDefinedTypes.Define(UdtMap.For<ScyllaBlobIdentifier>("blob_identifier", localKeyspaceName));
 
-            return new ScyllaSessionManager(replicatedSession, localSession);
+            bool isScylla = !settings.UseAzureCosmosDB;
+            return new ScyllaSessionManager(replicatedSession, localSession, isScylla, !isScylla);
         }
 
         private IReferencesStore ObjectStoreFactory(IServiceProvider provider)
diff --git a/Engine/Source/Programs/Horde/HordeStorage/Horde.Storage/Implementation/Blob/ScyllaBlobIndex.cs b/Engine/Source/Programs/Horde/HordeStorage/Horde.Storage/Implementation/Blob/ScyllaBlobIndex.cs
--- a/Engine/Source/Programs/Horde/HordeStorage/Horde.Storage/Implementation/Blob/ScyllaBlobIndex.cs
+++ b/Engine/Source/Programs/Horde/HordeStorage/Horde.Storage/Implementation/Blob/ScyllaBlobIndex.cs
@@ -15,12 +15,14 @@ namespace Horde.Storage.Implementation.Blob;
 
 public class ScyllaBlobIndex : IBlobIndex
 {
+    private readonly IScyllaSessionManager _scyllaSessionManager;
     private readonly IOptionsMonitor<JupiterSettings> _jupiterSettings;
     private readonly ISession _session;
     private readonly Mapper _mapper;
 
     public ScyllaBlobIndex(IScyllaSessionManager scyllaSessionManager, IOptionsMonitor<JupiterSettings> jupiterSettings)
     {
+        _scyllaSessionManager = scyllaSessionManager;
         _jupiterSettings = jupiterSettings;
         _session = scyllaSessionManager.GetSessionForReplicatedKeyspace();
         _mapper = new Mapper(_session);
@@ -112,7 +114,8 @@ public class ScyllaBlobIndex : IBlobIndex
     {
         using IScope _ = Tracer.Instance.StartActive("scylla.get_all_blobs");
 
-        foreach (ScyllaBlobIndexTable blobIndex in await _mapper.FetchAsync<ScyllaBlobIndexTable>("BYPASS CACHE"))
+        string cqlOptions = _scyllaSessionManager.IsScylla ? "BYPASS CACHE" : "";
+        foreach (ScyllaBlobIndexTable blobIndex in await _mapper.FetchAsync<ScyllaBlobIndexTable>(cqlOptions))
         {
             yield return blobIndex.ToBlobInfo();
         }
diff --git a/Engine/Source/Programs/Horde/HordeStorage/Horde.Storage/Implementation/Objects/ScyllaReferencesStore.cs b/Engine/Source/Programs/Horde/HordeStorage/Horde.Storage/Implementation/Objects/ScyllaReferencesStore.cs
--- a/Engine/Source/Programs/Horde/HordeStorage/Horde.Storage/Implementation/Objects/ScyllaReferencesStore.cs
+++ b/Engine/Source/Programs/Horde/HordeStorage/Horde.Storage/Implementation/Objects/ScyllaReferencesStore.cs
@@ -50,7 +50,8 @@ namespace Horde.Storage.Implementation
             ));
 
             // BYPASS CACHE is a scylla specific extension to disable populating the cache, should be ignored by other cassandra dbs
-            _getObjectsStatement = _session.Prepare("SELECT bucket, name, last_access_time FROM objects WHERE namespace = ? ALLOW FILTERING BYPASS CACHE");
+            string cqlOptions = scyllaSessionManager.IsScylla ? "BYPASS CACHE" : "";
+            _getObjectsStatement = _session.Prepare($"SELECT bucket, name, last_access_time FROM objects WHERE namespace = ? ALLOW FILTERING {cqlOptions}");
             _getNamespacesStatement = _session.Prepare("SELECT DISTINCT namespace FROM buckets");
         }
 
diff --git a/Engine/Source/Programs/Horde/HordeStorage/Horde.Storage/Implementation/ScyllaSessionManager.cs b/Engine/Source/Programs/Horde/HordeStorage/Horde.Storage/Implementation/ScyllaSessionManager.cs
--- a/Engine/Source/Programs/Horde/HordeStorage/Horde.Storage/Implementation/ScyllaSessionManager.cs
+++ b/Engine/Source/Programs/Horde/HordeStorage/Horde.Storage/Implementation/ScyllaSessionManager.cs
@@ -8,6 +8,16 @@ namespace Horde.Storage.Implementation
     {
         ISession GetSessionForReplicatedKeyspace();
         ISession GetSessionForLocalKeyspace();
+
+        /// <summary>
+        /// True if scylla specific features can be used
+        /// </summary>
+        bool IsScylla { get; }
+
+        /// <summary>
+        /// True when using a base cassandra protocol
+        /// </summary>
+        bool IsCassandra { get; }
     }
 
     public class ScyllaSessionManager : IScyllaSessionManager
@@ -15,10 +25,12 @@ namespace Horde.Storage.Implementation
         private readonly ISession _replicatedSession;
         private readonly ISession _localSession;
 
-        public ScyllaSessionManager(ISession replicatedSession, ISession localSession)
+        public ScyllaSessionManager(ISession replicatedSession, ISession localSession, bool isScylla, bool isCassandra)
         {
             _replicatedSession = replicatedSession;
             _localSession = localSession;
+            IsScylla = isScylla;
+            IsCassandra = isCassandra;
         }
 
         public ISession GetSessionForReplicatedKeyspace()
@@ -31,5 +43,8 @@ namespace Horde.Storage.Implementation
         {
             return _localSession;
         }
+
+        public bool IsScylla { get; }
+        public bool IsCassandra { get; }
     }
 }
