Horde: Fix some lines in logs being blanked out.

JSON logger was matching lines containing a "\\n" sequence and trigger the multi-line handling path, which assumed a format string was present. This defaulted it to an empty string if not, rather than the "message" element.

diff --git a/Engine/Source/Programs/Horde/Horde.Agent.Tests/JsonRpcLoggerTests.cs b/Engine/Source/Programs/Horde/Horde.Agent.Tests/JsonRpcLoggerTests.cs
new file mode 100644
--- /dev/null
+++ b/Engine/Source/Programs/Horde/Horde.Agent.Tests/JsonRpcLoggerTests.cs
@@ -0,0 +1,95 @@
+// Copyright Epic Games, Inc. All Rights Reserved.
+
+using System;
+using System.Buffers;
+using System.Collections.Generic;
+using System.Text;
+using EpicGames.Core;
+using Horde.Agent.Parser;
+using Microsoft.Extensions.Logging;
+using Microsoft.Extensions.Logging.Abstractions;
+using Microsoft.VisualStudio.TestTools.UnitTesting;
+
+namespace Horde.Agent.Tests
+{
+	[TestClass]
+	public class JsonRpcLoggerTests
+	{
+		class JsonLoggerImpl : ILogger
+		{
+			public List<LogEvent> _lines = new List<LogEvent>();
+
+			public IDisposable BeginScope<TState>(TState state) => null!;
+
+			public bool IsEnabled(LogLevel logLevel) => true;
+
+			public void Log<TState>(LogLevel logLevel, EventId eventId, TState state, Exception exception, Func<TState, Exception?, string> formatter)
+			{
+				JsonLogEvent logEvent = JsonLogEvent.FromLoggerState(logLevel, eventId, state, exception, formatter);
+				_lines.Add(LogEvent.Read(logEvent.Data.Span));
+			}
+		}
+
+		[TestMethod]
+		public void JsonLogEventTest()
+		{
+			JsonLoggerImpl impl = new JsonLoggerImpl();
+			impl.LogInformation("Hello {Text}", "world");
+			impl.LogInformation("Hello {Text}", "world");
+
+			Assert.AreEqual(impl._lines.Count, 2);
+			Assert.AreEqual(impl._lines[0].ToString(), "Hello world");
+			Assert.AreEqual(impl._lines[0].GetProperty("Text").ToString()!, "world");
+			Assert.AreEqual(impl._lines[1].ToString(), "Hello world");
+			Assert.AreEqual(impl._lines[1].GetProperty("Text").ToString(), "world");
+		}
+
+		static string[] SplitMultiLineMessage(string message)
+		{
+			JsonLogEvent jsonLogEvent;
+			Assert.IsTrue(JsonLogEvent.TryParse(Encoding.UTF8.GetBytes(message), out jsonLogEvent));
+			ArrayBufferWriter<byte> writer = new ArrayBufferWriter<byte>();
+			int count = JsonRpcLogger.WriteEvent(jsonLogEvent, writer);
+			string[] result = Encoding.UTF8.GetString(writer.WrittenSpan).Split('\n', StringSplitOptions.RemoveEmptyEntries);
+			Assert.AreEqual(count, result.Length);
+			return result;
+		}
+
+		[TestMethod]
+		public void MultiLineFormatTest()
+		{
+			const string msg = @"{""level"":""Information"",""message"":""ignored"",""format"":""This\nis a\nmulti-{Line}-end\nlog message {Var1} {Var2}"",""properties"":{""Line"":""line\nsplit\nin\nfour"",""Var1"":123,""Var2"":{""$type"":""SourceFile"",""$text"":""D:\\build\\\u002B\u002BUE5\\Sync\\GenerateProjectFiles.bat"",""relativePath"":""GenerateProjectFiles.bat"",""depotPath"":""//UE5/Main/GenerateProjectFiles.bat@20392842""}}}";
+
+			string[] result = SplitMultiLineMessage(msg);
+			Assert.AreEqual(7, result.Length);
+			Assert.AreEqual(@"{""level"":""Information"",""message"":""This"",""format"":""This"",""properties"":{""Line"":""line\nsplit\nin\nfour"",""Var1"":123,""Var2"":{""$type"":""SourceFile"",""$text"":""D:\\build\\\u002B\u002BUE5\\Sync\\GenerateProjectFiles.bat"",""relativePath"":""GenerateProjectFiles.bat"",""depotPath"":""//UE5/Main/GenerateProjectFiles.bat@20392842""},""Line$0"":""line"",""Line$1"":""split"",""Line$2"":""in"",""Line$3"":""four""},""line"":0,""lineCount"":7}", result[0]);
+			Assert.AreEqual(@"{""level"":""Information"",""message"":""is a"",""format"":""is a"",""properties"":{""Line"":""line\nsplit\nin\nfour"",""Var1"":123,""Var2"":{""$type"":""SourceFile"",""$text"":""D:\\build\\\u002B\u002BUE5\\Sync\\GenerateProjectFiles.bat"",""relativePath"":""GenerateProjectFiles.bat"",""depotPath"":""//UE5/Main/GenerateProjectFiles.bat@20392842""},""Line$0"":""line"",""Line$1"":""split"",""Line$2"":""in"",""Line$3"":""four""},""line"":1,""lineCount"":7}", result[1]);
+			Assert.AreEqual(@"{""level"":""Information"",""message"":""multi-line"",""format"":""multi-{Line$0}"",""properties"":{""Line"":""line\nsplit\nin\nfour"",""Var1"":123,""Var2"":{""$type"":""SourceFile"",""$text"":""D:\\build\\\u002B\u002BUE5\\Sync\\GenerateProjectFiles.bat"",""relativePath"":""GenerateProjectFiles.bat"",""depotPath"":""//UE5/Main/GenerateProjectFiles.bat@20392842""},""Line$0"":""line"",""Line$1"":""split"",""Line$2"":""in"",""Line$3"":""four""},""line"":2,""lineCount"":7}", result[2]);
+			Assert.AreEqual(@"{""level"":""Information"",""message"":""split"",""format"":""{Line$1}"",""properties"":{""Line"":""line\nsplit\nin\nfour"",""Var1"":123,""Var2"":{""$type"":""SourceFile"",""$text"":""D:\\build\\\u002B\u002BUE5\\Sync\\GenerateProjectFiles.bat"",""relativePath"":""GenerateProjectFiles.bat"",""depotPath"":""//UE5/Main/GenerateProjectFiles.bat@20392842""},""Line$0"":""line"",""Line$1"":""split"",""Line$2"":""in"",""Line$3"":""four""},""line"":3,""lineCount"":7}", result[3]);
+			Assert.AreEqual(@"{""level"":""Information"",""message"":""in"",""format"":""{Line$2}"",""properties"":{""Line"":""line\nsplit\nin\nfour"",""Var1"":123,""Var2"":{""$type"":""SourceFile"",""$text"":""D:\\build\\\u002B\u002BUE5\\Sync\\GenerateProjectFiles.bat"",""relativePath"":""GenerateProjectFiles.bat"",""depotPath"":""//UE5/Main/GenerateProjectFiles.bat@20392842""},""Line$0"":""line"",""Line$1"":""split"",""Line$2"":""in"",""Line$3"":""four""},""line"":4,""lineCount"":7}", result[4]);
+			Assert.AreEqual(@"{""level"":""Information"",""message"":""four-end"",""format"":""{Line$3}-end"",""properties"":{""Line"":""line\nsplit\nin\nfour"",""Var1"":123,""Var2"":{""$type"":""SourceFile"",""$text"":""D:\\build\\\u002B\u002BUE5\\Sync\\GenerateProjectFiles.bat"",""relativePath"":""GenerateProjectFiles.bat"",""depotPath"":""//UE5/Main/GenerateProjectFiles.bat@20392842""},""Line$0"":""line"",""Line$1"":""split"",""Line$2"":""in"",""Line$3"":""four""},""line"":5,""lineCount"":7}", result[5]);
+			Assert.AreEqual(@"{""level"":""Information"",""message"":""log message 123 D:\\build\\\u002B\u002BUE5\\Sync\\GenerateProjectFiles.bat"",""format"":""log message {Var1} {Var2}"",""properties"":{""Line"":""line\nsplit\nin\nfour"",""Var1"":123,""Var2"":{""$type"":""SourceFile"",""$text"":""D:\\build\\\u002B\u002BUE5\\Sync\\GenerateProjectFiles.bat"",""relativePath"":""GenerateProjectFiles.bat"",""depotPath"":""//UE5/Main/GenerateProjectFiles.bat@20392842""},""Line$0"":""line"",""Line$1"":""split"",""Line$2"":""in"",""Line$3"":""four""},""line"":6,""lineCount"":7}", result[6]);
+		}
+
+		[TestMethod]
+		public void MultiLineMessageTest()
+		{
+			const string msg = @"{""level"":""Information"",""message"":""This is \na multi-line string""}";
+
+			string[] result = SplitMultiLineMessage(msg);
+			Assert.AreEqual(2, result.Length);
+			Assert.AreEqual(result[0], @"{""level"":""Information"",""message"":""This is "",""line"":0,""lineCount"":2}");
+			Assert.AreEqual(result[1], @"{""level"":""Information"",""message"":""a multi-line string"",""line"":1,""lineCount"":2}");
+		}
+
+		[TestMethod]
+		public void MultiLineDecoyTest()
+		{
+			const string msg = @"{""level"":""Information"",""message"":""This is \\\\not a multi-line format string""}";
+
+			string[] result = SplitMultiLineMessage(msg);
+			Assert.AreEqual(1, result.Length);
+			Assert.AreEqual(result[0], msg);
+		}
+	}
+}
diff --git a/Engine/Source/Programs/Horde/Horde.Agent.Tests/LogParserTests.cs b/Engine/Source/Programs/Horde/Horde.Agent.Tests/LogParserTests.cs
--- a/Engine/Source/Programs/Horde/Horde.Agent.Tests/LogParserTests.cs
+++ b/Engine/Source/Programs/Horde/Horde.Agent.Tests/LogParserTests.cs
@@ -45,65 +45,6 @@ namespace Horde.Agent.Tests
 			}
 		}
 
-		class JsonLoggerImpl : ILogger
-		{
-			public List<LogEvent> _lines = new List<LogEvent>();
-
-			public JsonLoggerImpl(bool? warnings, ILogger inner) 
-			{
-			}
-
-			public IDisposable BeginScope<TState>(TState state) => null!;
-
-			public bool IsEnabled(LogLevel logLevel) => true;
-
-			public void Log<TState>(LogLevel logLevel, EventId eventId, TState state, Exception exception, Func<TState, Exception?, string> formatter)
-			{
-				JsonLogEvent logEvent = JsonLogEvent.FromLoggerState(logLevel, eventId, state, exception, formatter);
-				_lines.Add(LogEvent.Read(logEvent.Data.Span));
-			}
-		}
-
-		[TestMethod]
-		public void JsonLoggerTest()
-		{
-			JsonLoggerImpl impl = new JsonLoggerImpl(null, NullLogger.Instance);
-			impl.LogInformation("Hello {Text}", "world");
-			impl.LogInformation("Hello {Text}", "world");
-
-			Assert.AreEqual(impl._lines.Count, 2);
-			Assert.AreEqual(impl._lines[0].ToString(), "Hello world");
-			Assert.AreEqual(impl._lines[0].GetProperty("Text").ToString()!, "world");
-			Assert.AreEqual(impl._lines[1].ToString(), "Hello world");
-			Assert.AreEqual(impl._lines[1].GetProperty("Text").ToString(), "world");
-		}
-
-		static string[] SplitMultiLineMessage(string message)
-		{
-			JsonLogEvent jsonLogEvent;
-			Assert.IsTrue(JsonLogEvent.TryParse(Encoding.UTF8.GetBytes(message), out jsonLogEvent));
-			ArrayBufferWriter<byte> writer = new ArrayBufferWriter<byte>();
-			int count = JsonRpcLogger.WriteEvent(jsonLogEvent, writer);
-			string[] result = Encoding.UTF8.GetString(writer.WrittenSpan).Split('\n', StringSplitOptions.RemoveEmptyEntries);
-			Assert.AreEqual(count, result.Length);
-			return result;
-		}
-
-		[TestMethod]
-		public void MultiLineTest()
-		{
-			const string msg = @"{""level"":""Information"",""message"":""ignored"",""format"":""This\nis a\nmulti-{Line}-end\nlog message {Var1} {Var2}"",""properties"":{""Line"":""line\nsplit\nin\nfour"",""Var1"":123,""Var2"":{""$type"":""SourceFile"",""$text"":""D:\\build\\\u002B\u002BUE5\\Sync\\GenerateProjectFiles.bat"",""relativePath"":""GenerateProjectFiles.bat"",""depotPath"":""//UE5/Main/GenerateProjectFiles.bat@20392842""}}}";
-
-			string[] result = SplitMultiLineMessage(msg);
-			Assert.AreEqual(@"{""level"":""Information"",""message"":""This"",""format"":""This"",""properties"":{""Line"":""line\nsplit\nin\nfour"",""Var1"":123,""Var2"":{""$type"":""SourceFile"",""$text"":""D:\\build\\\u002B\u002BUE5\\Sync\\GenerateProjectFiles.bat"",""relativePath"":""GenerateProjectFiles.bat"",""depotPath"":""//UE5/Main/GenerateProjectFiles.bat@20392842""},""Line$0"":""line"",""Line$1"":""split"",""Line$2"":""in"",""Line$3"":""four""},""line"":0,""lineCount"":7}", result[0]);
-			Assert.AreEqual(@"{""level"":""Information"",""message"":""is a"",""format"":""is a"",""properties"":{""Line"":""line\nsplit\nin\nfour"",""Var1"":123,""Var2"":{""$type"":""SourceFile"",""$text"":""D:\\build\\\u002B\u002BUE5\\Sync\\GenerateProjectFiles.bat"",""relativePath"":""GenerateProjectFiles.bat"",""depotPath"":""//UE5/Main/GenerateProjectFiles.bat@20392842""},""Line$0"":""line"",""Line$1"":""split"",""Line$2"":""in"",""Line$3"":""four""},""line"":1,""lineCount"":7}", result[1]);
-			Assert.AreEqual(@"{""level"":""Information"",""message"":""multi-line"",""format"":""multi-{Line$0}"",""properties"":{""Line"":""line\nsplit\nin\nfour"",""Var1"":123,""Var2"":{""$type"":""SourceFile"",""$text"":""D:\\build\\\u002B\u002BUE5\\Sync\\GenerateProjectFiles.bat"",""relativePath"":""GenerateProjectFiles.bat"",""depotPath"":""//UE5/Main/GenerateProjectFiles.bat@20392842""},""Line$0"":""line"",""Line$1"":""split"",""Line$2"":""in"",""Line$3"":""four""},""line"":2,""lineCount"":7}", result[2]);
-			Assert.AreEqual(@"{""level"":""Information"",""message"":""split"",""format"":""{Line$1}"",""properties"":{""Line"":""line\nsplit\nin\nfour"",""Var1"":123,""Var2"":{""$type"":""SourceFile"",""$text"":""D:\\build\\\u002B\u002BUE5\\Sync\\GenerateProjectFiles.bat"",""relativePath"":""GenerateProjectFiles.bat"",""depotPath"":""//UE5/Main/GenerateProjectFiles.bat@20392842""},""Line$0"":""line"",""Line$1"":""split"",""Line$2"":""in"",""Line$3"":""four""},""line"":3,""lineCount"":7}", result[3]);
-			Assert.AreEqual(@"{""level"":""Information"",""message"":""in"",""format"":""{Line$2}"",""properties"":{""Line"":""line\nsplit\nin\nfour"",""Var1"":123,""Var2"":{""$type"":""SourceFile"",""$text"":""D:\\build\\\u002B\u002BUE5\\Sync\\GenerateProjectFiles.bat"",""relativePath"":""GenerateProjectFiles.bat"",""depotPath"":""//UE5/Main/GenerateProjectFiles.bat@20392842""},""Line$0"":""line"",""Line$1"":""split"",""Line$2"":""in"",""Line$3"":""four""},""line"":4,""lineCount"":7}", result[4]);
-			Assert.AreEqual(@"{""level"":""Information"",""message"":""four-end"",""format"":""{Line$3}-end"",""properties"":{""Line"":""line\nsplit\nin\nfour"",""Var1"":123,""Var2"":{""$type"":""SourceFile"",""$text"":""D:\\build\\\u002B\u002BUE5\\Sync\\GenerateProjectFiles.bat"",""relativePath"":""GenerateProjectFiles.bat"",""depotPath"":""//UE5/Main/GenerateProjectFiles.bat@20392842""},""Line$0"":""line"",""Line$1"":""split"",""Line$2"":""in"",""Line$3"":""four""},""line"":5,""lineCount"":7}", result[5]);
-			Assert.AreEqual(@"{""level"":""Information"",""message"":""log message 123 D:\\build\\\u002B\u002BUE5\\Sync\\GenerateProjectFiles.bat"",""format"":""log message {Var1} {Var2}"",""properties"":{""Line"":""line\nsplit\nin\nfour"",""Var1"":123,""Var2"":{""$type"":""SourceFile"",""$text"":""D:\\build\\\u002B\u002BUE5\\Sync\\GenerateProjectFiles.bat"",""relativePath"":""GenerateProjectFiles.bat"",""depotPath"":""//UE5/Main/GenerateProjectFiles.bat@20392842""},""Line$0"":""line"",""Line$1"":""split"",""Line$2"":""in"",""Line$3"":""four""},""line"":6,""lineCount"":7}", result[6]);
-		}
-
 		[TestMethod]
 		public void StructuredOutputMatcher()
 		{
diff --git a/Engine/Source/Programs/Horde/Horde.Agent/Utility/JsonRpcLogger.cs b/Engine/Source/Programs/Horde/Horde.Agent/Utility/JsonRpcLogger.cs
--- a/Engine/Source/Programs/Horde/Horde.Agent/Utility/JsonRpcLogger.cs
+++ b/Engine/Source/Programs/Horde/Horde.Agent/Utility/JsonRpcLogger.cs
@@ -342,16 +342,30 @@ namespace Horde.Agent.Parser
 		public static int WriteEvent(JsonLogEvent jsonLogEvent, IBufferWriter<byte> writer)
 		{
 			ReadOnlySpan<byte> span = jsonLogEvent.Data.Span;
-			if (jsonLogEvent.LineCount > 1 || span.IndexOf(s_escapedNewline) == -1)
+			if (jsonLogEvent.LineCount == 1 && span.IndexOf(s_escapedNewline) != -1)
 			{
-				writer.Write(span);
-				writer.Write(s_newline);
-				return 1;
+				JsonObject obj = (JsonObject)JsonNode.Parse(span)!;
+
+				JsonValue? formatValue = obj["format"] as JsonValue;
+				if (formatValue != null && formatValue.TryGetValue(out string? format))
+				{
+					return WriteEventWithFormat(obj, format, writer);
+				}
+
+				JsonValue? messageValue = obj["message"] as JsonValue;
+				if (messageValue != null && messageValue.TryGetValue(out string? message))
+				{
+					return WriteEventWithMessage(obj, message, writer);
+				}
 			}
 
-			JsonObject obj = (JsonObject)JsonNode.Parse(span)!;
+			writer.Write(span);
+			writer.Write(s_newline);
+			return 1;
+		}
 
-			string format = (obj["format"] as JsonValue)?.ToString() ?? String.Empty;
+		static int WriteEventWithFormat(JsonObject obj, string format, IBufferWriter<byte> writer)
+		{
 			IEnumerable<KeyValuePair<string, object?>> propertyValueList = Enumerable.Empty<KeyValuePair<string, object?>>();
 
 			// Split all the multi-line properties into separate properties
@@ -433,21 +447,43 @@ namespace Horde.Agent.Parser
 			string[] lines = format.Split('\n');
 			for (int idx = 0; idx < lines.Length; idx++)
 			{
-				obj[s_messagePropertyName] = MessageTemplate.Render(lines[idx], propertyValueList);
-				obj[s_formatPropertyName] = lines[idx];
-				obj[s_linePropertyName] = idx;
-				obj[s_lineCountPropertyName] = lines.Length;
-
-				using (Utf8JsonWriter jsonWriter = new Utf8JsonWriter(writer))
-				{
-					obj.WriteTo(jsonWriter);
-				}
+				string message = MessageTemplate.Render(lines[idx], propertyValueList);
+				WriteSingleEvent(obj, message, lines[idx], idx, lines.Length, writer);
+			}
+			return lines.Length;
+		}
 
-				writer.Write(s_newline.Span);
+		static int WriteEventWithMessage(JsonObject obj, string message, IBufferWriter<byte> writer)
+		{
+			string[] lines = message.Split('\n');
+			for (int idx = 0; idx < lines.Length; idx++)
+			{
+				WriteSingleEvent(obj, lines[idx], null, idx, lines.Length, writer);
 			}
 			return lines.Length;
 		}
 
+		static void WriteSingleEvent(JsonObject obj, string message, string? format, int line, int lineCount, IBufferWriter<byte> writer)
+		{
+			obj[s_messagePropertyName] = message;
+			if (format != null)
+			{
+				obj[s_formatPropertyName] = format;
+			}
+			if (lineCount > 1)
+			{
+				obj[s_linePropertyName] = line;
+				obj[s_lineCountPropertyName] = lineCount;
+			}
+
+			using (Utf8JsonWriter jsonWriter = new Utf8JsonWriter(writer))
+			{
+				obj.WriteTo(jsonWriter);
+			}
+
+			writer.Write(s_newline.Span);
+		}
+
 		void AddEvent(ReadOnlySpan<byte> span, int lineIndex, int lineCount, EventSeverity severity, List<CreateEventRequest> events)
 		{
 			try
