Horde: Fix incorrect values in agent tests.

diff --git a/Engine/Source/Programs/Horde/Horde.Agent.Tests/LogParserTests.cs b/Engine/Source/Programs/Horde/Horde.Agent.Tests/LogParserTests.cs
--- a/Engine/Source/Programs/Horde/Horde.Agent.Tests/LogParserTests.cs
+++ b/Engine/Source/Programs/Horde/Horde.Agent.Tests/LogParserTests.cs
@@ -106,7 +106,7 @@ namespace Horde.Agent.Tests
 			};
 
 			List<LogEvent> logEvents = Parse(lines);
-			CheckEventGroup(logEvents, 1, 5, LogLevel.Error, KnownLogEvents.ExitCode);
+			CheckEventGroup(logEvents, 2, 1, LogLevel.Error, KnownLogEvents.ExitCode);
 		}
 
 		[TestMethod]
