Horde Storage - Run Orphan Blob Cleanup more in parallell to speed it up a bit

diff --git a/Engine/Source/Programs/Horde/HordeStorage/Horde.Storage/Implementation/GC/OrphanBlobCleanupRefs.cs b/Engine/Source/Programs/Horde/HordeStorage/Horde.Storage/Implementation/GC/OrphanBlobCleanupRefs.cs
--- a/Engine/Source/Programs/Horde/HordeStorage/Horde.Storage/Implementation/GC/OrphanBlobCleanupRefs.cs
+++ b/Engine/Source/Programs/Horde/HordeStorage/Horde.Storage/Implementation/GC/OrphanBlobCleanupRefs.cs
@@ -86,7 +86,8 @@ namespace Horde.Storage.Implementation
 
                     bool found = false;
                     NamespacePolicy policy = _namespacePolicyResolver.GetPoliciesForNs(@namespace);
-                    
+
+                    List<Task> referenceChecks = new List<Task>();
                     List<NamespaceId> namespacesThatSharePool = namespaces.Where(ns => _namespacePolicyResolver.GetPoliciesForNs(ns).StoragePool == policy.StoragePool).ToList();
                     // check all other namespaces that share the same storage pool for presence of the blob
                     foreach (NamespaceId blobNamespace in namespacesThatSharePool)
@@ -103,30 +104,35 @@ namespace Horde.Storage.Implementation
                             continue;
                         }
 
-                        foreach ((BucketId bucket, IoHashKey key) in blobIndex.References)
-                        {
-                            if (cancellationToken.IsCancellationRequested)
-                            {
-                                break;
-                            }
-
-                            try
-                            {
-                                (ObjectRecord, BlobContents?) _ = await _objectService.Get(blobNamespace, bucket, key, new string[] {"name"});
-                                found = true;
-                            }
-                            catch (ObjectNotFoundException)
+                        Task referenceCheckTask = Parallel.ForEachAsync(blobIndex.References, cancellationToken,
+                            async (tuple, token) =>
                             {
-                                // this is not a valid reference so we should delete
-                            }
-                            catch (MissingBlobsException)
-                            {
-                                // we do not care if there are missing blobs, as long as the record is valid we keep this blob around
-                                found = true;
-                            }
-                        }
+                                (BucketId bucket, IoHashKey key) = tuple;
+                                if (token.IsCancellationRequested)
+                                {
+                                    return;
+                                }
+
+                                try
+                                {
+                                    (ObjectRecord, BlobContents?) _ = await _objectService.Get(blobNamespace, bucket, key, new string[] {"name"});
+                                    found = true;
+                                }
+                                catch (ObjectNotFoundException)
+                                {
+                                    // this is not a valid reference so we should delete
+                                }
+                                catch (MissingBlobsException)
+                                {
+                                    // we do not care if there are missing blobs, as long as the record is valid we keep this blob around
+                                    found = true;
+                                }
+                            });
+                        referenceChecks.Add(referenceCheckTask);
                     }
 
+                    await Task.WhenAll(referenceChecks);
+
                     if (cancellationToken.IsCancellationRequested)
                     {
                         break;
@@ -136,10 +142,10 @@ namespace Horde.Storage.Implementation
 
                     if (!found)
                     {
-                        foreach (NamespaceId namespaceToDeleteFrom in namespacesThatSharePool)
+                        await Parallel.ForEachAsync(namespacesThatSharePool, cancellationToken, async (ns, _) =>
                         {
-                            await RemoveBlob(namespaceToDeleteFrom, blob);
-                        }
+                            await RemoveBlob(ns, blob);
+                        });
                         ++countOfBlobsRemoved;
                     }
                 }
