Horde: Disable link unfurling by default.

diff --git a/Engine/Source/Programs/Horde/Horde.Build/Utilities/BlockKit/BlockKitMessage.cs b/Engine/Source/Programs/Horde/Horde.Build/Utilities/BlockKit/BlockKitMessage.cs
--- a/Engine/Source/Programs/Horde/Horde.Build/Utilities/BlockKit/BlockKitMessage.cs
+++ b/Engine/Source/Programs/Horde/Horde.Build/Utilities/BlockKit/BlockKitMessage.cs
@@ -60,5 +60,11 @@ namespace Horde.Build.Utilities.Slack.BlockKit
 		/// </summary>
 		[JsonPropertyName("attachments")]
 		public List<BlockKitAttachment> Attachments { get; } = new List<BlockKitAttachment>();
+
+		/// <summary>
+		/// Whether to display previews for any linked content
+		/// </summary>
+		[JsonPropertyName("unfurl_links")]
+		public bool UnfurlLinks { get; set; } = false;
     }
 }
