Horde: Add links to external issues in summary reports.

diff --git a/Engine/Source/Programs/Horde/Horde.Build/Notifications/Impl/SlackNotificationSink.cs b/Engine/Source/Programs/Horde/Horde.Build/Notifications/Impl/SlackNotificationSink.cs
--- a/Engine/Source/Programs/Horde/Horde.Build/Notifications/Impl/SlackNotificationSink.cs
+++ b/Engine/Source/Programs/Horde/Horde.Build/Notifications/Impl/SlackNotificationSink.cs
@@ -1269,7 +1269,14 @@ namespace Horde.Build.Notifications.Impl
 						}
 					}
 
-					body.Append($"*: {issue.Summary} [{FormatReadableTimeSpan(report.Time - issue.CreatedAt)}] - {status}");
+					body.Append($"*: {issue.Summary} [{FormatReadableTimeSpan(report.Time - issue.CreatedAt)}]");
+
+					if (!String.IsNullOrEmpty(issue.ExternalIssueKey))
+					{
+						body.Append($" ({FormatExternalIssue(issue.ExternalIssueKey)})");
+					}
+
+					body.Append($" - {status}");
 				}
 
 				if (report.WorkflowStats.NumSteps > 0)
