Horde Storage - Added scope to better account for time spent doing scylla lookups when resolving a content id

diff --git a/Engine/Source/Programs/Horde/HordeStorage/Horde.Storage/Implementation/ContentId/ScyllaContentIdStore.cs b/Engine/Source/Programs/Horde/HordeStorage/Horde.Storage/Implementation/ContentId/ScyllaContentIdStore.cs
--- a/Engine/Source/Programs/Horde/HordeStorage/Horde.Storage/Implementation/ContentId/ScyllaContentIdStore.cs
+++ b/Engine/Source/Programs/Horde/HordeStorage/Horde.Storage/Implementation/ContentId/ScyllaContentIdStore.cs
@@ -45,27 +45,33 @@ namespace Horde.Storage.Implementation
                 blobStoreExistsTask = _blobStore.Exists(ns, contentIdBlob);
             }
 
-            // lower content_weight means its a better candidate to resolve to
-            foreach (ScyllaContentId? resolvedContentId in await _mapper.FetchAsync<ScyllaContentId>("WHERE content_id = ? ORDER BY content_weight DESC", new ScyllaBlobIdentifier(contentId)))
             {
-                if (resolvedContentId == null)
-                {
-                    throw new InvalidContentIdException(contentId);
-                }
-
-                BlobIdentifier[] blobs = resolvedContentId.Chunks.Select(b => b.AsBlobIdentifier()).ToArray();
+                using IScope contentIdFetchScope = Tracer.Instance.StartActive("ScyllaContentIdStore.FetchContentId");
+                scope.Span.ResourceName = contentId.ToString();
 
+                // lower content_weight means its a better candidate to resolve to
+                foreach (ScyllaContentId? resolvedContentId in await _mapper.FetchAsync<ScyllaContentId>("WHERE content_id = ? ORDER BY content_weight DESC", new ScyllaBlobIdentifier(contentId)))
                 {
-                    using IScope _ = Tracer.Instance.StartActive("ScyllaContentIdStore.FindMissingBlobs");
+                    if (resolvedContentId == null)
+                    {
+                        throw new InvalidContentIdException(contentId);
+                    }
+
+                    BlobIdentifier[] blobs = resolvedContentId.Chunks.Select(b => b.AsBlobIdentifier()).ToArray();
 
-                    BlobIdentifier[] missingBlobs = await _blobStore.FilterOutKnownBlobs(ns, blobs);
-                    if (missingBlobs.Length == 0)
                     {
-                        return blobs;
+                        using IScope _ = Tracer.Instance.StartActive("ScyllaContentIdStore.FindMissingBlobs");
+
+                        BlobIdentifier[] missingBlobs = await _blobStore.FilterOutKnownBlobs(ns, blobs);
+                        if (missingBlobs.Length == 0)
+                        {
+                            return blobs;
+                        }
                     }
+                    // blobs are missing continue testing with the next content id in the weighted list as that might exist
                 }
-                // blobs are missing continue testing with the next content id in the weighted list as that might exist
             }
+            
 
             if (!mustBeContentId)
             {
