Horde Storage - Added scope when buffering payload, and a tag to indicate if we are buffering to the filesystem or memory. Helps to cover certain gaps in our timing breakdowns esspecially for larger files.

diff --git a/Engine/Source/Programs/Horde/HordeStorage/Jupiter.Common/Implementation/StreamingBodyPayload.cs b/Engine/Source/Programs/Horde/HordeStorage/Jupiter.Common/Implementation/StreamingBodyPayload.cs
--- a/Engine/Source/Programs/Horde/HordeStorage/Jupiter.Common/Implementation/StreamingBodyPayload.cs
+++ b/Engine/Source/Programs/Horde/HordeStorage/Jupiter.Common/Implementation/StreamingBodyPayload.cs
@@ -3,6 +3,7 @@
 using System;
 using System.IO;
 using System.Threading.Tasks;
+using Datadog.Trace;
 using Jupiter.Utils;
 using Microsoft.AspNetCore.Http;
 using Microsoft.Extensions.Options;
@@ -29,6 +30,8 @@ namespace Jupiter.Common.Implementation
 
         public static async Task<MemoryBufferedPayload> Create(Stream s)
         {
+            using IScope scope = Tracer.Instance.StartActive("payload.buffer");
+            scope.Span.SetTag("bufferType", "Memory");
             MemoryBufferedPayload payload = new MemoryBufferedPayload(await s.ToByteArray());
 
             return payload;
@@ -71,6 +74,8 @@ namespace Jupiter.Common.Implementation
             FilesystemBufferedPayload payload = new FilesystemBufferedPayload();
 
             {
+                using IScope scope = Tracer.Instance.StartActive("payload.buffer");
+                scope.Span.SetTag("bufferType", "Filesystem");
                 await using FileStream fs = payload._tempFile.OpenWrite();
                 await s.CopyToAsync(fs);
             }
