Horde: Do not post messages about needing to integrate a fix if the fix has failed.

diff --git a/Engine/Source/Programs/Horde/Horde.Build/Notifications/Impl/SlackNotificationSink.cs b/Engine/Source/Programs/Horde/Horde.Build/Notifications/Impl/SlackNotificationSink.cs
--- a/Engine/Source/Programs/Horde/Horde.Build/Notifications/Impl/SlackNotificationSink.cs
+++ b/Engine/Source/Programs/Horde/Horde.Build/Notifications/Impl/SlackNotificationSink.cs
@@ -858,14 +858,17 @@ namespace Horde.Build.Notifications.Impl
 						await PostSingleMessageToThreadAsync(triageChannel, state.Ts, fixFailedEventId, fixFailedMessage);
 					}
 
-					foreach (IIssueStream stream in issue.Streams)
+					if (fixFailedSpan == null)
 					{
-						if ((stream.MergeOrigin ?? false) && !(stream.ContainsFix ?? false))
+						foreach (IIssueStream stream in issue.Streams)
 						{
-							string streamName = spans.FirstOrDefault(x => x.StreamId == stream.StreamId)?.StreamName ?? stream.StreamId.ToString();
-							string missingEventId = $"issue_{issue.Id}_fixmissing_{issue.FixChange}_{stream.StreamId}";
-							string missingMessage = $"Note: Fix may need manually merging to {streamName}";
-							await PostSingleMessageToThreadAsync(triageChannel, state.Ts, missingEventId, missingMessage);
+							if ((stream.MergeOrigin ?? false) && !(stream.ContainsFix ?? false))
+							{
+								string streamName = spans.FirstOrDefault(x => x.StreamId == stream.StreamId)?.StreamName ?? stream.StreamId.ToString();
+								string missingEventId = $"issue_{issue.Id}_fixmissing_{issue.FixChange}_{stream.StreamId}";
+								string missingMessage = $"Note: Fix may need manually merging to {streamName}";
+								await PostSingleMessageToThreadAsync(triageChannel, state.Ts, missingEventId, missingMessage);
+							}
 						}
 					}
 				}
