Horde: Remove reactions from Slack threads when issue state changes.

diff --git a/Engine/Source/Programs/Horde/Horde.Build/Notifications/Impl/SlackNotificationSink.cs b/Engine/Source/Programs/Horde/Horde.Build/Notifications/Impl/SlackNotificationSink.cs
--- a/Engine/Source/Programs/Horde/Horde.Build/Notifications/Impl/SlackNotificationSink.cs
+++ b/Engine/Source/Programs/Horde/Horde.Build/Notifications/Impl/SlackNotificationSink.cs
@@ -826,10 +826,19 @@ namespace Horde.Build.Notifications.Impl
 				{
 					await AddReactionAsync(state.Channel, state.Ts, "eyes");
 				}
+				else
+				{
+					await RemoveReactionAsync(state.Channel, state.Ts, "eyes");
+				}
+
 				if (issue.ResolvedAt != null)
 				{
 					await AddReactionAsync(state.Channel, state.Ts, "tick");
 				}
+				else
+				{
+					await RemoveReactionAsync(state.Channel, state.Ts, "tick");
+				}
 
 				if (issue.FixChange != null)
 				{
