Horde: Mention the owner of an issue in messages saying that it's not fixed.

diff --git a/Engine/Source/Programs/Horde/Horde.Build/Notifications/Impl/SlackNotificationSink.cs b/Engine/Source/Programs/Horde/Horde.Build/Notifications/Impl/SlackNotificationSink.cs
--- a/Engine/Source/Programs/Horde/Horde.Build/Notifications/Impl/SlackNotificationSink.cs
+++ b/Engine/Source/Programs/Horde/Horde.Build/Notifications/Impl/SlackNotificationSink.cs
@@ -864,6 +864,11 @@ namespace Horde.Build.Notifications.Impl
 					{
 						string fixFailedEventId = $"issue_{issue.Id}_fixfailed_{issue.FixChange}";
 						string fixFailedMessage = $"Issue not fixed by {FormatChange(issue.FixChange.Value)}; see {FormatJobStep(fixFailedSpan.LastFailure, fixFailedSpan.NodeName)} in {fixFailedSpan.StreamName}.";
+						if (issue.OwnerId.HasValue)
+						{
+							string mention = await FormatMentionAsync(issue.OwnerId.Value, workflow.AllowMentions);
+							fixFailedMessage += $" ({mention})";
+						}
 						await PostSingleMessageToThreadAsync(triageChannel, state.Ts, fixFailedEventId, fixFailedMessage);
 					}
 
