Horde: Fix incorrect summary for link errors without missing symbols.

diff --git a/Engine/Source/Programs/Horde/Horde.Build.Tests/IssueServiceTests.cs b/Engine/Source/Programs/Horde/Horde.Build.Tests/IssueServiceTests.cs
--- a/Engine/Source/Programs/Horde/Horde.Build.Tests/IssueServiceTests.cs
+++ b/Engine/Source/Programs/Horde/Horde.Build.Tests/IssueServiceTests.cs
@@ -1089,6 +1089,26 @@ namespace Horde.Build.Tests
 			}
 		}
 
+		[TestMethod]
+		public async Task LinkerIssueTest2()
+		{
+			string[] lines =
+			{
+				@"..\Intermediate\Build\Win64\UnrealEditor\Development\AdvancedPreviewScene\Default.rc2.res : fatal error LNK1123: failure during conversion to COFF: file invalid or corrupt"
+			};
+
+			IJob job = CreateJob(_mainStreamId, 120, "Linker Test", _graph);
+			await UpdateCompleteStep(job, 0, 0, JobStepOutcome.Success);
+			await ParseEventsAsync(job, 0, 0, lines);
+			await UpdateCompleteStep(job, 0, 0, JobStepOutcome.Failure);
+
+			List<IIssue> issues = await IssueService.FindIssuesAsync();
+			Assert.AreEqual(1, issues.Count);
+
+			IIssue issue = issues[0];
+			Assert.AreEqual("Linker errors in Update Version Files", issue.Summary);
+		}
+
 		[TestMethod]
 		public async Task MaskIssueTest()
 		{
diff --git a/Engine/Source/Programs/Horde/Horde.Build/Issues/IIssueFingerprint.cs b/Engine/Source/Programs/Horde/Horde.Build/Issues/IIssueFingerprint.cs
--- a/Engine/Source/Programs/Horde/Horde.Build/Issues/IIssueFingerprint.cs
+++ b/Engine/Source/Programs/Horde/Horde.Build/Issues/IIssueFingerprint.cs
@@ -90,6 +90,17 @@ namespace Horde.Build.Models
 			return true;
 		}
 
+		/// <summary>
+		/// Gets all the metadata values with a given key
+		/// </summary>
+		/// <param name="fingerprint">Fingerprint to find values for</param>
+		/// <param name="value">Key name to search for</param>
+		/// <returns>All values with the given key</returns>
+		public static bool HasMetadataValue(this IIssueFingerprint fingerprint, string value)
+		{
+			return fingerprint.Metadata != null && fingerprint.Metadata.Contains(value);
+		}
+
 		/// <summary>
 		/// Gets all the metadata values with a given key
 		/// </summary>
diff --git a/Engine/Source/Programs/Horde/Horde.Build/Issues/Impl/Handlers/SymbolIssueHandler.cs b/Engine/Source/Programs/Horde/Horde.Build/Issues/Impl/Handlers/SymbolIssueHandler.cs
--- a/Engine/Source/Programs/Horde/Horde.Build/Issues/Impl/Handlers/SymbolIssueHandler.cs
+++ b/Engine/Source/Programs/Horde/Horde.Build/Issues/Impl/Handlers/SymbolIssueHandler.cs
@@ -4,6 +4,7 @@ using System;
 using System.Collections.Generic;
 using System.Diagnostics.CodeAnalysis;
 using System.Linq;
+using System.Text;
 using EpicGames.Core;
 using Horde.Build.Collections;
 using Horde.Build.Models;
@@ -17,6 +18,11 @@ namespace Horde.Build.IssueHandlers.Impl
 	/// </summary>
 	class SymbolIssueHandler : IIssueHandler
 	{
+		const string NodeName = "Node";
+		const string EventIdName = "EventId";
+
+		static readonly string DuplicateEventIdMetadata = $"{EventIdName}={KnownLogEvents.Linker_DuplicateSymbol.Id}";
+
 		/// <inheritdoc/>
 		public string Type => "Symbol";
 
@@ -73,13 +79,30 @@ namespace Horde.Build.IssueHandlers.Impl
 		public string GetSummary(IIssueFingerprint fingerprint, IssueSeverity severity)
 		{
 			HashSet<string> symbols = fingerprint.Keys;
-			if (symbols.Count == 1)
+			if (symbols.Count == 0)
 			{
-				return $"Undefined symbol '{symbols.First()}'";
+				string[] nodes = fingerprint.GetMetadataValues(NodeName).ToArray();
+
+				StringBuilder summary = new StringBuilder("Linker ");
+				summary.Append((severity == IssueSeverity.Warning) ? "warnings" : "errors");
+				if (nodes.Length > 0)
+				{
+					summary.Append($" in {StringUtils.FormatList(nodes, 2)}");
+				}
+
+				return summary.ToString();
 			}
 			else
 			{
-				return $"Undefined symbols: {StringUtils.FormatList(symbols.ToArray(), 3)}";
+				string problemType = fingerprint.HasMetadataValue(DuplicateEventIdMetadata) ? "Duplicate" : "Undefined";
+				if (symbols.Count == 1)
+				{
+					return $"{problemType} symbol '{symbols.First()}'";
+				}
+				else
+				{
+					return $"{problemType} symbols: {StringUtils.FormatList(symbols.ToArray(), 3)}";
+				}
 			}
 		}
 
@@ -94,7 +117,14 @@ namespace Horde.Build.IssueHandlers.Impl
 			SortedSet<string> symbolNames = new SortedSet<string>();
 			GetSymbolNames(eventData, symbolNames);
 
-			fingerprint = new NewIssueFingerprint(Type, symbolNames, null, null);
+			List<string> metadata = new List<string>();
+			metadata.Add($"{NodeName}={node.Name}");
+			if (eventData.EventId != null)
+			{
+				metadata.Add($"{EventIdName}={eventData.EventId.Value.Id}");
+			}
+
+			fingerprint = new NewIssueFingerprint(Type, symbolNames, null, metadata);
 			return true;
 		}
 	}
