Horde: Prevent suspect list being modified when an owner is assigned to an issue. This will prevent other issues being matched with it.

diff --git a/Engine/Source/Programs/Horde/Horde.Build/Issues/Impl/IssueService.cs b/Engine/Source/Programs/Horde/Horde.Build/Issues/Impl/IssueService.cs
--- a/Engine/Source/Programs/Horde/Horde.Build/Issues/Impl/IssueService.cs
+++ b/Engine/Source/Programs/Horde/Horde.Build/Issues/Impl/IssueService.cs
@@ -975,21 +975,13 @@ namespace Horde.Build.Services.Impl
 				List<NewIssueSuspectData> newSuspects = new List<NewIssueSuspectData>();
 				if (spans.Count > 0)
 				{
-					HashSet<int> suspectChanges = new HashSet<int>(spans[0].Suspects.Where(x => issue.OwnerId == null || x.AuthorId == issue.OwnerId).Select(x => x.OriginatingChange ?? x.Change));
+					HashSet<int> suspectChanges = new HashSet<int>(spans[0].Suspects.Select(x => x.OriginatingChange ?? x.Change));
 					for (int spanIdx = 1; spanIdx < spans.Count; spanIdx++)
 					{
 						suspectChanges.IntersectWith(spans[spanIdx].Suspects.Select(x => x.OriginatingChange ?? x.Change));
 					}
 					newSuspects = spans[0].Suspects.Where(x => suspectChanges.Contains(x.OriginatingChange ?? x.Change)).Select(x => new NewIssueSuspectData(x.AuthorId, x.OriginatingChange ?? x.Change)).ToList();
 				}
-				if (issue.OwnerId != null)
-				{
-					newSuspects.RemoveAll(x => x.AuthorId != issue.OwnerId);
-					if (newSuspects.Count == 0)
-					{
-						newSuspects.Add(new NewIssueSuspectData(issue.OwnerId.Value, 0));
-					}
-				}
 
 				// Create the combined fingerprint for this issue
 				List<NewIssueFingerprint> newFingerprints = new List<NewIssueFingerprint>();
