Horde: Prevent issues being marked as verified when only one span has succeeded.

diff --git a/Engine/Source/Programs/Horde/Horde.Build/Issues/Impl/IssueService.cs b/Engine/Source/Programs/Horde/Horde.Build/Issues/Impl/IssueService.cs
--- a/Engine/Source/Programs/Horde/Horde.Build/Issues/Impl/IssueService.cs
+++ b/Engine/Source/Programs/Horde/Horde.Build/Issues/Impl/IssueService.cs
@@ -1048,6 +1048,7 @@ namespace Horde.Build.Services.Impl
 						{
 							if (span.NextSuccess == null)
 							{
+								newVerifiedAt = null;
 								break;
 							}
 							else if (newVerifiedAt == null || span.NextSuccess.StepTime < newVerifiedAt.Value)
