Horde: Additional logging for issue matching.

diff --git a/Engine/Source/Programs/Horde/Horde.Build/Issues/Impl/IssueService.cs b/Engine/Source/Programs/Horde/Horde.Build/Issues/Impl/IssueService.cs
--- a/Engine/Source/Programs/Horde/Horde.Build/Issues/Impl/IssueService.cs
+++ b/Engine/Source/Programs/Horde/Horde.Build/Issues/Impl/IssueService.cs
@@ -1257,10 +1257,12 @@ namespace Horde.Build.Services.Impl
 			if (span.LastSuccess == null)
 			{
 				existingIssues = await _issueCollection.FindIssuesAsync(streamId: span.StreamId, minChange: span.FirstFailure.Change, maxChange: span.FirstFailure.Change, resolved: false);
+				_logger.LogDebug("Found {NumIssues} open issues at {Change} in {StreamId}", existingIssues.Count, span.FirstFailure.Change, span.StreamId);
 			}
 			else
 			{
 				existingIssues = await _issueCollection.FindIssuesForChangesAsync(span.Suspects.ConvertAll(x => x.OriginatingChange ?? x.Change));
+				_logger.LogDebug("Found {NumIssues} open issues in {StreamId} from [{ChangeList}]", existingIssues.Count, span.StreamId, String.Join(", ", span.Suspects.ConvertAll(x => (x.OriginatingChange ?? x.Change).ToString())));
 			}
 
 			IIssue? issue = existingIssues.FirstOrDefault(x => x.Fingerprints.Any(y => y.IsMatch(span.Fingerprint)));
