Horde: Pool editiing support
#rnx
#skipci

diff --git a/Engine/Source/Programs/Horde/HordeDashboard/src/backend/Api.ts b/Engine/Source/Programs/Horde/HordeDashboard/src/backend/Api.ts
--- a/Engine/Source/Programs/Horde/HordeDashboard/src/backend/Api.ts
+++ b/Engine/Source/Programs/Horde/HordeDashboard/src/backend/Api.ts
@@ -499,6 +499,18 @@ export type UpdatePoolRequest = {
 	/**New name of the pool */
 	name?: string;
 
+	/** Whether to enable autoscaling for this pool */
+	enableAutoscaling?: boolean;
+		
+	/** Pool sizing strategy */
+	sizeStrategy?: PoolSizeStrategy;
+		
+	/** The minimum nunmber of agents to retain in this pool */
+	minAgents?: number;
+
+	/** The minimum number of idle agents in this pool, if autoscaling is enabled */
+	numReserveAgents?: number;
+
 	/**Arbitrary properties associated with this event */
 	properties?: { [key: string]: string };
 }
