Horde: Add support for setting default pool strategy
#rnx
#skipci

diff --git a/Engine/Source/Programs/Horde/Horde.Build.Tests/Fleet/AutoscaleServiceV2Test.cs b/Engine/Source/Programs/Horde/Horde.Build.Tests/Fleet/AutoscaleServiceV2Test.cs
--- a/Engine/Source/Programs/Horde/Horde.Build.Tests/Fleet/AutoscaleServiceV2Test.cs
+++ b/Engine/Source/Programs/Horde/Horde.Build.Tests/Fleet/AutoscaleServiceV2Test.cs
@@ -83,7 +83,7 @@ namespace Horde.Build.Tests.Fleet
 
 			await service.TickLeaderAsync(CancellationToken.None);
 			
-			Assert.AreEqual(1, leaseUtilizationSpy.CallCount);
+			Assert.AreEqual(2, leaseUtilizationSpy.CallCount);
 			Assert.IsTrue(leaseUtilizationSpy.PoolIdsSeen.Contains(pool1.Id));
 			Assert.IsTrue(leaseUtilizationSpy.PoolIdsSeen.Contains(pool2.Id));
 			
diff --git a/Engine/Source/Programs/Horde/Horde.Build/Fleet/Autoscale/AutoscaleServiceV2.cs b/Engine/Source/Programs/Horde/Horde.Build/Fleet/Autoscale/AutoscaleServiceV2.cs
--- a/Engine/Source/Programs/Horde/Horde.Build/Fleet/Autoscale/AutoscaleServiceV2.cs
+++ b/Engine/Source/Programs/Horde/Horde.Build/Fleet/Autoscale/AutoscaleServiceV2.cs
@@ -80,7 +80,7 @@ namespace Horde.Build.Fleet.Autoscale
 			Dictionary<PoolSizeStrategy, List<PoolSizeData>> poolSizeDataByStrategy = await GetPoolSizeDataByStrategyType();
 			foreach ((PoolSizeStrategy strategyType, List<PoolSizeData> currentData) in poolSizeDataByStrategy)
 			{
-				List<PoolSizeData> newData = await GetPoolSizeStrategy(strategyType).CalcDesiredPoolSizesAsync(currentData);
+				List<PoolSizeData> newData = await GetPoolSizeStrategy(strategyType, _settings.DefaultAgentPoolSizeStrategy).CalcDesiredPoolSizesAsync(currentData);
 				await ResizePools(newData);
 			}
 			
@@ -98,7 +98,7 @@ namespace Horde.Build.Fleet.Autoscale
 			foreach (PoolSizeStrategy strategyType in Enum.GetValues<PoolSizeStrategy>())
 			{
 				result[strategyType] = pools
-					.Where(x => (x.SizeStrategy ?? _settings.DefaultAgentPoolSizeStrategy) == strategyType)
+					.Where(x => (x.SizeStrategy == PoolSizeStrategy.Default ? _settings.DefaultAgentPoolSizeStrategy : x.SizeStrategy) == strategyType)
 					.Select(x => new PoolSizeData(x, GetAgentsInPool(x.Id), null))
 					.ToList();
 			}
@@ -190,14 +190,20 @@ namespace Horde.Build.Fleet.Autoscale
 			_noOpStrategy = noOpStrategy;
 		}
 		
-		private IPoolSizeStrategy GetPoolSizeStrategy(PoolSizeStrategy type)
+		private IPoolSizeStrategy GetPoolSizeStrategy(PoolSizeStrategy type, PoolSizeStrategy defaultStrategy)
 		{
+
+			if (type == PoolSizeStrategy.Default)
+			{
+				type = defaultStrategy;
+			}
+
 			return type switch
 			{
 				PoolSizeStrategy.LeaseUtilization => _leaseUtilizationStrategy,
 				PoolSizeStrategy.JobQueue => _jobQueueStrategy,
 				PoolSizeStrategy.NoOp => _noOpStrategy,
-				_ => _noOpStrategy,
+				_ => throw new Exception($"Unknown strategy: {type}")
 			};
 		}
 	}
diff --git a/Engine/Source/Programs/Horde/Horde.Build/Fleet/Autoscale/IPoolSizeStrategy.cs b/Engine/Source/Programs/Horde/Horde.Build/Fleet/Autoscale/IPoolSizeStrategy.cs
--- a/Engine/Source/Programs/Horde/Horde.Build/Fleet/Autoscale/IPoolSizeStrategy.cs
+++ b/Engine/Source/Programs/Horde/Horde.Build/Fleet/Autoscale/IPoolSizeStrategy.cs
@@ -12,6 +12,11 @@ namespace Horde.Build.Fleet.Autoscale
 	/// </summary>
 	public enum PoolSizeStrategy
 	{
+		/// <summary>
+		/// Use the default strategy configured in server settings
+		/// </summary>
+		Default,
+
 		/// <summary>
 		/// Strategy based on lease utilization
 		/// <see cref="Horde.Build.Fleet.Autoscale.LeaseUtilizationStrategy"/>
diff --git a/Engine/Source/Programs/Horde/Horde.Build/Fleet/IPool.cs b/Engine/Source/Programs/Horde/Horde.Build/Fleet/IPool.cs
--- a/Engine/Source/Programs/Horde/Horde.Build/Fleet/IPool.cs
+++ b/Engine/Source/Programs/Horde/Horde.Build/Fleet/IPool.cs
@@ -83,7 +83,7 @@ namespace Horde.Build.Models
 		/// <summary>
 		/// Pool sizing strategy to be used for this pool
 		/// </summary>
-		public PoolSizeStrategy? SizeStrategy { get; }
+		public PoolSizeStrategy SizeStrategy { get; }
 		
 		/// <summary>
 		/// Settings for lease utilization pool sizing strategy (if used)
diff --git a/Engine/Source/Programs/Horde/Horde.Build/Fleet/Impl/PoolCollection.cs b/Engine/Source/Programs/Horde/Horde.Build/Fleet/Impl/PoolCollection.cs
--- a/Engine/Source/Programs/Horde/Horde.Build/Fleet/Impl/PoolCollection.cs
+++ b/Engine/Source/Programs/Horde/Horde.Build/Fleet/Impl/PoolCollection.cs
@@ -42,6 +42,7 @@ namespace Horde.Build.Collections.Impl
 			public TimeSpan? ScaleOutCooldown { get; set; }
 			public TimeSpan? ScaleInCooldown { get; set; }
 			public PoolSizeStrategy? SizeStrategy { get; set; }
+			PoolSizeStrategy IPool.SizeStrategy => SizeStrategy ?? PoolSizeStrategy.Default;
 			public LeaseUtilizationSettings? LeaseUtilizationSettings { get; set; }
 			public JobQueueSettings? JobQueueSettings { get; set; }
 			public int UpdateIndex { get; set; }
@@ -296,7 +297,15 @@ namespace Horde.Build.Collections.Impl
 			}
 			if (sizeStrategy != null)
 			{
-				transaction.Set(x => x.SizeStrategy, sizeStrategy);
+				if (sizeStrategy == PoolSizeStrategy.Default)
+				{
+					transaction.Set(x => x.SizeStrategy, null);
+				}
+				else
+				{
+					transaction.Set(x => x.SizeStrategy, sizeStrategy);
+				}
+				
 			}
 			if (leaseUtilizationSettings != null)
 			{
