Horde: Remove unnecessary IComputeService interface. We don't abstract over anything here; it just adds boilerplate code to maintain.

diff --git a/Engine/Source/Programs/Horde/Horde.Build/Compute/IComputeService.cs b/Engine/Source/Programs/Horde/Horde.Build/Compute/IComputeService.cs
deleted file mode 100644
diff --git a/Engine/Source/Programs/Horde/Horde.Build/Compute/Impl/ComputeController.cs b/Engine/Source/Programs/Horde/Horde.Build/Compute/Impl/ComputeController.cs
--- a/Engine/Source/Programs/Horde/Horde.Build/Compute/Impl/ComputeController.cs
+++ b/Engine/Source/Programs/Horde/Horde.Build/Compute/Impl/ComputeController.cs
@@ -21,14 +21,14 @@ namespace Horde.Build.Compute.Impl
 	public class ComputeController : HordeControllerBase
 	{
 		readonly AclService _aclService;
-		readonly IComputeService _computeService;
+		readonly ComputeService _computeService;
 
 		/// <summary>
 		/// Constructor
 		/// </summary>
 		/// <param name="aclService"></param>
 		/// <param name="computeService">The compute service singleton</param>
-		public ComputeController(AclService aclService, IComputeService computeService)
+		internal ComputeController(AclService aclService, ComputeService computeService)
 		{
 			_aclService = aclService;
 			_computeService = computeService;
diff --git a/Engine/Source/Programs/Horde/Horde.Build/Compute/Impl/ComputeService.cs b/Engine/Source/Programs/Horde/Horde.Build/Compute/Impl/ComputeService.cs
--- a/Engine/Source/Programs/Horde/Horde.Build/Compute/Impl/ComputeService.cs
+++ b/Engine/Source/Programs/Horde/Horde.Build/Compute/Impl/ComputeService.cs
@@ -61,7 +61,7 @@ namespace Horde.Build.Compute.Impl
 	/// <summary>
 	/// Dispatches remote actions. Does not implement any cross-pod communication to satisfy leases; only agents connected to this server instance will be stored.
 	/// </summary>
-	class ComputeService : TaskSourceBase<ComputeTaskMessage>, IComputeService, IHostedService, IDisposable
+	class ComputeService : TaskSourceBase<ComputeTaskMessage>, IHostedService, IDisposable
 	{
 		[RedisConverter(typeof(QueueKeySerializer))]
 		class QueueKey
@@ -209,7 +209,14 @@ namespace Horde.Build.Compute.Impl
 			return new ClusterInfo(config);
 		}
 
-		/// <inheritdoc/>
+		/// <summary>
+		/// Post tasks to be executed to a channel
+		/// </summary>
+		/// <param name="clusterId">Cluster to use for execution</param>
+		/// <param name="channelId">Unique identifier of the client</param>
+		/// <param name="taskRefIds">List of task hashes</param>
+		/// <param name="requirementsHash">Requirements document for execution</param>
+		/// <returns>Async task</returns>
 		public async Task AddTasksAsync(ClusterId clusterId, ChannelId channelId, List<RefId> taskRefIds, CbObjectAttachment requirementsHash)
 		{
 			List<Task> tasks = new List<Task>();
@@ -312,12 +319,24 @@ namespace Horde.Build.Compute.Impl
 			return _taskScheduler.EnqueueAsync(new QueueKey(clusterId, new IoHash(message.RequirementsHash.ToByteArray())), taskInfo, true);
 		}
 
-		/// <inheritdoc/>
+		/// <summary>
+		/// Dequeue completed items from a queue and return immediately
+		/// </summary>
+		/// <param name="clusterId">Cluster containing the channel</param>
+		/// <param name="channelId">Queue to remove items from</param>
+		/// <returns>List of status updates</returns>
 		public async Task<List<ComputeTaskStatus>> GetTaskUpdatesAsync(ClusterId clusterId, ChannelId channelId)
 		{
 			return await _messageQueue.ReadMessagesAsync(GetMessageQueueId(clusterId, channelId));
 		}
 
+		/// <summary>
+		/// Dequeue completed items from a queue
+		/// </summary>
+		/// <param name="clusterId">Cluster containing the channel</param>
+		/// <param name="channelId">Queue to remove items from</param>
+		/// <param name="cancellationToken">Cancellation token to stop waiting for items</param>
+		/// <returns>List of status updates</returns>
 		public async Task<List<ComputeTaskStatus>> WaitForTaskUpdatesAsync(ClusterId clusterId, ChannelId channelId, CancellationToken cancellationToken)
 		{
 			return await _messageQueue.WaitForMessagesAsync(GetMessageQueueId(clusterId, channelId), cancellationToken);
diff --git a/Engine/Source/Programs/Horde/Horde.Build/Startup.cs b/Engine/Source/Programs/Horde/Horde.Build/Startup.cs
--- a/Engine/Source/Programs/Horde/Horde.Build/Startup.cs
+++ b/Engine/Source/Programs/Horde/Horde.Build/Startup.cs
@@ -614,7 +614,7 @@ namespace Horde.Build
 			services.AddHostedService<JobTaskSource>(provider => provider.GetRequiredService<JobTaskSource>());
 			services.AddSingleton<ConformTaskSource>();
 			services.AddHostedService<ConformTaskSource>(provider => provider.GetRequiredService<ConformTaskSource>());
-			services.AddSingleton<IComputeService, ComputeService>();
+			services.AddSingleton<ComputeService>();
 
 			services.AddSingleton<ITaskSource, UpgradeTaskSource>();
 			services.AddSingleton<ITaskSource, ShutdownTaskSource>();
