Horde: Merge a few API definition files.

diff --git a/Engine/Source/Programs/Horde/Horde.Build/Api/Events.cs b/Engine/Source/Programs/Horde/Horde.Build/Api/Events.cs
deleted file mode 100644
diff --git a/Engine/Source/Programs/Horde/Horde.Build/Api/StreamTabs.cs b/Engine/Source/Programs/Horde/Horde.Build/Api/StreamTabs.cs
deleted file mode 100644
diff --git a/Engine/Source/Programs/Horde/Horde.Build/Api/Streams.cs b/Engine/Source/Programs/Horde/Horde.Build/Api/Streams.cs
--- a/Engine/Source/Programs/Horde/Horde.Build/Api/Streams.cs
+++ b/Engine/Source/Programs/Horde/Horde.Build/Api/Streams.cs
@@ -3,6 +3,7 @@
 using System;
 using System.Collections.Generic;
 using System.ComponentModel.DataAnnotations;
+using EpicGames.Core;
 using Horde.Build.Models;
 using Horde.Build.Server;
 using Horde.Build.Utilities;
@@ -696,4 +697,250 @@ namespace Horde.Build.Api
 			Workflows = workflows;
 		}
 	}
+
+	/// <summary>
+	/// Information about a page to display in the dashboard for a stream
+	/// </summary>
+	[JsonKnownTypes(typeof(CreateJobsTabRequest))]
+	public abstract class CreateStreamTabRequest
+	{
+		/// <summary>
+		/// Title of this page
+		/// </summary>
+		[Required]
+		public string Title { get; set; } = null!;
+	}
+
+	/// <summary>
+	/// Type of a column in a jobs tab
+	/// </summary>
+	public enum JobsTabColumnType
+	{
+		/// <summary>
+		/// Contains labels
+		/// </summary>
+		Labels,
+
+		/// <summary>
+		/// Contains parameters
+		/// </summary>
+		Parameter
+	}
+
+	/// <summary>
+	/// Describes a column to display on the jobs page
+	/// </summary>
+	public class CreateJobsTabColumnRequest
+	{
+		/// <summary>
+		/// The type of column
+		/// </summary>
+		public JobsTabColumnType Type { get; set; } = JobsTabColumnType.Labels;
+
+		/// <summary>
+		/// Heading for this column
+		/// </summary>
+		[Required]
+		public string Heading { get; set; } = null!;
+
+		/// <summary>
+		/// Category of aggregates to display in this column. If null, includes any aggregate not matched by another column.
+		/// </summary>
+		public string? Category { get; set; }
+
+		/// <summary>
+		/// Parameter to show in this column
+		/// </summary>
+		public string? Parameter { get; set; }
+
+		/// <summary>
+		/// Relative width of this column.
+		/// </summary>
+		public int? RelativeWidth { get; set; }
+
+		/// <summary>
+		/// Construct a JobsTabColumn object from this request
+		/// </summary>
+		/// <returns>Column object</returns>
+		public JobsTabColumn ToModel()
+		{
+			switch (Type)
+			{
+				case JobsTabColumnType.Labels:
+					return new JobsTabLabelColumn(Heading, Category, RelativeWidth);
+				case JobsTabColumnType.Parameter:
+					return new JobsTabParameterColumn(Heading, Parameter ?? "Undefined", RelativeWidth);
+				default:
+					return new JobsTabLabelColumn(Heading, "Undefined", RelativeWidth);
+			}
+		}
+	}
+
+	/// <summary>
+	/// Describes a job page
+	/// </summary>
+	[JsonDiscriminator("Jobs")]
+	public class CreateJobsTabRequest : CreateStreamTabRequest
+	{
+		/// <summary>
+		/// Whether to show job names on this page
+		/// </summary>
+		public bool ShowNames { get; set; }
+
+		/// <summary>
+		/// Names of jobs to include on this page. If there is only one name specified, the name column does not need to be displayed.
+		/// </summary>
+		public List<string>? JobNames { get; set; }
+
+		/// <summary>
+		/// List of job template names to show on this page.
+		/// </summary>
+		public List<string>? Templates { get; set; }
+
+		/// <summary>
+		/// Columns to display for different types of aggregates
+		/// </summary>
+		public List<CreateJobsTabColumnRequest>? Columns { get; set; }
+	}
+
+	/// <summary>
+	/// Information about a page to display in the dashboard for a stream
+	/// </summary>
+	[JsonKnownTypes(typeof(GetJobsTabResponse))]
+	public abstract class GetStreamTabResponse
+	{
+		/// <summary>
+		/// Title of this page
+		/// </summary>
+		public string Title { get; set; }
+
+		/// <summary>
+		/// Constructor
+		/// </summary>
+		/// <param name="title">Title of this page</param>
+		protected GetStreamTabResponse(string title)//, StreamPageType Type)
+		{
+			Title = title;
+		}
+	}
+
+	/// <summary>
+	/// Describes a column to display on the jobs page
+	/// </summary>
+	[JsonKnownTypes(typeof(GetJobsTabLabelColumnResponse), typeof(GetJobsTabParameterColumnResponse))]
+	public abstract class GetJobsTabColumnResponse
+	{
+		/// <summary>
+		/// Heading for this column
+		/// </summary>
+		public string Heading { get; set; } = null!;
+
+		/// <summary>
+		/// Relative width of this column.
+		/// </summary>
+		public int? RelativeWidth { get; set; }
+
+		/// <summary>
+		/// Constructor
+		/// </summary>
+		/// <param name="heading">Heading for this column</param>
+		/// <param name="relativeWidth">Relative width of this column</param>
+		protected GetJobsTabColumnResponse(string heading, int? relativeWidth)
+		{
+			Heading = heading;
+			RelativeWidth = relativeWidth;
+		}
+	}
+
+	/// <summary>
+	/// Describes a label column to display on the jobs page
+	/// </summary>
+	[JsonDiscriminator(nameof(JobsTabColumnType.Labels))]
+	public class GetJobsTabLabelColumnResponse : GetJobsTabColumnResponse
+	{
+		/// <summary>
+		/// Category of aggregates to display in this column. If null, includes any aggregate not matched by another column.
+		/// </summary>
+		public string? Category { get; set; }
+
+		/// <summary>
+		/// Constructor
+		/// </summary>
+		/// <param name="heading">Heading for this column</param>
+		/// <param name="category">Category of aggregates to display in this column. If null, includes any aggregate not matched by another column.</param>
+		/// <param name="relativeWidth">Relative width of this column</param>
+		public GetJobsTabLabelColumnResponse(string heading, string? category, int? relativeWidth)
+			: base(heading, relativeWidth)
+		{
+			Category = category;
+		}
+	}
+
+	/// <summary>
+	/// Describes a parameter column to display on the jobs page
+	/// </summary>
+	[JsonDiscriminator(nameof(JobsTabColumnType.Parameter))]
+	public class GetJobsTabParameterColumnResponse : GetJobsTabColumnResponse
+	{
+		/// <summary>
+		/// Parameter to show in this column
+		/// </summary>
+		public string? Parameter { get; set; }
+
+		/// <summary>
+		/// Constructor
+		/// </summary>
+		/// <param name="heading">Heading for this column</param>
+		/// <param name="parameter">Name of the parameter to display</param>
+		/// <param name="relativeWidth">Relative width of this column</param>
+		public GetJobsTabParameterColumnResponse(string heading, string? parameter, int? relativeWidth)
+			: base(heading, relativeWidth)
+		{
+			Parameter = parameter;
+		}
+	}
+
+	/// <summary>
+	/// Describes a job page
+	/// </summary>
+	[JsonDiscriminator("Jobs")]
+	public class GetJobsTabResponse : GetStreamTabResponse
+	{
+		/// <summary>
+		/// Whether to show names on the page
+		/// </summary>
+		public bool ShowNames { get; set; }
+
+		/// <summary>
+		/// List of templates to show on the page
+		/// </summary>
+		public List<string>? Templates { get; set; }
+
+		/// <summary>
+		/// Names of jobs to include on this page. If there is only one name specified, the name column does not need to be displayed.
+		/// </summary>
+		public List<string>? JobNames { get; set; }
+
+		/// <summary>
+		/// Columns to display for different types of aggregates
+		/// </summary>
+		public List<GetJobsTabColumnResponse>? Columns { get; set; }
+
+		/// <summary>
+		/// Constructor
+		/// </summary>
+		/// <param name="title">Title for this page</param>
+		/// <param name="bShowNames">Whether to show names on the page</param>
+		/// <param name="templates">Templates to include on this page</param>
+		/// <param name="jobNames">List of job names to include on this page</param>
+		/// <param name="columns">List of columns to display</param>
+		public GetJobsTabResponse(string title, bool bShowNames, List<string>? templates, List<string>? jobNames, List<GetJobsTabColumnResponse>? columns)
+			: base(title)
+		{
+			ShowNames = bShowNames;
+			Templates = templates;
+			JobNames = jobNames;
+			Columns = columns;
+		}
+	}
 }
diff --git a/Engine/Source/Programs/Horde/Horde.Build/Api/Subscriptions.cs b/Engine/Source/Programs/Horde/Horde.Build/Api/Subscriptions.cs
--- a/Engine/Source/Programs/Horde/Horde.Build/Api/Subscriptions.cs
+++ b/Engine/Source/Programs/Horde/Horde.Build/Api/Subscriptions.cs
@@ -1,10 +1,17 @@
 // Copyright Epic Games, Inc. All Rights Reserved.
 
+using System;
 using System.ComponentModel.DataAnnotations;
+using EpicGames.Core;
 using Horde.Build.Models;
+using Horde.Build.Utilities;
+using HordeCommon;
 
 namespace Horde.Build.Api
 {
+	using StreamId = StringId<IStream>;
+	using TemplateRefId = StringId<TemplateRef>;
+
 	/// <summary>
 	/// Request object for creating a new subscription
 	/// </summary>
@@ -85,4 +92,175 @@ namespace Horde.Build.Api
 			NotificationType = subscription.NotificationType;
 		}
 	}
+
+	/// <summary>
+	/// Base class for event records
+	/// </summary>
+	[JsonKnownTypes(typeof(JobCompleteEventRecord), typeof(LabelCompleteEventRecord), typeof(StepCompleteEventRecord))]
+	public class EventRecord
+	{
+	}
+
+	/// <summary>
+	/// Event that occurs when a job completes
+	/// </summary>
+	[JsonDiscriminator("Job")]
+	public class JobCompleteEventRecord : EventRecord
+	{
+		/// <summary>
+		/// The stream id
+		/// </summary>
+		[Required]
+		public string StreamId { get; set; }
+
+		/// <summary>
+		/// The template id
+		/// </summary>
+		[Required]
+		public string TemplateId { get; set; }
+
+		/// <summary>
+		/// Outcome of the job
+		/// </summary>
+		public LabelOutcome Outcome { get; set; } = LabelOutcome.Success;
+
+		/// <summary>
+		/// Default constructor
+		/// </summary>
+		public JobCompleteEventRecord()
+		{
+			StreamId = String.Empty;
+			TemplateId = String.Empty;
+		}
+
+		/// <summary>
+		/// Constructor
+		/// </summary>
+		/// <param name="streamId">The stream id</param>
+		/// <param name="templateId">The template id</param>
+		/// <param name="outcome">Outcome of the job</param>
+		public JobCompleteEventRecord(StreamId streamId, TemplateRefId templateId, LabelOutcome outcome)
+		{
+			StreamId = streamId.ToString();
+			TemplateId = templateId.ToString();
+			Outcome = outcome;
+		}
+	}
+
+	/// <summary>
+	/// Event that occurs when a label completes
+	/// </summary>
+	[JsonDiscriminator("Label")]
+	public class LabelCompleteEventRecord : EventRecord
+	{
+		/// <summary>
+		/// The stream id
+		/// </summary>
+		[Required]
+		public string StreamId { get; set; }
+
+		/// <summary>
+		/// The template id
+		/// </summary>
+		[Required]
+		public string TemplateId { get; set; }
+
+		/// <summary>
+		/// Name of the category for this label
+		/// </summary>
+		public string? CategoryName { get; set; }
+
+		/// <summary>
+		/// Name of the label to monitor
+		/// </summary>
+		[Required]
+		public string LabelName { get; set; }
+
+		/// <summary>
+		/// Outcome of the job
+		/// </summary>
+		public LabelOutcome Outcome { get; set; } = LabelOutcome.Success;
+
+		/// <summary>
+		/// Default constructor
+		/// </summary>
+		public LabelCompleteEventRecord()
+		{
+			StreamId = String.Empty;
+			TemplateId = String.Empty;
+			LabelName = String.Empty;
+		}
+
+		/// <summary>
+		/// Constructor
+		/// </summary>
+		/// <param name="streamId">The stream id</param>
+		/// <param name="templateId">The template id</param>
+		/// <param name="categoryName">Name of the category</param>
+		/// <param name="labelName">The label name</param>
+		/// <param name="outcome">Outcome of the label</param>
+		public LabelCompleteEventRecord(StreamId streamId, TemplateRefId templateId, string? categoryName, string labelName, LabelOutcome outcome)
+		{
+			StreamId = streamId.ToString();
+			TemplateId = templateId.ToString();
+			CategoryName = categoryName;
+			LabelName = labelName;
+			Outcome = outcome;
+		}
+	}
+
+	/// <summary>
+	/// Event that occurs when a step completes
+	/// </summary>
+	[JsonDiscriminator("Step")]
+	public class StepCompleteEventRecord : EventRecord
+	{
+		/// <summary>
+		/// The stream id
+		/// </summary>
+		[Required]
+		public string StreamId { get; set; }
+
+		/// <summary>
+		/// The template id
+		/// </summary>
+		[Required]
+		public string TemplateId { get; set; }
+
+		/// <summary>
+		/// Name of the step to monitor
+		/// </summary>
+		[Required]
+		public string StepName { get; set; }
+
+		/// <summary>
+		/// Outcome of the job step
+		/// </summary>
+		public JobStepOutcome Outcome { get; set; } = JobStepOutcome.Success;
+
+		/// <summary>
+		/// Default constructor
+		/// </summary>
+		public StepCompleteEventRecord()
+		{
+			StreamId = String.Empty;
+			TemplateId = String.Empty;
+			StepName = String.Empty;
+		}
+
+		/// <summary>
+		/// Constructor
+		/// </summary>
+		/// <param name="streamId">The stream id</param>
+		/// <param name="templateId">The template id</param>
+		/// <param name="stepName">The label name</param>
+		/// <param name="outcome">Outcome of the step</param>
+		public StepCompleteEventRecord(StreamId streamId, TemplateRefId templateId, string stepName, JobStepOutcome outcome)
+		{
+			StreamId = streamId.ToString();
+			TemplateId = templateId.ToString();
+			StepName = stepName;
+			Outcome = outcome;
+		}
+	}
 }
