Horde: Remove the IIssueService interface. This does not abstract anything useful.

diff --git a/Engine/Source/Programs/Horde/Horde.Build.Tests/IssueServiceTests.cs b/Engine/Source/Programs/Horde/Horde.Build.Tests/IssueServiceTests.cs
--- a/Engine/Source/Programs/Horde/Horde.Build.Tests/IssueServiceTests.cs
+++ b/Engine/Source/Programs/Horde/Horde.Build.Tests/IssueServiceTests.cs
@@ -14,6 +14,7 @@ using Horde.Build.Models;
 using Horde.Build.Server;
 using Horde.Build.Services;
 using Horde.Build.Tests.Stubs.Services;
+using Horde.Build.Collections;
 using Horde.Build.Utilities;
 using HordeCommon;
 using Microsoft.Extensions.DependencyInjection;
@@ -22,14 +23,14 @@ using Microsoft.VisualStudio.TestTools.UnitTesting;
 using MongoDB.Bson;
 using MongoDB.Driver;
 using Moq;
-using ProjectId = Horde.Build.Utilities.StringId<Horde.Build.Models.IProject>;
-using StreamId = Horde.Build.Utilities.StringId<Horde.Build.Models.IStream>;
-using TemplateRefId = Horde.Build.Utilities.StringId<Horde.Build.Models.TemplateRef>;
 
 namespace Horde.Build.Tests
 {
 	using JobId = ObjectId<IJob>;
 	using LogId = ObjectId<ILogFile>;
+	using ProjectId = StringId<IProject>;
+	using StreamId = StringId<IStream>;
+	using TemplateRefId = StringId<TemplateRef>;
 	using UserId = ObjectId<IUser>;
 
 	[TestClass]
@@ -289,7 +290,7 @@ namespace Horde.Build.Tests
 				await AddEvent(job, 0, 0, new { level = nameof(LogLevel.Warning) }, EventSeverity.Warning);
 				await UpdateCompleteStep(job, 0, 0, JobStepOutcome.Warnings);
 
-				List<IIssue> issues = await IssueService.FindIssuesAsync();
+				List<IIssue> issues = await IssueCollection.FindIssuesAsync();
 				Assert.AreEqual(1, issues.Count);
 
 				Assert.AreEqual("Warnings in Update Version Files", issues[0].Summary);
@@ -305,7 +306,7 @@ namespace Horde.Build.Tests
 				await AddEvent(job, 0, 2, new { });
 				await UpdateCompleteStep(job, 0, 2, JobStepOutcome.Failure);
 
-				List<IIssue> issues = (await IssueService.FindIssuesAsync()).OrderBy(x => x.Summary).ToList();
+				List<IIssue> issues = (await IssueCollection.FindIssuesAsync()).OrderBy(x => x.Summary).ToList();
 				Assert.AreEqual(3, issues.Count);
 
 				Assert.AreEqual("Errors in Compile ShooterGameEditor Win64", issues[0].Summary);
@@ -321,7 +322,7 @@ namespace Horde.Build.Tests
 				await AddEvent(job, 0, 0, new { });
 				await UpdateCompleteStep(job, 0, 0, JobStepOutcome.Failure);
 
-				List<IIssue> issues = (await IssueService.FindIssuesAsync()).OrderBy(x => x.Summary).ToList();
+				List<IIssue> issues = (await IssueCollection.FindIssuesAsync()).OrderBy(x => x.Summary).ToList();
 				Assert.AreEqual(3, issues.Count);
 
 				Assert.AreEqual("Errors in Compile ShooterGameEditor Win64", issues[0].Summary);
@@ -337,18 +338,18 @@ namespace Horde.Build.Tests
 				await AddEvent(job, 0, 3, new { });
 				await UpdateCompleteStep(job, 0, 3, JobStepOutcome.Warnings);
 
-				List<IIssue> issues = await IssueService.FindIssuesAsync();
+				List<IIssue> issues = await IssueCollection.FindIssuesAsync();
 				Assert.AreEqual(4, issues.Count);
 			}
 
 			// #5
 			// Add a description to the issue
 			{
-				List<IIssue> issues = await IssueService.FindIssuesAsync();
+				List<IIssue> issues = await IssueCollection.FindIssuesAsync();
 
 				IIssue issue = issues[0];
 				await IssueService.UpdateIssueAsync(issue.Id, description: "Hello world!");
-				IIssue? newIssue = await IssueService.GetIssueAsync(issue.Id);
+				IIssue? newIssue = await IssueCollection.GetIssueAsync(issue.Id);
 				Assert.AreEqual(newIssue?.Description, "Hello world!");
 			}
 		}
@@ -364,7 +365,7 @@ namespace Horde.Build.Tests
 				await AddEvent(job, 0, 0, new { level = nameof(LogLevel.Warning) }, EventSeverity.Warning);
 				await UpdateCompleteStep(job, 0, 0, JobStepOutcome.Warnings);
 
-				List<IIssue> issues = await IssueService.FindIssuesAsync();
+				List<IIssue> issues = await IssueCollection.FindIssuesAsync();
 				Assert.AreEqual(1, issues.Count);
 				Assert.AreEqual(IssueSeverity.Warning, issues[0].Severity);
 
@@ -379,7 +380,7 @@ namespace Horde.Build.Tests
 				await AddEvent(job, 0, 0, new { level = nameof(LogLevel.Error) }, EventSeverity.Error);
 				await UpdateCompleteStep(job, 0, 0, JobStepOutcome.Failure);
 
-				List<IIssue> issues = await IssueService.FindIssuesAsync();
+				List<IIssue> issues = await IssueCollection.FindIssuesAsync();
 				Assert.AreEqual(1, issues.Count);
 				Assert.AreEqual(IssueSeverity.Error, issues[0].Severity);
 
@@ -397,10 +398,10 @@ namespace Horde.Build.Tests
 				IJob job = CreateJob(_mainStreamId, 105, "Test Build", _graph);
 				await UpdateCompleteStep(job, 0, 0, JobStepOutcome.Success);
 
-				List<IIssue> issues = await IssueService.FindIssuesAsync();
+				List<IIssue> issues = await IssueCollection.FindIssuesAsync();
 				Assert.AreEqual(0, issues.Count);
 
-				List<IIssue> openIssues = await IssueService.FindIssuesAsync(resolved: false);
+				List<IIssue> openIssues = await IssueCollection.FindIssuesAsync(resolved: false);
 				Assert.AreEqual(0, openIssues.Count);
 			}
 
@@ -423,11 +424,11 @@ namespace Horde.Build.Tests
 				Assert.AreEqual(1, events.Count);
 				Assert.AreEqual(2, events[0].LineCount);
 
-				List<IIssue> issues = await IssueService.FindIssuesAsync();
+				List<IIssue> issues = await IssueCollection.FindIssuesAsync();
 				Assert.AreEqual(1, issues.Count);
 
 				IIssue issue = issues[0];
-				List<IIssueSpan> spans = await IssueService.GetIssueSpansAsync(issue);
+				List<IIssueSpan> spans = await IssueCollection.FindSpansAsync(issue.Id);
 				Assert.AreEqual(1, spans.Count);
 				Assert.AreEqual(1, issue.Fingerprints.Count);
 				Assert.AreEqual("Compile", issue.Fingerprints[0].Type);
@@ -445,10 +446,10 @@ namespace Horde.Build.Tests
 				IJob job = CreateJob(_mainStreamId, 105, "Test Build", _graph);
 				await UpdateCompleteStep(job, 0, 0, JobStepOutcome.Success);
 
-				List<IIssue> issues = await IssueService.FindIssuesAsync();
+				List<IIssue> issues = await IssueCollection.FindIssuesAsync();
 				Assert.AreEqual(0, issues.Count);
 
-				List<IIssue> openIssues = await IssueService.FindIssuesAsync(resolved: false);
+				List<IIssue> openIssues = await IssueCollection.FindIssuesAsync(resolved: false);
 				Assert.AreEqual(0, openIssues.Count);
 			}
 
@@ -465,11 +466,11 @@ namespace Horde.Build.Tests
 				await ParseEventsAsync(job, 0, 0, lines);
 				await UpdateCompleteStep(job, 0, 0, JobStepOutcome.Failure);
 
-				List<IIssue> issues = await IssueService.FindIssuesAsync();
+				List<IIssue> issues = await IssueCollection.FindIssuesAsync();
 				Assert.AreEqual(1, issues.Count);
 
 				IIssue issue = issues[0];
-				List<IIssueSpan> spans = await IssueService.GetIssueSpansAsync(issue);
+				List<IIssueSpan> spans = await IssueCollection.FindSpansAsync(issue.Id);
 				Assert.AreEqual(1, spans.Count);
 				Assert.AreEqual(1, issue.Fingerprints.Count);
 				Assert.AreEqual("Compile", issue.Fingerprints[0].Type);
@@ -479,7 +480,7 @@ namespace Horde.Build.Tests
 				Assert.AreEqual(105, stream.LastSuccess?.Change);
 				Assert.AreEqual(null, stream.NextSuccess?.Change);
 
-				IReadOnlyList<IIssueSuspect> suspects = await IssueService.GetIssueSuspectsAsync(issue);
+				IReadOnlyList<IIssueSuspect> suspects = await IssueCollection.FindSuspectsAsync(issue);
 				suspects = suspects.OrderBy(x => x.Id).ToList();
 				Assert.AreEqual(3, suspects.Count);
 
@@ -495,7 +496,7 @@ namespace Horde.Build.Tests
 				Assert.AreEqual(_timId, suspects[0].AuthorId);
 				//				Assert.AreEqual(null, Suspects[0].OriginatingChange);
 
-				List<IIssue> openIssues = await IssueService.FindIssuesAsync(resolved: false);
+				List<IIssue> openIssues = await IssueCollection.FindIssuesAsync(resolved: false);
 				Assert.AreEqual(1, openIssues.Count);
 			}
 
@@ -506,18 +507,18 @@ namespace Horde.Build.Tests
 				IJob job = CreateJob(_mainStreamId, 110, "Test Build", _graph);
 				await UpdateCompleteStep(job, 0, 0, JobStepOutcome.Success);
 
-				List<IIssue> issues = await IssueService.FindIssuesAsync();
+				List<IIssue> issues = await IssueCollection.FindIssuesAsync();
 				Assert.AreEqual(1, issues.Count);
 
 				IIssue issue = issues[0];
-				List<IIssueSpan> spans = await IssueService.GetIssueSpansAsync(issue);
+				List<IIssueSpan> spans = await IssueCollection.FindSpansAsync(issue.Id);
 				Assert.AreEqual(spans.Count, 1);
 
 				IIssueSpan stream = spans[0];
 				Assert.AreEqual(110, stream.LastSuccess?.Change);
 				Assert.AreEqual(null, stream.NextSuccess?.Change);
 
-				IReadOnlyList<IIssueSuspect> suspects = (await IssueService.GetIssueSuspectsAsync(issue)).OrderByDescending(x => x.Change).ToList();
+				IReadOnlyList<IIssueSuspect> suspects = (await IssueCollection.FindSuspectsAsync(issue)).OrderByDescending(x => x.Change).ToList();
 				Assert.AreEqual(2, suspects.Count);
 
 				Assert.AreEqual(120, suspects[0].Change);
@@ -526,7 +527,7 @@ namespace Horde.Build.Tests
 				Assert.AreEqual(75, suspects[1].Change);
 				Assert.AreEqual(_jerryId, suspects[1].AuthorId);
 
-				List<IIssue> openIssues = await IssueService.FindIssuesAsync(resolved: false);
+				List<IIssue> openIssues = await IssueCollection.FindIssuesAsync(resolved: false);
 				Assert.AreEqual(1, openIssues.Count);
 			}
 
@@ -537,18 +538,18 @@ namespace Horde.Build.Tests
 				IJob job = CreateJob(_mainStreamId, 125, "Test Build", _graph);
 				await UpdateCompleteStep(job, 0, 0, JobStepOutcome.Success);
 
-				List<IIssue> issues = await IssueService.FindIssuesAsync(resolved: true);
+				List<IIssue> issues = await IssueCollection.FindIssuesAsync(resolved: true);
 				Assert.AreEqual(1, issues.Count);
 
 				IIssue issue = issues[0];
-				List<IIssueSpan> spans = await IssueService.GetIssueSpansAsync(issue);
+				List<IIssueSpan> spans = await IssueCollection.FindSpansAsync(issue.Id);
 				Assert.AreEqual(spans.Count, 1);
 
 				IIssueSpan stream = spans[0];
 				Assert.AreEqual(110, stream.LastSuccess?.Change);
 				Assert.AreEqual(125, stream.NextSuccess?.Change);
 
-				List<IIssue> openIssues = await IssueService.FindIssuesAsync(resolved: false);
+				List<IIssue> openIssues = await IssueCollection.FindIssuesAsync(resolved: false);
 				Assert.AreEqual(0, openIssues.Count);
 			}
 
@@ -565,18 +566,18 @@ namespace Horde.Build.Tests
 				await ParseEventsAsync(job, 0, 0, lines);
 				await UpdateCompleteStep(job, 0, 0, JobStepOutcome.Failure);
 
-				List<IIssue> issues = await IssueService.FindIssuesAsync(resolved: true);
+				List<IIssue> issues = await IssueCollection.FindIssuesAsync(resolved: true);
 				Assert.AreEqual(1, issues.Count);
 
 				IIssue issue = issues[0];
-				List<IIssueSpan> spans = await IssueService.GetIssueSpansAsync(issue);
+				List<IIssueSpan> spans = await IssueCollection.FindSpansAsync(issue.Id);
 				Assert.AreEqual(1, spans.Count);
 
 				IIssueSpan span = spans[0];
 				Assert.AreEqual(110, span.LastSuccess?.Change);
 				Assert.AreEqual(125, span.NextSuccess?.Change);
 
-				List<IIssueStep> steps = await IssueService.GetIssueStepsAsync(span);
+				List<IIssueStep> steps = await IssueCollection.FindStepsAsync(span.Id);
 				Assert.AreEqual(2, steps.Count);
 			}
 
@@ -588,10 +589,10 @@ namespace Horde.Build.Tests
 				await AddEvent(job, 0, 1, new { });
 				await UpdateCompleteStep(job, 0, 1, JobStepOutcome.Failure);
 
-				List<IIssue> resolvedIssues = await IssueService.FindIssuesAsync(resolved: true);
+				List<IIssue> resolvedIssues = await IssueCollection.FindIssuesAsync(resolved: true);
 				Assert.AreEqual(1, resolvedIssues.Count);
 
-				List<IIssue> unresolvedIssues = await IssueService.FindIssuesAsync(resolved: false);
+				List<IIssue> unresolvedIssues = await IssueCollection.FindIssuesAsync(resolved: false);
 				Assert.AreEqual(1, unresolvedIssues.Count);
 			}
 		}
@@ -606,7 +607,7 @@ namespace Horde.Build.Tests
 				IJob job = CreateJob(_mainStreamId, 105, "Compile Test", _graph);
 				await UpdateCompleteStep(job, 0, 0, JobStepOutcome.Success);
 
-				List<IIssue> issues = await IssueService.FindIssuesAsync();
+				List<IIssue> issues = await IssueCollection.FindIssuesAsync();
 				Assert.AreEqual(0, issues.Count);
 			}
 
@@ -627,11 +628,11 @@ namespace Horde.Build.Tests
 				await ParseEventsAsync(job, 0, 0, lines);
 				await UpdateCompleteStep(job, 0, 0, JobStepOutcome.Failure);
 
-				List<IIssue> issues = await IssueService.FindIssuesAsync();
+				List<IIssue> issues = await IssueCollection.FindIssuesAsync();
 				Assert.AreEqual(1, issues.Count);
 				Assert.IsTrue(issues[0].Promoted);
 
-				List<IIssueSuspect> suspects = await IssueService.GetIssueSuspectsAsync(issues[0]);
+				List<IIssueSuspect> suspects = await IssueCollection.FindSuspectsAsync(issues[0]);
 
 				List<UserId> primarySuspects = suspects.Select(x => x.AuthorId).ToList();
 				Assert.AreEqual(2, primarySuspects.Count);
@@ -646,7 +647,7 @@ namespace Horde.Build.Tests
 				IJob job = CreateJob(_mainStreamId, 115, "Compile Test", _graph);
 				await UpdateCompleteStep(job, 0, 0, JobStepOutcome.Success);
 
-				List<IIssue> issues = await IssueService.FindIssuesAsync();
+				List<IIssue> issues = await IssueCollection.FindIssuesAsync();
 				Assert.AreEqual(1, issues.Count);
 
 				IIssue issue = issues[0];
@@ -668,7 +669,7 @@ namespace Horde.Build.Tests
 				IJob job = CreateJob(_mainStreamId, 105, "Compile Test", _graph);
 				await UpdateCompleteStep(job, 0, 0, JobStepOutcome.Success);
 
-				List<IIssue> issues = await IssueService.FindIssuesAsync();
+				List<IIssue> issues = await IssueCollection.FindIssuesAsync();
 				Assert.AreEqual(0, issues.Count);
 			}
 
@@ -689,13 +690,13 @@ namespace Horde.Build.Tests
 				await ParseEventsAsync(job, 0, 0, lines);
 				await UpdateCompleteStep(job, 0, 0, JobStepOutcome.Failure);
 
-				List<IIssue> issues = await IssueService.FindIssuesAsync();
+				List<IIssue> issues = await IssueCollection.FindIssuesAsync();
 				Assert.AreEqual(1, issues.Count);
 				Assert.IsFalse(issues[0].Promoted);
 
 				await IssueService.UpdateIssueAsync(issues[0].Id, promoted: true);
 
-				issues = await IssueService.FindIssuesAsync();
+				issues = await IssueCollection.FindIssuesAsync();
 				Assert.AreEqual(1, issues.Count);
 				Assert.IsTrue(issues[0].Promoted);
 			}
@@ -711,7 +712,7 @@ namespace Horde.Build.Tests
 				IJob job = CreateJob(_mainStreamId, 105, "Compile Test", _graph);
 				await UpdateCompleteStep(job, 0, 0, JobStepOutcome.Success);
 
-				List<IIssue> issues = await IssueService.FindIssuesAsync();
+				List<IIssue> issues = await IssueCollection.FindIssuesAsync();
 				Assert.AreEqual(0, issues.Count);
 			}
 
@@ -732,7 +733,7 @@ namespace Horde.Build.Tests
 				await ParseEventsAsync(job, 0, 0, lines);
 				await UpdateCompleteStep(job, 0, 0, JobStepOutcome.Failure);
 
-				List<IIssue> issues = await IssueService.FindIssuesAsync();
+				List<IIssue> issues = await IssueCollection.FindIssuesAsync();
 				Assert.AreEqual(1, issues.Count);
 				Assert.IsTrue(issues[0].Promoted);
 
@@ -756,7 +757,7 @@ namespace Horde.Build.Tests
 			await ParseEventsAsync(job, 0, 4, lines);
 			await UpdateCompleteStep(job, 0, 4, JobStepOutcome.Failure);
 
-			List<IIssue> issues = await IssueService.FindIssuesAsync();
+			List<IIssue> issues = await IssueCollection.FindIssuesAsync();
 			Assert.AreEqual(1, issues.Count);
 
 			IIssue issue = issues[0];
@@ -773,7 +774,7 @@ namespace Horde.Build.Tests
 				IJob job = CreateJob(_mainStreamId, 105, "Compile Test", _graph);
 				await UpdateCompleteStep(job, 0, 0, JobStepOutcome.Success);
 
-				List<IIssue> issues = await IssueService.FindIssuesAsync();
+				List<IIssue> issues = await IssueCollection.FindIssuesAsync();
 				Assert.AreEqual(0, issues.Count);
 			}
 
@@ -791,7 +792,7 @@ namespace Horde.Build.Tests
 				await ParseEventsAsync(job, 0, 0, lines);
 				await UpdateCompleteStep(job, 0, 0, JobStepOutcome.Failure);
 
-				List<IIssue> issues = await IssueService.FindIssuesAsync();
+				List<IIssue> issues = await IssueCollection.FindIssuesAsync();
 				Assert.AreEqual(1, issues.Count);
 
 				IIssue issue = issues[0];
@@ -816,7 +817,7 @@ namespace Horde.Build.Tests
 				IJob job = CreateJob(_mainStreamId, 105, "Compile Test", _graph);
 				await UpdateCompleteStep(job, 0, 0, JobStepOutcome.Success);
 
-				List<IIssue> issues = await IssueService.FindIssuesAsync();
+				List<IIssue> issues = await IssueCollection.FindIssuesAsync();
 				Assert.AreEqual(0, issues.Count);
 			}
 
@@ -838,10 +839,10 @@ namespace Horde.Build.Tests
 				await ParseEventsAsync(job, 0, 0, lines);
 				await UpdateCompleteStep(job, 0, 0, JobStepOutcome.Failure);
 
-				List<IIssue> issues = await IssueService.FindIssuesAsync();
+				List<IIssue> issues = await IssueCollection.FindIssuesAsync();
 				Assert.AreEqual(1, issues.Count);
 
-				List<IIssueSuspect> suspects = await IssueService.GetIssueSuspectsAsync(issues[0]);
+				List<IIssueSuspect> suspects = await IssueCollection.FindSuspectsAsync(issues[0]);
 
 				List<UserId> primarySuspects = suspects.Select(x => x.AuthorId).ToList();
 				Assert.AreEqual(1, primarySuspects.Count);
@@ -859,7 +860,7 @@ namespace Horde.Build.Tests
 				IJob job = CreateJob(_mainStreamId, 105, "Compile Test", _graph);
 				await UpdateCompleteStep(job, 0, 0, JobStepOutcome.Success);
 
-				List<IIssue> issues = await IssueService.FindIssuesAsync();
+				List<IIssue> issues = await IssueCollection.FindIssuesAsync();
 				Assert.AreEqual(0, issues.Count);
 			}
 
@@ -880,10 +881,10 @@ namespace Horde.Build.Tests
 				await ParseEventsAsync(job, 0, 0, lines);
 				await UpdateCompleteStep(job, 0, 0, JobStepOutcome.Failure);
 
-				List<IIssue> issues = await IssueService.FindIssuesAsync();
+				List<IIssue> issues = await IssueCollection.FindIssuesAsync();
 				Assert.AreEqual(1, issues.Count);
 
-				List<IIssueSuspect> suspects = await IssueService.GetIssueSuspectsAsync(issues[0]);
+				List<IIssueSuspect> suspects = await IssueCollection.FindSuspectsAsync(issues[0]);
 
 				List<UserId> primarySuspects = suspects.Select(x => x.AuthorId).ToList();
 				Assert.AreEqual(2, primarySuspects.Count);
@@ -895,14 +896,14 @@ namespace Horde.Build.Tests
 			// Scenario: Tim declines the issue
 			// Expected: Only suspect is Jerry, but owner is still unassigned
 			{
-				List<IIssue> issues = await IssueService.FindIssuesAsync();
+				List<IIssue> issues = await IssueCollection.FindIssuesAsync();
 				Assert.AreEqual(1, issues.Count);
 				await IssueService.UpdateIssueAsync(issues[0].Id, declinedById: _timId);
 
-				issues = await IssueService.FindIssuesAsync();
+				issues = await IssueCollection.FindIssuesAsync();
 				Assert.AreEqual(1, issues.Count);
 
-				List<IIssueSuspect> suspects = await IssueService.GetIssueSuspectsAsync(issues[0]);
+				List<IIssueSuspect> suspects = await IssueCollection.FindSuspectsAsync(issues[0]);
 
 				List<UserId> primarySuspects = suspects.Where(x => x.DeclinedAt == null).Select(x => x.AuthorId).ToList();
 				Assert.AreEqual(1, primarySuspects.Count);
@@ -920,7 +921,7 @@ namespace Horde.Build.Tests
 				IJob job = CreateJob(_mainStreamId, 105, "Test Build", _graph);
 				await UpdateCompleteStep(job, 0, 0, JobStepOutcome.Success);
 
-				List<IIssue> issues = await IssueService.FindIssuesAsync();
+				List<IIssue> issues = await IssueCollection.FindIssuesAsync();
 				Assert.AreEqual(0, issues.Count);
 			}
 
@@ -937,11 +938,11 @@ namespace Horde.Build.Tests
 				await ParseEventsAsync(job, 0, 0, lines);
 				await UpdateCompleteStep(job, 0, 0, JobStepOutcome.Failure);
 
-				List<IIssue> issues = await IssueService.FindIssuesAsync();
+				List<IIssue> issues = await IssueCollection.FindIssuesAsync();
 				Assert.AreEqual(1, issues.Count);
 
 				IIssue issue = issues[0];
-				List<IIssueSpan> spans = await IssueService.GetIssueSpansAsync(issue);
+				List<IIssueSpan> spans = await IssueCollection.FindSpansAsync(issue.Id);
 				Assert.AreEqual(spans.Count, 1);
 				Assert.AreEqual(issue.Fingerprints.Count, 1);
 				Assert.AreEqual(issue.Fingerprints[0].Type, "Symbol");
@@ -950,7 +951,7 @@ namespace Horde.Build.Tests
 				Assert.AreEqual(105, stream.LastSuccess?.Change);
 				Assert.AreEqual(null, stream.NextSuccess?.Change);
 
-				List<IIssueSuspect> suspects = await IssueService.GetIssueSuspectsAsync(issues[0]);
+				List<IIssueSuspect> suspects = await IssueCollection.FindSuspectsAsync(issues[0]);
 
 				List<UserId> primarySuspects = suspects.Select(x => x.AuthorId).ToList();
 				Assert.AreEqual(1, primarySuspects.Count);
@@ -968,7 +969,7 @@ namespace Horde.Build.Tests
 				IJob job = CreateJob(_mainStreamId, 105, "Test Build", _graph);
 				await UpdateCompleteStep(job, 0, 0, JobStepOutcome.Success);
 
-				List<IIssue> issues = await IssueService.FindIssuesAsync();
+				List<IIssue> issues = await IssueCollection.FindIssuesAsync();
 				Assert.AreEqual(0, issues.Count);
 			}
 
@@ -1000,11 +1001,11 @@ namespace Horde.Build.Tests
 				await ParseEventsAsync(job, 0, 1, lines2);
 				await UpdateCompleteStep(job, 0, 1, JobStepOutcome.Failure);
 
-				List<IIssue> issues = await IssueService.FindIssuesAsync();
+				List<IIssue> issues = await IssueCollection.FindIssuesAsync();
 				Assert.AreEqual(1, issues.Count);
 
 				IIssue issue = issues[0];
-				List<IIssueSpan> spans = await IssueService.GetIssueSpansAsync(issue);
+				List<IIssueSpan> spans = await IssueCollection.FindSpansAsync(issue.Id);
 				Assert.AreEqual(spans.Count, 2);
 				Assert.AreEqual(issue.Fingerprints.Count, 1);
 				Assert.AreEqual(issue.Fingerprints[0].Type, "Symbol");
@@ -1029,7 +1030,7 @@ namespace Horde.Build.Tests
 				IJob job = CreateJob(_mainStreamId, 105, "Test Build", _graph);
 				await UpdateCompleteStep(job, 0, 0, JobStepOutcome.Success);
 
-				List<IIssue> issues = await IssueService.FindIssuesAsync();
+				List<IIssue> issues = await IssueCollection.FindIssuesAsync();
 				Assert.AreEqual(0, issues.Count);
 			}
 
@@ -1047,11 +1048,11 @@ namespace Horde.Build.Tests
 				await ParseEventsAsync(job, 0, 0, lines);
 				await UpdateCompleteStep(job, 0, 0, JobStepOutcome.Failure);
 
-				List<IIssue> issues = await IssueService.FindIssuesAsync();
+				List<IIssue> issues = await IssueCollection.FindIssuesAsync();
 				Assert.AreEqual(1, issues.Count);
 
 				IIssue issue = issues[0];
-				List<IIssueSpan> spans = await IssueService.GetIssueSpansAsync(issue);
+				List<IIssueSpan> spans = await IssueCollection.FindSpansAsync(issue.Id);
 				Assert.AreEqual(spans.Count, 1);
 				Assert.AreEqual(issue.Fingerprints.Count, 1);
 				Assert.AreEqual(issue.Fingerprints[0].Type, "Symbol");
@@ -1074,11 +1075,11 @@ namespace Horde.Build.Tests
 				await ParseEventsAsync(job, 0, 0, lines);
 				await UpdateCompleteStep(job, 0, 0, JobStepOutcome.Failure);
 
-				List<IIssue> issues = await IssueService.FindIssuesAsync();
+				List<IIssue> issues = await IssueCollection.FindIssuesAsync();
 				Assert.AreEqual(1, issues.Count);
 
 				IIssue issue = issues[0];
-				List<IIssueSpan> spans = await IssueService.GetIssueSpansAsync(issue);
+				List<IIssueSpan> spans = await IssueCollection.FindSpansAsync(issue.Id);
 				Assert.AreEqual(spans.Count, 1);
 				Assert.AreEqual(issue.Fingerprints.Count, 1);
 				Assert.AreEqual(issue.Fingerprints[0].Type, "Symbol");
@@ -1102,7 +1103,7 @@ namespace Horde.Build.Tests
 			await ParseEventsAsync(job, 0, 0, lines);
 			await UpdateCompleteStep(job, 0, 0, JobStepOutcome.Failure);
 
-			List<IIssue> issues = await IssueService.FindIssuesAsync();
+			List<IIssue> issues = await IssueCollection.FindIssuesAsync();
 			Assert.AreEqual(1, issues.Count);
 
 			IIssue issue = issues[0];
@@ -1119,7 +1120,7 @@ namespace Horde.Build.Tests
 				IJob job = CreateJob(_mainStreamId, 105, "Test Build", _graph);
 				await UpdateCompleteStep(job, 0, 0, JobStepOutcome.Success);
 
-				List<IIssue> issues = await IssueService.FindIssuesAsync();
+				List<IIssue> issues = await IssueCollection.FindIssuesAsync();
 				Assert.AreEqual(0, issues.Count);
 			}
 
@@ -1137,11 +1138,11 @@ namespace Horde.Build.Tests
 				await ParseEventsAsync(job, 0, 0, lines);
 				await UpdateCompleteStep(job, 0, 0, JobStepOutcome.Failure);
 
-				List<IIssue> issues = await IssueService.FindIssuesAsync();
+				List<IIssue> issues = await IssueCollection.FindIssuesAsync();
 				Assert.AreEqual(1, issues.Count);
 
 				IIssue issue = issues[0];
-				List<IIssueSpan> spans = await IssueService.GetIssueSpansAsync(issue);
+				List<IIssueSpan> spans = await IssueCollection.FindSpansAsync(issue.Id);
 				Assert.AreEqual(1, spans.Count);
 				Assert.AreEqual(1, issue.Fingerprints.Count);
 				Assert.AreEqual("Compile", issue.Fingerprints[0].Type);
@@ -1160,11 +1161,11 @@ namespace Horde.Build.Tests
 			await ParseEventsAsync(job, 0, 0, lines);
 			await UpdateCompleteStep(job, 0, 0, JobStepOutcome.Failure);
 
-			List<IIssue> issues = await IssueService.FindIssuesAsync();
+			List<IIssue> issues = await IssueCollection.FindIssuesAsync();
 			Assert.AreEqual(1, issues.Count);
 
 			IIssue issue = issues[0];
-			List<IIssueSpan> spans = await IssueService.GetIssueSpansAsync(issue);
+			List<IIssueSpan> spans = await IssueCollection.FindSpansAsync(issue.Id);
 			Assert.AreEqual(1, spans.Count);
 			Assert.AreEqual(1, issue.Fingerprints.Count);
 			Assert.AreEqual("Copyright", issue.Fingerprints[0].Type);
@@ -1187,11 +1188,11 @@ namespace Horde.Build.Tests
 				await ParseEventsAsync(job, 0, 0, lines);
 				await UpdateCompleteStep(job, 0, 0, JobStepOutcome.Failure);
 
-				List<IIssue> issues = await IssueService.FindIssuesAsync();
+				List<IIssue> issues = await IssueCollection.FindIssuesAsync();
 				Assert.AreEqual(1, issues.Count);
 
 				issueA = issues[0];
-				List<IIssueSpan> spans = await IssueService.GetIssueSpansAsync(issueA);
+				List<IIssueSpan> spans = await IssueCollection.FindSpansAsync(issueA.Id);
 				Assert.AreEqual(spans.Count, 1);
 				spanA = spans[0];
 			}
@@ -1209,12 +1210,12 @@ namespace Horde.Build.Tests
 				await ParseEventsAsync(job, 0, 0, lines);
 				await UpdateCompleteStep(job, 0, 0, JobStepOutcome.Failure);
 
-				List<IIssue> issues = await IssueService.FindIssuesAsync();
+				List<IIssue> issues = await IssueCollection.FindIssuesAsync();
 				issues.RemoveAll(x => x.Id == issueA.Id);
 				Assert.AreEqual(1, issues.Count);
 
 				issueB = issues[0];
-				List<IIssueSpan> spans = await IssueService.GetIssueSpansAsync(issueB);
+				List<IIssueSpan> spans = await IssueCollection.FindSpansAsync(issueB.Id);
 				Assert.AreEqual(1, spans.Count);
 				spanB = spans[0];
 			}
@@ -1223,20 +1224,20 @@ namespace Horde.Build.Tests
 			{
 				await IssueService.UpdateIssueAsync(issueA.Id, addSpanIds: new List<ObjectId> { spanB.Id });
 
-				IIssue newIssueA = (await IssueService.GetIssueAsync(issueA.Id))!;
+				IIssue newIssueA = (await IssueCollection.GetIssueAsync(issueA.Id))!;
 				Assert.IsNull(newIssueA.VerifiedAt);
 				Assert.IsNull(newIssueA.ResolvedAt);
 				Assert.AreEqual(2, newIssueA.Fingerprints.Count);
-				List<IIssueSpan> newSpansA = await IssueService.GetIssueSpansAsync(newIssueA!);
+				List<IIssueSpan> newSpansA = await IssueCollection.FindSpansAsync(newIssueA!.Id);
 				Assert.AreEqual(2, newSpansA.Count);
 				Assert.AreEqual(newIssueA.Id, newSpansA[0].IssueId);
 				Assert.AreEqual(newIssueA.Id, newSpansA[1].IssueId);
 
-				IIssue newIssueB = (await IssueService.GetIssueAsync(issueB.Id))!;
+				IIssue newIssueB = (await IssueCollection.GetIssueAsync(issueB.Id))!;
 				Assert.IsNotNull(newIssueB.VerifiedAt);
 				Assert.IsNotNull(newIssueB.ResolvedAt);
 				Assert.AreEqual(0, newIssueB.Fingerprints.Count);
-				List<IIssueSpan> newSpansB = await IssueService.GetIssueSpansAsync(newIssueB);
+				List<IIssueSpan> newSpansB = await IssueCollection.FindSpansAsync(newIssueB.Id);
 				Assert.AreEqual(0, newSpansB.Count);
 			}
 
@@ -1244,18 +1245,18 @@ namespace Horde.Build.Tests
 			{
 				await IssueService.UpdateIssueAsync(issueB.Id, addSpanIds: new List<ObjectId> { spanA.Id, spanB.Id });
 
-				IIssue newIssueA = (await IssueService.GetIssueAsync(issueA.Id))!;
+				IIssue newIssueA = (await IssueCollection.GetIssueAsync(issueA.Id))!;
 				Assert.IsNotNull(newIssueA.VerifiedAt);
 				Assert.IsNotNull(newIssueA.ResolvedAt);
 				Assert.AreEqual(0, newIssueA.Fingerprints.Count);
-				List<IIssueSpan> newSpansA = await IssueService.GetIssueSpansAsync(newIssueA);
+				List<IIssueSpan> newSpansA = await IssueCollection.FindSpansAsync(newIssueA.Id);
 				Assert.AreEqual(0, newSpansA.Count);
 
-				IIssue newIssueB = (await IssueService.GetIssueAsync(issueB.Id))!;
+				IIssue newIssueB = (await IssueCollection.GetIssueAsync(issueB.Id))!;
 				Assert.IsNull(newIssueB.VerifiedAt);
 				Assert.IsNull(newIssueB.ResolvedAt);
 				Assert.AreEqual(2, newIssueB.Fingerprints.Count);
-				List<IIssueSpan> newSpansB = await IssueService.GetIssueSpansAsync(newIssueB!);
+				List<IIssueSpan> newSpansB = await IssueCollection.FindSpansAsync(newIssueB!.Id);
 				Assert.AreEqual(2, newSpansB.Count);
 				Assert.AreEqual(newIssueB.Id, newSpansB[0].IssueId);
 				Assert.AreEqual(newIssueB.Id, newSpansB[1].IssueId);
@@ -1322,11 +1323,11 @@ namespace Horde.Build.Tests
 				Assert.AreEqual(1, events.Count);
 				Assert.AreEqual(40, events[0].LineCount);
 
-				List<IIssue> issues = await IssueService.FindIssuesAsync();
+				List<IIssue> issues = await IssueCollection.FindIssuesAsync();
 				Assert.AreEqual(1, issues.Count);
 
 				IIssue issue = issues[0];
-				List<IIssueSpan> spans = await IssueService.GetIssueSpansAsync(issue);
+				List<IIssueSpan> spans = await IssueCollection.FindSpansAsync(issue.Id);
 				Assert.AreEqual(spans.Count, 1);
 				Assert.AreEqual(issue.Fingerprints.Count, 1);
 				Assert.AreEqual(issue.Fingerprints[0].Type, "Gauntlet");
@@ -1347,7 +1348,7 @@ namespace Horde.Build.Tests
 				await AddEvent(job, 0, 0, new { level = nameof(LogLevel.Warning) }, EventSeverity.Warning);
 				await UpdateCompleteStep(job, 0, 0, JobStepOutcome.Warnings);
 
-				List<IIssue> issues = await IssueService.FindIssuesAsync();
+				List<IIssue> issues = await IssueCollection.FindIssuesAsync();
 				Assert.AreEqual(1, issues.Count);
 				Assert.AreEqual(IssueSeverity.Warning, issues[0].Severity);
 
@@ -1362,10 +1363,10 @@ namespace Horde.Build.Tests
 			{
 				await IssueService.UpdateIssueAsync(issueId, resolvedById: _bobId);
 
-				List<IIssue> openIssues = await IssueService.FindIssuesAsync();
+				List<IIssue> openIssues = await IssueCollection.FindIssuesAsync();
 				Assert.AreEqual(0, openIssues.Count);
 
-				IIssue issue = (await IssueService.GetIssueAsync(issueId))!;
+				IIssue issue = (await IssueCollection.GetIssueAsync(issueId))!;
 				Assert.IsNotNull(issue.ResolvedAt);
 				Assert.AreEqual(_bobId, issue.OwnerId);
 				Assert.AreEqual(_bobId, issue.ResolvedById);
@@ -1379,10 +1380,10 @@ namespace Horde.Build.Tests
 				await AddEvent(job, 0, 0, new { level = nameof(LogLevel.Warning) }, EventSeverity.Warning);
 				await UpdateCompleteStep(job, 0, 0, JobStepOutcome.Warnings);
 
-				List<IIssue> openIssues = await IssueService.FindIssuesAsync();
+				List<IIssue> openIssues = await IssueCollection.FindIssuesAsync();
 				Assert.AreEqual(0, openIssues.Count);
 
-				IIssue issue = (await IssueService.GetIssueAsync(issueId))!;
+				IIssue issue = (await IssueCollection.GetIssueAsync(issueId))!;
 				Assert.IsNotNull(issue.ResolvedAt);
 				Assert.AreEqual(_bobId, issue.OwnerId);
 				Assert.AreEqual(_bobId, issue.ResolvedById);
@@ -1396,7 +1397,7 @@ namespace Horde.Build.Tests
 				await AddEvent(job, 0, 0, new { level = nameof(LogLevel.Warning) }, EventSeverity.Warning);
 				await UpdateCompleteStep(job, 0, 0, JobStepOutcome.Warnings);
 
-				List<IIssue> openIssues = await IssueService.FindIssuesAsync();
+				List<IIssue> openIssues = await IssueCollection.FindIssuesAsync();
 				Assert.AreEqual(1, openIssues.Count);
 
 				IIssue issue = openIssues[0];
@@ -1412,10 +1413,10 @@ namespace Horde.Build.Tests
 			{
 				await IssueService.UpdateIssueAsync(issueId, resolvedById: _bobId, fixChange: 115);
 
-				List<IIssue> openIssues = await IssueService.FindIssuesAsync();
+				List<IIssue> openIssues = await IssueCollection.FindIssuesAsync();
 				Assert.AreEqual(0, openIssues.Count);
 
-				IIssue issue = (await IssueService.GetIssueAsync(issueId))!;
+				IIssue issue = (await IssueCollection.GetIssueAsync(issueId))!;
 				Assert.IsNotNull(issue.ResolvedAt);
 				Assert.AreEqual(_bobId, issue.OwnerId);
 				Assert.AreEqual(_bobId, issue.ResolvedById);
@@ -1429,7 +1430,7 @@ namespace Horde.Build.Tests
 				await AddEvent(job, 0, 0, new { level = nameof(LogLevel.Warning) }, EventSeverity.Warning);
 				await UpdateCompleteStep(job, 0, 0, JobStepOutcome.Warnings);
 
-				List<IIssue> openIssues = await IssueService.FindIssuesAsync();
+				List<IIssue> openIssues = await IssueCollection.FindIssuesAsync();
 				Assert.AreEqual(1, openIssues.Count);
 
 				IIssue issue = openIssues[0];
@@ -1445,10 +1446,10 @@ namespace Horde.Build.Tests
 			{
 				await IssueService.UpdateIssueAsync(issueId, resolvedById: _bobId, fixChange: 125);
 
-				List<IIssue> openIssues = await IssueService.FindIssuesAsync();
+				List<IIssue> openIssues = await IssueCollection.FindIssuesAsync();
 				Assert.AreEqual(0, openIssues.Count);
 
-				IIssue issue = (await IssueService.GetIssueAsync(issueId))!;
+				IIssue issue = (await IssueCollection.GetIssueAsync(issueId))!;
 				Assert.IsNotNull(issue.ResolvedAt);
 				Assert.AreEqual(_bobId, issue.OwnerId);
 				Assert.AreEqual(_bobId, issue.ResolvedById);
@@ -1461,7 +1462,7 @@ namespace Horde.Build.Tests
 				IJob job = CreateJob(_mainStreamId, 125, "Test Build", _graph, TimeSpan.FromHours(25.0));
 				await UpdateCompleteStep(job, 0, 0, JobStepOutcome.Success);
 
-				List<IIssue> openIssues = await IssueService.FindIssuesAsync();
+				List<IIssue> openIssues = await IssueCollection.FindIssuesAsync();
 				Assert.AreEqual(0, openIssues.Count);
 			}
 
@@ -1473,7 +1474,7 @@ namespace Horde.Build.Tests
 				await AddEvent(job, 0, 0, new { level = nameof(LogLevel.Warning) }, EventSeverity.Warning);
 				await UpdateCompleteStep(job, 0, 0, JobStepOutcome.Warnings);
 
-				List<IIssue> openIssues = await IssueService.FindIssuesAsync();
+				List<IIssue> openIssues = await IssueCollection.FindIssuesAsync();
 				Assert.AreEqual(1, openIssues.Count);
 
 				IIssue issue = openIssues[0];
@@ -1494,7 +1495,7 @@ namespace Horde.Build.Tests
 				await AddEvent(job, 0, 0, new { level = nameof(LogLevel.Warning) }, EventSeverity.Warning);
 				await UpdateCompleteStep(job, 0, 0, JobStepOutcome.Warnings);
 
-				List<IIssue> issues = await IssueService.FindIssuesAsync();
+				List<IIssue> issues = await IssueCollection.FindIssuesAsync();
 				Assert.AreEqual(1, issues.Count);
 				Assert.AreEqual(IssueSeverity.Warning, issues[0].Severity);
 
@@ -1510,10 +1511,10 @@ namespace Horde.Build.Tests
 				IJob job = CreateJob(_mainStreamId, 115, "Test Build", _graph);
 				await UpdateCompleteStep(job, 0, 0, JobStepOutcome.Success);
 
-				IIssue? issue = await IssueService.GetIssueAsync(issueId);
+				IIssue? issue = await IssueCollection.GetIssueAsync(issueId);
 				Assert.IsNotNull(issue!.ResolvedAt);
 
-				List<IIssueSpan> spans = await IssueService.GetIssueSpansAsync(issue);
+				List<IIssueSpan> spans = await IssueCollection.FindSpansAsync(issue.Id);
 				Assert.AreEqual(spans.Count, 1);
 			}
 		}
@@ -1562,7 +1563,7 @@ namespace Horde.Build.Tests
 				spans.Add(MockSpan(new StreamId("ue5-release-staging"), MockSuspect(101, 1), MockSuspect(102, 2)));
 				spans.Add(MockSpan(new StreamId("ue5-release"), MockSuspect(1, null), MockSuspect(2, null)));
 
-				List<IIssueSpan> results = Horde.Build.Services.Impl.IssueService.FindMergeOriginSpans(spans);
+				List<IIssueSpan> results = Horde.Build.Services.IssueService.FindMergeOriginSpans(spans);
 				Assert.AreEqual(1, results.Count);
 				Assert.AreEqual(new StreamId("ue5-release"), results[0].StreamId);
 			}
@@ -1573,7 +1574,7 @@ namespace Horde.Build.Tests
 				spans.Add(MockSpan(new StreamId("ue5-release-staging"), MockSuspect(101, 1), MockSuspect(102, null)));
 				spans.Add(MockSpan(new StreamId("ue5-release"), MockSuspect(1, null), MockSuspect(2, null)));
 
-				List<IIssueSpan> results = Horde.Build.Services.Impl.IssueService.FindMergeOriginSpans(spans);
+				List<IIssueSpan> results = Horde.Build.Services.IssueService.FindMergeOriginSpans(spans);
 				Assert.AreEqual(1, results.Count);
 				Assert.AreEqual(new StreamId("ue5-release"), results[0].StreamId);
 			}
@@ -1584,7 +1585,7 @@ namespace Horde.Build.Tests
 				spans.Add(MockSpan(new StreamId("ue5-release-staging"), MockSuspect(101, 1), MockSuspect(102, null)));
 				spans.Add(MockSpan(new StreamId("ue5-main"), MockSuspect(201, null), MockSuspect(202, 2)));
 
-				List<IIssueSpan> results = Horde.Build.Services.Impl.IssueService.FindMergeOriginSpans(spans);
+				List<IIssueSpan> results = Horde.Build.Services.IssueService.FindMergeOriginSpans(spans);
 				Assert.AreEqual(1, results.Count);
 				Assert.AreEqual(new StreamId("ue5-release"), results[0].StreamId);
 			}
diff --git a/Engine/Source/Programs/Horde/Horde.Build.Tests/TestSetup.cs b/Engine/Source/Programs/Horde/Horde.Build.Tests/TestSetup.cs
--- a/Engine/Source/Programs/Horde/Horde.Build.Tests/TestSetup.cs
+++ b/Engine/Source/Programs/Horde/Horde.Build.Tests/TestSetup.cs
@@ -88,7 +88,7 @@ namespace Horde.Build.Tests
 		public StreamService StreamService => ServiceProvider.GetRequiredService<StreamService>();
 		public ISubscriptionCollection SubscriptionCollection => ServiceProvider.GetRequiredService<ISubscriptionCollection>();
 		public INotificationService NotificationService => ServiceProvider.GetRequiredService<INotificationService>();
-		public IIssueService IssueService => ServiceProvider.GetRequiredService<IIssueService>();
+		public IssueService IssueService => ServiceProvider.GetRequiredService<IssueService>();
 		public JobTaskSource JobTaskSource => ServiceProvider.GetRequiredService<JobTaskSource>();
 		public JobService JobService => ServiceProvider.GetRequiredService<JobService>();
 		public IArtifactCollection ArtifactCollection => ServiceProvider.GetRequiredService<IArtifactCollection>();
@@ -186,7 +186,7 @@ namespace Horde.Build.Tests
 			services.AddSingleton<JobTaskSource>();
 			services.AddSingleton<IDowntimeService, DowntimeServiceStub>();
 			services.AddSingleton<IDogStatsd, NoOpDogStatsd>();
-			services.AddSingleton<IIssueService, IssueService>();
+			services.AddSingleton<IssueService>();
 			services.AddSingleton<IFleetManager, AwsFleetManager>();
 			services.AddSingleton<JobService>();
 			services.AddSingleton<LifetimeService>();
diff --git a/Engine/Source/Programs/Horde/Horde.Build/Controllers/IssuesController.cs b/Engine/Source/Programs/Horde/Horde.Build/Controllers/IssuesController.cs
--- a/Engine/Source/Programs/Horde/Horde.Build/Controllers/IssuesController.cs
+++ b/Engine/Source/Programs/Horde/Horde.Build/Controllers/IssuesController.cs
@@ -34,7 +34,7 @@ namespace Horde.Build.Controllers
 	public class IssuesController : HordeControllerBase
 	{
 		private readonly IIssueCollection _issueCollection;
-		private readonly IIssueService _issueService;
+		private readonly IssueService _issueService;
 		private readonly IExternalIssueService _externalIssueService;
 		private readonly JobService _jobService;
 		private readonly StreamService _streamService;
@@ -45,7 +45,7 @@ namespace Horde.Build.Controllers
 		/// <summary>
 		/// Constructor
 		/// </summary>
-		public IssuesController(ILogger<IssuesController> logger, IIssueCollection issueCollection, IIssueService issueService, JobService jobService, StreamService streamService, IUserCollection userCollection, ILogFileService logFileService, IExternalIssueService externalIssueService)
+		public IssuesController(ILogger<IssuesController> logger, IIssueCollection issueCollection, IssueService issueService, JobService jobService, StreamService streamService, IUserCollection userCollection, ILogFileService logFileService, IExternalIssueService externalIssueService)
 		{			
 			_issueCollection = issueCollection;
 			_issueService = issueService;
@@ -221,7 +221,7 @@ namespace Horde.Build.Controllers
 					streamIdValue = new StreamId(streamId);
 				}
 
-				issues = await _issueService.FindIssuesAsync(ids, userIdValue, streamIdValue, minChange ?? change, maxChange ?? change, resolved, promoted, index, count);
+				issues = await _issueService.Collection.FindIssuesAsync(ids, userIdValue, streamIdValue, minChange ?? change, maxChange ?? change, resolved, promoted, index, count);
 			}
 			else
 			{
@@ -236,7 +236,7 @@ namespace Horde.Build.Controllers
 				}
 
 				IGraph graph = await _jobService.GetGraphAsync(job);
-				issues = await _issueService.FindIssuesForJobAsync(ids, job, graph, stepId?.ToSubResourceId(), batchId?.ToSubResourceId(), labelIdx, userIdValue, resolved, promoted, index, count);
+				issues = await _issueService.Collection.FindIssuesForJobAsync(job, graph, stepId?.ToSubResourceId(), batchId?.ToSubResourceId(), labelIdx, userIdValue, resolved, promoted, index, count);
 			}
 
 			StreamPermissionsCache permissionsCache = new StreamPermissionsCache();
@@ -476,7 +476,8 @@ namespace Horde.Build.Controllers
 				logIdValues.UnionWith(logIds.Select(x => new LogId(x)));
 			}
 
-			List<ILogEvent> events = await _issueService.FindEventsForIssueAsync(issueId, logIdValues.ToArray(), index, count);
+			List<IIssueSpan> spans = await _issueCollection.FindSpansAsync(issueId);
+			List<ILogEvent> events = await _logFileService.FindEventsForSpansAsync(spans.Select(x => x.Id), logIdValues.ToArray(), index, count);
 
 			JobPermissionsCache permissionsCache = new JobPermissionsCache();
 			Dictionary<LogId, ILogFile?> logFiles = new Dictionary<LogId, ILogFile?>();
diff --git a/Engine/Source/Programs/Horde/Horde.Build/Controllers/UgsController.cs b/Engine/Source/Programs/Horde/Horde.Build/Controllers/UgsController.cs
--- a/Engine/Source/Programs/Horde/Horde.Build/Controllers/UgsController.cs
+++ b/Engine/Source/Programs/Horde/Horde.Build/Controllers/UgsController.cs
@@ -29,7 +29,7 @@ namespace Horde.Build.Controllers
 		/// <summary>
 		/// Singleton instance of the issue service
 		/// </summary>
-		private readonly IIssueService _issueService;
+		private readonly IssueService _issueService;
 
 		/// <summary>
 		/// Collection of metadata documents
@@ -59,7 +59,7 @@ namespace Horde.Build.Controllers
 		/// <summary>
 		/// Constructor
 		/// </summary>
-		public UgsController(IIssueService issueService, IUgsMetadataCollection ugsMetadataCollection, IUserCollection userCollection, ILogFileService logFileService, IOptionsMonitor<ServerSettings> optionsMonitor, ILogger<UgsController> logger)
+		public UgsController(IssueService issueService, IUgsMetadataCollection ugsMetadataCollection, IUserCollection userCollection, ILogFileService logFileService, IOptionsMonitor<ServerSettings> optionsMonitor, ILogger<UgsController> logger)
 		{
 			_issueService = issueService;
 			_ugsMetadataCollection = ugsMetadataCollection;
@@ -160,7 +160,7 @@ namespace Horde.Build.Controllers
 			List<GetUgsIssueResponse> responses = new List<GetUgsIssueResponse>();
 			if (includeResolved)
 			{
-				List<IIssue> issues = await _issueService.FindIssuesAsync(null, resolved: null, count: maxResults);
+				List<IIssue> issues = await _issueService.Collection.FindIssuesAsync(null, resolved: null, count: maxResults);
 				foreach(IIssue issue in issues)
 				{
 					IIssueDetails details = await _issueService.GetIssueDetailsAsync(issue);
@@ -258,7 +258,9 @@ namespace Horde.Build.Controllers
 
 			Dictionary<LogId, ILogFile?> logFiles = new Dictionary<LogId, ILogFile?>();
 
-			List<ILogEvent> events = await _issueService.FindEventsForIssueAsync(issueId, count: 10);
+			List<IIssueSpan> spans = await _issueService.Collection.FindSpansAsync(issueId);
+			List<ILogEvent> events = await _logFileService.FindEventsForSpansAsync(spans.Select(x => x.Id), null, 0, count: 10);
+
 			foreach (ILogEvent logEvent in events)
 			{
 				ILogFile? logFile;
diff --git a/Engine/Source/Programs/Horde/Horde.Build/Issues/External/Impl/JiraService.cs b/Engine/Source/Programs/Horde/Horde.Build/Issues/External/Impl/JiraService.cs
--- a/Engine/Source/Programs/Horde/Horde.Build/Issues/External/Impl/JiraService.cs
+++ b/Engine/Source/Programs/Horde/Horde.Build/Issues/External/Impl/JiraService.cs
@@ -117,7 +117,7 @@ namespace Horde.Build.Services
 		/// <summary>
 		/// Singleton instance of the project service
 		/// </summary>
-		readonly IIssueService _issueService;
+		readonly IssueService _issueService;
 
 		/// <summary>
 		/// Config collection
@@ -141,7 +141,7 @@ namespace Horde.Build.Services
 		/// <param name="configCollection"></param>
 		/// <param name="clock"></param>
 		/// <param name="logger"></param>
-		public JiraService(IOptions<ServerSettings> settings, StreamService streamService, IIssueService issueService, ConfigCollection configCollection, IClock clock, ILogger<JiraService> logger)
+		public JiraService(IOptions<ServerSettings> settings, StreamService streamService, IssueService issueService, ConfigCollection configCollection, IClock clock, ILogger<JiraService> logger)
 		{
 
 			_settings = settings.Value;
@@ -358,7 +358,7 @@ namespace Horde.Build.Services
 			HashSet<string> jiraKeys = new HashSet<string>();
 
 			// Refresh issues
-			List<IIssue> openIssues = await _issueService.FindIssuesAsync(resolved: false);
+			List<IIssue> openIssues = await _issueService.Collection.FindIssuesAsync(resolved: false);
 			for (int idx = 0; idx < openIssues.Count; idx++)
 			{
 				IIssue openIssue = openIssues[idx];
@@ -375,7 +375,7 @@ namespace Horde.Build.Services
 		public async Task<(string? key, string? url)> CreateIssueAsync(IUser user, string? externalIssueUser, int issueId, string summary, string projectId, string componentId, string issueType, string? description, string? hordeIssueLink)
 		{
 
-			IIssue? issue = await _issueService.GetIssueAsync(issueId);
+			IIssue? issue = await _issueService.Collection.GetIssueAsync(issueId);
 			if (issue == null)
 			{
 				throw new Exception($"Issue not found: {issueId}");
diff --git a/Engine/Source/Programs/Horde/Horde.Build/Issues/IIssueCollection.cs b/Engine/Source/Programs/Horde/Horde.Build/Issues/IIssueCollection.cs
--- a/Engine/Source/Programs/Horde/Horde.Build/Issues/IIssueCollection.cs
+++ b/Engine/Source/Programs/Horde/Horde.Build/Issues/IIssueCollection.cs
@@ -580,4 +580,58 @@ namespace Horde.Build.Collections
 		/// <returns>Logger for this issue</returns>
 		IAuditLogChannel<int> GetLogger(int issueId);
 	}
+
+	/// <summary>
+	/// Extension methods for issues
+	/// </summary>
+	public static class IssueCollectionExtensions
+	{
+		/// <summary>
+		/// Find steps for the given spans
+		/// </summary>
+		/// <param name="issueCollection"></param>
+		/// <param name="spanId">Span ids</param>
+		/// <returns>List of steps</returns>
+		public static Task<List<IIssueStep>> FindStepsAsync(this IIssueCollection issueCollection, ObjectId spanId)
+		{
+			return issueCollection.FindStepsAsync(new[] { spanId });
+		}
+
+		/// <summary>
+		/// Searches for open issues affecting a job
+		/// </summary>
+		/// <param name="issueCollection"></param>
+		/// <param name="job"></param>
+		/// <param name="graph"></param>
+		/// <param name="stepId"></param>
+		/// <param name="batchId"></param>
+		/// <param name="labelIdx"></param>
+		/// <param name="userId"></param>
+		/// <param name="resolved">Whether to include results that are resolved</param>
+		/// <param name="promoted">Whether to filter by promoted issues</param>
+		/// <param name="index">Index within the results to return</param>
+		/// <param name="count">Number of results</param>
+		/// <returns></returns>
+		public static async Task<List<IIssue>> FindIssuesForJobAsync(this IIssueCollection issueCollection, IJob job, IGraph graph, SubResourceId? stepId = null, SubResourceId? batchId = null, int? labelIdx = null, UserId? userId = null, bool? resolved = null, bool? promoted = null, int? index = null, int? count = null)
+		{
+			List<IIssueStep> steps = await issueCollection.FindStepsAsync(job.Id, batchId, stepId);
+			List<IIssueSpan> spans = await issueCollection.FindSpansAsync(steps.Select(x => x.SpanId));
+
+			if (labelIdx != null)
+			{
+				HashSet<string> nodeNames = new HashSet<string>(job.GetNodesForLabel(graph, labelIdx.Value).Select(x => graph.GetNode(x).Name));
+				spans.RemoveAll(x => !nodeNames.Contains(x.NodeName));
+			}
+
+			List<int> issueIds = new List<int>(spans.Select(x => x.IssueId));
+			if (issueIds.Count == 0)
+			{
+				return new List<IIssue>();
+			}
+
+			return await issueCollection.FindIssuesAsync(ids: issueIds, userId: userId, resolved: resolved, promoted: promoted, index: index, count: count);
+		}
+
+
+	}
 }
diff --git a/Engine/Source/Programs/Horde/Horde.Build/Issues/IIssueDetails.cs b/Engine/Source/Programs/Horde/Horde.Build/Issues/IIssueDetails.cs
--- a/Engine/Source/Programs/Horde/Horde.Build/Issues/IIssueDetails.cs
+++ b/Engine/Source/Programs/Horde/Horde.Build/Issues/IIssueDetails.cs
@@ -94,9 +94,9 @@ namespace Horde.Build.Models
 		/// <param name="issueService">The issue service</param>
 		/// <param name="issueId">Issue id to query </param>
 		/// <returns></returns>
-		public static async Task<IIssueDetails?> GetIssueDetailsAsync(this IIssueService issueService, int issueId)
+		public static async Task<IIssueDetails?> GetIssueDetailsAsync(this IssueService issueService, int issueId)
 		{
-			IIssue? issue = await issueService.GetIssueAsync(issueId);
+			IIssue? issue = await issueService.Collection.GetIssueAsync(issueId);
 			if(issue == null)
 			{
 				return null;
diff --git a/Engine/Source/Programs/Horde/Horde.Build/Issues/IIssueService.cs b/Engine/Source/Programs/Horde/Horde.Build/Issues/IIssueService.cs
deleted file mode 100644
diff --git a/Engine/Source/Programs/Horde/Horde.Build/Issues/Impl/IssueService.cs b/Engine/Source/Programs/Horde/Horde.Build/Issues/Impl/IssueService.cs
--- a/Engine/Source/Programs/Horde/Horde.Build/Issues/Impl/IssueService.cs
+++ b/Engine/Source/Programs/Horde/Horde.Build/Issues/Impl/IssueService.cs
@@ -26,7 +26,7 @@ using MongoDB.Driver;
 using OpenTracing;
 using OpenTracing.Util;
 
-namespace Horde.Build.Services.Impl
+namespace Horde.Build.Services
 {
 	using LogId = ObjectId<ILogFile>;
 	using StreamId = StringId<IStream>;
@@ -141,7 +141,7 @@ namespace Horde.Build.Services.Impl
 	/// <summary>
 	/// Wraps funtionality for manipulating build health issues
 	/// </summary>
-	public sealed class IssueService : IIssueService, IHostedService, IDisposable
+	public sealed class IssueService : IHostedService, IDisposable
 	{
 		/// <summary>
 		/// Maximum number of changes to query from Perforce in one go
@@ -158,10 +158,20 @@ namespace Horde.Build.Services.Impl
 		readonly IClock _clock;
 		readonly ITicker _ticker;
 
+		/// <summary>
+		/// Accessor for the issue collection
+		/// </summary>
+		public IIssueCollection Collection => _issueCollection;
+
+		/// <summary>
+		/// Delegate for issue creation events
+		/// </summary>
+		public delegate void IssueUpdatedEvent(IIssue issue);
+
 		/// <summary>
 		/// 
 		/// </summary>
-		public event IIssueService.IssueUpdatedEvent? OnIssueUpdated;
+		public event IssueUpdatedEvent? OnIssueUpdated;
 
 		/// <summary>
 		/// List of issue matchers
@@ -262,7 +272,7 @@ namespace Horde.Build.Services.Impl
 			_cachedDesktopAlerts = newCachedDesktopAlerts;
 
 			// Resolve any issues that haven't been seen in a week
-			List<IIssue> openIssues = await FindIssuesAsync(resolved: false);
+			List<IIssue> openIssues = await _issueCollection.FindIssuesAsync(resolved: false);
 			for (int idx = 0; idx < openIssues.Count; idx++)
 			{
 				IIssue openIssue = openIssues[idx];
@@ -329,51 +339,6 @@ namespace Horde.Build.Services.Impl
 			return new IssueDetails(issue, owner, nominatedBy, resolvedBy, quarantinedBy, spans, steps, suspects, suspectUsers, ShowDesktopAlertsForIssue(issue, spans), issue.ExternalIssueKey);
 		}
 
-		/// <inheritdoc/>
-		public Task<List<IIssueSpan>> GetIssueSpansAsync(IIssue issue)
-		{
-			return _issueCollection.FindSpansAsync(issue.Id);
-		}
-
-		/// <inheritdoc/>
-		public Task<List<IIssueStep>> GetIssueStepsAsync(IIssueSpan span)
-		{
-			return _issueCollection.FindStepsAsync(new[] { span.Id });
-		}
-
-		/// <inheritdoc/>
-		public Task<List<IIssueSuspect>> GetIssueSuspectsAsync(IIssue issue)
-		{
-			return _issueCollection.FindSuspectsAsync(issue);
-		}
-
-		/// <inheritdoc/>
-		public async Task<List<IIssue>> FindIssuesAsync(IEnumerable<int>? ids = null, UserId? userId = null, StreamId? streamId = null, int? minChange = null, int? maxChange = null, bool? resolved = null, bool? promoted = null, int? index = null, int? count = null)
-		{
-			return await _issueCollection.FindIssuesAsync(ids, userId, streamId, minChange, maxChange, resolved, promoted, index, count);
-		}
-
-		/// <inheritdoc/>
-		public async Task<List<IIssue>> FindIssuesForJobAsync(int[]? ids, IJob job, IGraph graph, SubResourceId? stepId = null, SubResourceId? batchId = null, int? labelIdx = null, UserId? userId = null, bool? resolved = null, bool? promoted = null, int? index = null, int? count = null)
-		{
-			List<IIssueStep> steps = await _issueCollection.FindStepsAsync(job.Id, batchId, stepId);
-			List<IIssueSpan> spans = await _issueCollection.FindSpansAsync(steps.Select(x => x.SpanId));
-
-			if (labelIdx != null)
-			{
-				HashSet<string> nodeNames = new HashSet<string>(job.GetNodesForLabel(graph, labelIdx.Value).Select(x => graph.GetNode(x).Name));
-				spans.RemoveAll(x => !nodeNames.Contains(x.NodeName));
-			}
-
-			List<int> issueIds = new List<int>(spans.Select(x => x.IssueId));
-			if(issueIds.Count == 0)
-			{
-				return new List<IIssue>();
-			}
-
-			return await FindIssuesAsync(ids: issueIds, userId: userId, resolved: resolved, promoted: promoted, index: index, count: count);
-		}
-
 		/// <inheritdoc/>
 		public bool ShowDesktopAlertsForIssue(IIssue issue, IReadOnlyList<IIssueSpan> spans)
 		{
@@ -397,17 +362,24 @@ namespace Horde.Build.Services.Impl
 		}
 
 		/// <summary>
-		/// Gets an issue with the given id
+		/// Updates the state of an issue
 		/// </summary>
-		/// <param name="issueId"></param>
-		/// <returns></returns>
-		public Task<IIssue?> GetIssueAsync(int issueId)
-		{
-			return _issueCollection.GetIssueAsync(issueId);
-		}
-
-		/// <inheritdoc/>
-		public async Task<bool> UpdateIssueAsync(int id, string? userSummary = null, string? description = null, bool? promoted = null, UserId? ownerId = null, UserId? nominatedById = null, bool? acknowledged = null, UserId? declinedById = null, int? fixChange = null, UserId? resolvedById = null, List<ObjectId>? addSpanIds = null, List<ObjectId>? removeSpanIds = null, string? externalIssueKey = null, UserId? quarantinedById = null)
+		/// <param name="id">The current issue id</param>s
+		/// <param name="summary">New summary for the issue</param>
+		/// <param name="description">New description for the issue</param>
+		/// <param name="promoted">Whether the issue should be set as promoted</param>
+		/// <param name="ownerId">New owner of the issue</param>
+		/// <param name="nominatedById">Person that nominated the new owner</param>
+		/// <param name="acknowledged">Whether the issue has been acknowledged</param>
+		/// <param name="declinedById">Name of the user that has declined this issue</param>
+		/// <param name="fixChange">Fix changelist for the issue</param>
+		/// <param name="resolvedById">Whether the issue has been resolved</param>
+		/// <param name="addSpanIds">Add spans to this issue</param>
+		/// <param name="removeSpanIds">Remove spans from this issue</param>
+		/// <param name="externalIssueKey">Key for external issue tracking</param>
+		/// <param name="quarantinedById">User who has quarantined the issue</param>
+		/// <returns>True if the issue was updated</returns>
+		public async Task<bool> UpdateIssueAsync(int id, string? summary = null, string? description = null, bool? promoted = null, UserId? ownerId = null, UserId? nominatedById = null, bool? acknowledged = null, UserId? declinedById = null, int? fixChange = null, UserId? resolvedById = null, List<ObjectId>? addSpanIds = null, List<ObjectId>? removeSpanIds = null, string? externalIssueKey = null, UserId? quarantinedById = null)
 		{
 			IIssue? issue;
 			for (; ; )
@@ -418,7 +390,7 @@ namespace Horde.Build.Services.Impl
 					return false;
 				}
 
-				issue = await _issueCollection.TryUpdateIssueAsync(issue, newUserSummary: userSummary, newDescription: description, newPromoted: promoted, newOwnerId: ownerId ?? resolvedById, newNominatedById: nominatedById, newDeclinedById: declinedById, newAcknowledged: acknowledged, newFixChange: fixChange, newResolvedById: resolvedById, newExcludeSpanIds: removeSpanIds, newExternalIssueKey: externalIssueKey, newQuarantinedById: quarantinedById);
+				issue = await _issueCollection.TryUpdateIssueAsync(issue, newUserSummary: summary, newDescription: description, newPromoted: promoted, newOwnerId: ownerId ?? resolvedById, newNominatedById: nominatedById, newDeclinedById: declinedById, newAcknowledged: acknowledged, newFixChange: fixChange, newResolvedById: resolvedById, newExcludeSpanIds: removeSpanIds, newExternalIssueKey: externalIssueKey, newQuarantinedById: quarantinedById);
 				if (issue != null)
 				{
 					break;
@@ -1371,12 +1343,5 @@ namespace Horde.Build.Services.Impl
 				return null;
 			}
 		}
-
-		/// <inheritdoc/>
-		public async Task<List<ILogEvent>> FindEventsForIssueAsync(int issueId, LogId[]? logIds, int index, int count)
-		{
-			List<IIssueSpan> spans = await _issueCollection.FindSpansAsync(issueId);
-			return await _logFileService.FindEventsForSpansAsync(spans.Select(x => x.Id), logIds, index, count);
-		}
 	}
 }
diff --git a/Engine/Source/Programs/Horde/Horde.Build/Jobs/Impl/JobService.cs b/Engine/Source/Programs/Horde/Horde.Build/Jobs/Impl/JobService.cs
--- a/Engine/Source/Programs/Horde/Horde.Build/Jobs/Impl/JobService.cs
+++ b/Engine/Source/Programs/Horde/Horde.Build/Jobs/Impl/JobService.cs
@@ -78,7 +78,7 @@ namespace Horde.Build.Services
 		readonly JobTaskSource _jobTaskSource;
 		readonly StreamService _streamService;
 		readonly ITemplateCollection _templateCollection;
-		readonly IIssueService? _issueService;
+		readonly IssueService? _issueService;
 		readonly IPerforceService _perforceService;
 		readonly ILogger _logger;
 
@@ -125,7 +125,7 @@ namespace Horde.Build.Services
 		/// <param name="issueService">The issue service</param>
 		/// <param name="perforceService">The perforce service</param>
 		/// <param name="logger">Log output</param>
-		public JobService(IJobCollection jobs, IGraphCollection graphs, IAgentCollection agents, IJobStepRefCollection jobStepRefs, IJobTimingCollection jobTimings, IUserCollection userCollection, INotificationTriggerCollection triggerCollection, JobTaskSource jobTaskSource, StreamService streamService, ITemplateCollection templateCollection, IIssueService issueService, IPerforceService perforceService, ILogger<JobService> logger)
+		public JobService(IJobCollection jobs, IGraphCollection graphs, IAgentCollection agents, IJobStepRefCollection jobStepRefs, IJobTimingCollection jobTimings, IUserCollection userCollection, INotificationTriggerCollection triggerCollection, JobTaskSource jobTaskSource, StreamService streamService, ITemplateCollection templateCollection, IssueService issueService, IPerforceService perforceService, ILogger<JobService> logger)
 		{
 			_jobs = jobs;
 			_graphs = graphs;
diff --git a/Engine/Source/Programs/Horde/Horde.Build/Notifications/Impl/NotificationService.cs b/Engine/Source/Programs/Horde/Horde.Build/Notifications/Impl/NotificationService.cs
--- a/Engine/Source/Programs/Horde/Horde.Build/Notifications/Impl/NotificationService.cs
+++ b/Engine/Source/Programs/Horde/Horde.Build/Notifications/Impl/NotificationService.cs
@@ -72,7 +72,7 @@ namespace Horde.Build.Notifications.Impl
 		/// <summary>
 		/// 
 		/// </summary>
-		private readonly IIssueService _issueService;
+		private readonly IssueService _issueService;
 
 		/// <summary>
 		/// Instance of the <see cref="_logFileService"/>.
@@ -124,7 +124,7 @@ namespace Horde.Build.Notifications.Impl
 		/// <summary>
 		/// Constructor
 		/// </summary>
-		public NotificationService(IEnumerable<INotificationSink> sinks, IOptionsMonitor<ServerSettings> settings, ILogger<NotificationService> logger, IGraphCollection graphCollection, ISubscriptionCollection subscriptionCollection, INotificationTriggerCollection triggerCollection, IUserCollection userCollection, JobService jobService, StreamService streamService, IIssueService issueService, ILogFileService logFileService, IDogStatsd dogStatsd, RedisService redisService, IClock clock)
+		public NotificationService(IEnumerable<INotificationSink> sinks, IOptionsMonitor<ServerSettings> settings, ILogger<NotificationService> logger, IGraphCollection graphCollection, ISubscriptionCollection subscriptionCollection, INotificationTriggerCollection triggerCollection, IUserCollection userCollection, JobService jobService, StreamService streamService, IssueService issueService, ILogFileService logFileService, IDogStatsd dogStatsd, RedisService redisService, IClock clock)
 		{
 			_sinks = sinks.ToList();
 			_settings = settings;
diff --git a/Engine/Source/Programs/Horde/Horde.Build/Notifications/Impl/SlackNotificationSink.cs b/Engine/Source/Programs/Horde/Horde.Build/Notifications/Impl/SlackNotificationSink.cs
--- a/Engine/Source/Programs/Horde/Horde.Build/Notifications/Impl/SlackNotificationSink.cs
+++ b/Engine/Source/Programs/Horde/Horde.Build/Notifications/Impl/SlackNotificationSink.cs
@@ -219,7 +219,7 @@ namespace Horde.Build.Notifications.Impl
 			}
 		}
 
-		readonly IIssueService _issueService;
+		readonly IssueService _issueService;
 		readonly IUserCollection _userCollection;
 		readonly ILogFileService _logFileService;
 		readonly StreamService _streamService;
@@ -240,7 +240,7 @@ namespace Horde.Build.Notifications.Impl
 		/// <summary>
 		/// Constructor
 		/// </summary>
-		public SlackNotificationSink(MongoService mongoService, IIssueService issueService, IUserCollection userCollection, ILogFileService logFileService, StreamService streamService, ConfigCollection configCollection, IExternalIssueService externalIssueService, IWebHostEnvironment environment, IOptions<ServerSettings> settings, ILogger<SlackNotificationSink> logger)
+		public SlackNotificationSink(MongoService mongoService, IssueService issueService, IUserCollection userCollection, ILogFileService logFileService, StreamService streamService, ConfigCollection configCollection, IExternalIssueService externalIssueService, IWebHostEnvironment environment, IOptions<ServerSettings> settings, ILogger<SlackNotificationSink> logger)
 		{
 			_issueService = issueService;
 			_userCollection = userCollection;
@@ -801,7 +801,7 @@ namespace Horde.Build.Notifications.Impl
 					}
 					else
 					{
-						List<IIssueSuspect> suspects = await _issueService.GetIssueSuspectsAsync(issue);
+						List<IIssueSuspect> suspects = await _issueService.Collection.FindSuspectsAsync(issue);
 						IGrouping<UserId, IIssueSuspect>[] suspectGroups = suspects.GroupBy(x => x.AuthorId).ToArray();
 
 						if (suspectGroups.Length > 0 && suspectGroups.Length <= workflow.MaxMentions)
@@ -1198,7 +1198,7 @@ namespace Horde.Build.Notifications.Impl
 				await _issueService.UpdateIssueAsync(issueId, declinedById: userId);
 			}
 
-			IIssue? newIssue = await _issueService.GetIssueAsync(issueId);
+			IIssue? newIssue = await _issueService.Collection.GetIssueAsync(issueId);
 			if (newIssue != null)
 			{
 				IUser? user = await _userCollection.GetUserAsync(userId);
diff --git a/Engine/Source/Programs/Horde/Horde.Build/Startup.cs b/Engine/Source/Programs/Horde/Horde.Build/Startup.cs
--- a/Engine/Source/Programs/Horde/Horde.Build/Startup.cs
+++ b/Engine/Source/Programs/Horde/Horde.Build/Startup.cs
@@ -387,7 +387,7 @@ namespace Horde.Build
 			services.AddSingleton<ICommitService>(sp => sp.GetRequiredService<CommitService>());
 			services.AddSingleton<IClock, Clock>();
 			services.AddSingleton<IDowntimeService, DowntimeService>();
-			services.AddSingleton<IIssueService, IssueService>();
+			services.AddSingleton<IssueService>();
 			services.AddSingleton<JobService>();
 			services.AddSingleton<LifetimeService>();
 			services.AddSingleton<ILogFileService, LogFileService>();
@@ -589,7 +589,7 @@ namespace Horde.Build
 				services.AddHostedService(provider => provider.GetRequiredService<AgentService>());
 				services.AddHostedService(provider => provider.GetRequiredService<CommitService>());
 				services.AddHostedService(provider => provider.GetRequiredService<ConsistencyService>());
-				services.AddHostedService(provider => provider.GetRequiredService<IIssueService>());
+				services.AddHostedService(provider => provider.GetRequiredService<IssueService>());
 				services.AddHostedService<IssueReportService>();
 				services.AddHostedService(provider => (LogFileService)provider.GetRequiredService<ILogFileService>());
 				services.AddHostedService(provider => (NotificationService)provider.GetRequiredService<INotificationService>());
