Horde: Finish re-org of Horde solution started a while ago.

* Controllers and messages are now stored alongside their related classes.
* "Impl" folders have been removed to allow more logical grouping of code by subject area rather than implementation details.

diff --git a/Engine/Source/Programs/Horde/Horde.Build/Acls/AclMessages.cs b/Engine/Source/Programs/Horde/Horde.Build/Acls/AclMessages.cs
new file mode 100644
--- /dev/null
+++ b/Engine/Source/Programs/Horde/Horde.Build/Acls/AclMessages.cs
@@ -0,0 +1,160 @@
+﻿// Copyright Epic Games, Inc. All Rights Reserved.
+
+using System.Collections.Generic;
+using System.ComponentModel.DataAnnotations;
+using Horde.Build.Acls;
+
+namespace Horde.Build.Api
+{
+	/// <summary>
+	/// New claim to create
+	/// </summary>
+	public class CreateAclClaimRequest
+	{
+		/// <summary>
+		/// The claim type
+		/// </summary>
+		[Required]
+		public string Type { get; set; } = null!;
+
+		/// <summary>
+		/// The claim value
+		/// </summary>
+		[Required]
+		public string Value { get; set; } = null!;
+	}
+
+	/// <summary>
+	/// Individual entry in an ACL
+	/// </summary>
+	public class CreateAclEntryRequest
+	{
+		/// <summary>
+		/// Name of the user or group
+		/// </summary>
+		[Required]
+		public CreateAclClaimRequest Claim { get; set; } = null!;
+
+		/// <summary>
+		/// Array of actions to allow
+		/// </summary>
+		[Required]
+		public string[] Actions { get; set; } = null!;
+	}
+
+	/// <summary>
+	/// Parameters to update an ACL
+	/// </summary>
+	public class UpdateAclRequest
+	{
+		/// <summary>
+		/// Entries to replace the existing ACL
+		/// </summary>
+		public List<CreateAclEntryRequest>? Entries { get; set; }
+
+		/// <summary>
+		/// Whether to inherit permissions from the parent ACL
+		/// </summary>
+		public bool? Inherit { get; set; }
+
+		/// <summary>
+		/// List of exceptions to the inherited setting
+		/// </summary>
+		public List<string>? Exceptions { get; set; }
+	}
+
+	/// <summary>
+	/// New claim to update
+	/// </summary>
+	public class GetAclClaimResponse
+	{
+		/// <summary>
+		/// The claim type
+		/// </summary>
+		public string Type { get; set; }
+
+		/// <summary>
+		/// The claim value
+		/// </summary>
+		public string Value { get; set; }
+
+		/// <summary>
+		/// Constructor
+		/// </summary>
+		/// <param name="claim">The claim to construct from</param>
+		public GetAclClaimResponse(AclClaim claim)
+			: this(claim.Type, claim.Value)
+		{
+		}
+
+		/// <summary>
+		/// Constructor
+		/// </summary>
+		/// <param name="type">The claim type</param>
+		/// <param name="value">The claim value</param>
+		public GetAclClaimResponse(string type, string value)
+		{
+			Type = type;
+			Value = value;
+		}
+	}
+
+	/// <summary>
+	/// Individual entry in an ACL
+	/// </summary>
+	public class GetAclEntryResponse
+	{
+		/// <summary>
+		/// Names of the user or group
+		/// </summary>
+		public GetAclClaimResponse Claim { get; set; }
+
+		/// <summary>
+		/// Array of actions to allow
+		/// </summary>
+		public List<string> Actions { get; set; }
+
+		/// <summary>
+		/// Constructor
+		/// </summary>
+		/// <param name="aclEntry">The acl entry to construct from</param>
+		public GetAclEntryResponse(AclEntry aclEntry)
+		{
+			Claim = new GetAclClaimResponse(aclEntry.Claim.Type, aclEntry.Claim.Value);
+			Actions = AclEntry.GetActionNames(aclEntry.Actions);
+		}
+	}
+
+	/// <summary>
+	/// Information about an ACL
+	/// </summary>
+	public class GetAclResponse
+	{
+		/// <summary>
+		/// Entries to replace the existing ACL
+		/// </summary>
+		[Required]
+		public List<GetAclEntryResponse> Entries { get; set; }
+
+		/// <summary>
+		/// Whether to inherit permissions from the parent entity
+		/// </summary>
+		public bool Inherit { get; set; }
+
+		/// <summary>
+		/// Exceptions from permission inheritance setting
+		/// </summary>
+		public List<string>? Exceptions { get; set; }
+
+		/// <summary>
+		/// Constructor
+		/// </summary>
+		/// <param name="acl">The ACL to construct from</param>
+		public GetAclResponse(Acl acl)
+		{
+			Entries = acl.Entries.ConvertAll(x => new GetAclEntryResponse(x));
+			Inherit = acl.Inherit;
+			Exceptions = acl.Exceptions?.ConvertAll(x => x.ToString());
+		}
+	}
+}
diff --git a/Engine/Source/Programs/Horde/Horde.Build/Agents/AgentCollection.cs b/Engine/Source/Programs/Horde/Horde.Build/Agents/AgentCollection.cs
new file mode 100644
--- /dev/null
+++ b/Engine/Source/Programs/Horde/Horde.Build/Agents/AgentCollection.cs
@@ -0,0 +1,574 @@
+// Copyright Epic Games, Inc. All Rights Reserved.
+
+using System;
+using System.Collections.Generic;
+using System.Linq;
+using System.Linq.Expressions;
+using System.Threading.Tasks;
+using EpicGames.Redis;
+using Google.Protobuf.WellKnownTypes;
+using Horde.Build.Acls;
+using Horde.Build.Models;
+using Horde.Build.Server;
+using Horde.Build.Services;
+using Horde.Build.Utilities;
+using HordeCommon;
+using HordeCommon.Rpc.Tasks;
+using MongoDB.Bson;
+using MongoDB.Bson.Serialization.Attributes;
+using MongoDB.Driver;
+
+namespace Horde.Build.Collections.Impl
+{
+	using AgentSoftwareChannelName = StringId<AgentSoftwareChannels>;
+	using PoolId = StringId<IPool>;
+	using SessionId = ObjectId<ISession>;
+
+	/// <summary>
+	/// Collection of agent documents
+	/// </summary>
+	public class AgentCollection : IAgentCollection
+	{
+		/// <summary>
+		/// Legacy information about a device attached to an agent
+		/// </summary>
+		class DeviceCapabilities
+		{
+			[BsonIgnoreIfNull]
+			public HashSet<string>? Properties { get; set; }
+
+			[BsonIgnoreIfNull]
+			public Dictionary<string, int>? Resources { get; set; }
+		}
+
+		/// <summary>
+		/// Legacy capabilities of an agent
+		/// </summary>
+		class AgentCapabilities
+		{
+			public List<DeviceCapabilities> Devices { get; set; } = new List<DeviceCapabilities>();
+
+			[BsonIgnoreIfNull]
+			public HashSet<string>? Properties { get; set; }
+		}
+
+		/// <summary>
+		/// Concrete implementation of an agent document
+		/// </summary>
+		class AgentDocument : IAgent
+		{
+			static readonly IReadOnlyList<AgentLease> s_emptyLeases = new List<AgentLease>();
+
+			[BsonRequired, BsonId]
+			public AgentId Id { get; set; }
+
+			public SessionId? SessionId { get; set; }
+			public DateTime? SessionExpiresAt { get; set; }
+
+			public AgentStatus Status { get; set; }
+
+			[BsonRequired]
+			public bool Enabled { get; set; }
+
+			public bool Ephemeral { get; set; }
+
+			[BsonIgnoreIfDefault, BsonDefaultValue(false)]
+			public bool Deleted { get; set; }
+
+			[BsonElement("Version2")]
+			public string? Version { get; set; }
+
+			public List<string>? Properties { get; set; }
+			public Dictionary<string, int>? Resources { get; set; }
+
+			[BsonIgnoreIfNull]
+			public AgentSoftwareChannelName? Channel { get; set; }
+
+			[BsonIgnoreIfNull]
+			public string? LastUpgradeVersion { get; set; }
+
+			[BsonIgnoreIfNull]
+			public DateTime? LastUpgradeTime { get; set; }
+
+			public List<PoolId> DynamicPools { get; set; } = new List<PoolId>();
+			public List<PoolId> Pools { get; set; } = new List<PoolId>();
+
+			[BsonIgnoreIfDefault, BsonDefaultValue(false)]
+			public bool RequestConform { get; set; }
+
+			[BsonIgnoreIfDefault, BsonDefaultValue(false)]
+			public bool RequestFullConform { get; set; }
+
+			[BsonIgnoreIfDefault, BsonDefaultValue(false)]
+			public bool RequestRestart { get; set; }
+
+			[BsonIgnoreIfDefault, BsonDefaultValue(false)]
+			public bool RequestShutdown { get; set; }
+
+			[BsonIgnoreIfNull]
+			public string? LastShutdownReason { get; set; }
+
+			public List<AgentWorkspace> Workspaces { get; set; } = new List<AgentWorkspace>();
+			public DateTime LastConformTime { get; set; }
+
+			[BsonIgnoreIfNull]
+			public int? ConformAttemptCount { get; set; }
+
+			public AgentCapabilities Capabilities { get; set; } = new AgentCapabilities();
+			public List<AgentLease>? Leases { get; set; }
+			public Acl? Acl { get; set; }
+			public DateTime UpdateTime { get; set; }
+			public uint UpdateIndex { get; set; }
+			public string? Comment { get; set; }
+
+			IReadOnlyList<PoolId> IAgent.DynamicPools => DynamicPools;
+			IReadOnlyList<PoolId> IAgent.ExplicitPools => Pools;
+			IReadOnlyList<AgentWorkspace> IAgent.Workspaces => Workspaces;
+			IReadOnlyList<AgentLease> IAgent.Leases => Leases ?? s_emptyLeases;
+			IReadOnlyList<string> IAgent.Properties => Properties ?? Capabilities.Devices.FirstOrDefault()?.Properties?.ToList() ?? new List<string>();
+			IReadOnlyDictionary<string, int> IAgent.Resources => Resources ?? Capabilities.Devices.FirstOrDefault()?.Resources ?? new Dictionary<string, int>();
+
+			[BsonConstructor]
+			private AgentDocument()
+			{
+			}
+
+			public AgentDocument(AgentId id, bool bEnabled, AgentSoftwareChannelName? channel, List<PoolId> pools)
+			{
+				Id = id;
+				Acl = new Acl();
+				Enabled = bEnabled;
+				Channel = channel;
+				Pools = pools;
+			}
+		}
+
+		readonly IMongoCollection<AgentDocument> _agents;
+		readonly IAuditLog<AgentId> _auditLog;
+		readonly RedisService _redisService;
+		readonly RedisChannel<AgentId> _updateEventChannel;
+
+		/// <summary>
+		/// Constructor
+		/// </summary>
+		public AgentCollection(MongoService mongoService, RedisService redisService, IAuditLog<AgentId> auditLog)
+		{
+			List<MongoIndex<AgentDocument>> indexes = new List<MongoIndex<AgentDocument>>();
+			indexes.Add(keys => keys.Ascending(x => x.Deleted).Ascending(x => x.Id).Ascending(x => x.Pools));
+
+			_agents = mongoService.GetCollection<AgentDocument>("Agents", indexes);
+			_redisService = redisService;
+			_updateEventChannel = new RedisChannel<AgentId>("agents/notify");
+			_auditLog = auditLog;
+		}
+
+		/// <inheritdoc/>
+		public async Task<IAgent> AddAsync(AgentId id, bool bEnabled, AgentSoftwareChannelName? channel, List<PoolId>? pools)
+		{
+			AgentDocument agent = new AgentDocument(id, bEnabled, channel, pools ?? new List<PoolId>());
+			await _agents.InsertOneAsync(agent);
+			return agent;
+		}
+
+		/// <inheritdoc/>
+		public async Task<IAgent?> TryDeleteAsync(IAgent agentInterface)
+		{
+			AgentDocument agent = (AgentDocument)agentInterface;
+
+			UpdateDefinition<AgentDocument> update = Builders<AgentDocument>.Update.Set(x => x.Deleted, true);
+			return await TryUpdateAsyncWithRetries(agent, update);
+		}
+
+		/// <inheritdoc/>
+		public async Task ForceDeleteAsync(AgentId agentId)
+		{
+			await _agents.DeleteOneAsync(x => x.Id == agentId);
+		}
+
+		/// <inheritdoc/>
+		public async Task<IAgent?> GetAsync(AgentId agentId)
+		{
+			return await _agents.Find<AgentDocument>(x => x.Id == agentId).FirstOrDefaultAsync();
+		}
+
+		/// <inheritdoc/>
+		public async Task<List<IAgent>> FindAsync(PoolId? poolId, DateTime? modifiedAfter, AgentStatus? status, int? index, int? count)
+		{
+			FilterDefinitionBuilder<AgentDocument> filterBuilder = new FilterDefinitionBuilder<AgentDocument>();
+
+			FilterDefinition<AgentDocument> filter = filterBuilder.Ne(x => x.Deleted, true);
+			
+			if (poolId != null)
+			{
+				filter &= filterBuilder.Eq(nameof(AgentDocument.Pools), poolId);
+			}
+			
+			if (modifiedAfter != null)
+			{
+				filter &= filterBuilder.Gt(x => x.UpdateTime, modifiedAfter.Value);
+			}
+			
+			if (status != null)
+			{
+				filter &= filterBuilder.Eq(x => x.Status, status.Value);
+			}
+
+			IFindFluent<AgentDocument, AgentDocument> search = _agents.Find(filter);
+			if (index != null)
+			{
+				search = search.Skip(index.Value);
+			}
+			if (count != null)
+			{
+				search = search.Limit(count.Value);
+			}
+
+			List<AgentDocument> results = await search.ToListAsync();
+			return results.ConvertAll<IAgent>(x => x);
+		}
+
+		/// <inheritdoc/>
+		public async Task<List<IAgent>> FindExpiredAsync(DateTime utcNow, int maxAgents)
+		{
+			List<AgentDocument> results = await _agents.Find(x => x.SessionId.HasValue && !(x.SessionExpiresAt > utcNow)).Limit(maxAgents).ToListAsync();
+			return results.ConvertAll<IAgent>(x => x);
+		}
+
+		/// <summary>
+		/// Update a single document
+		/// </summary>
+		/// <param name="current">The document to update</param>
+		/// <param name="update">The update definition</param>
+		/// <returns>The updated agent document or null if update failed</returns>
+		private async Task<AgentDocument?> TryUpdateAsync(AgentDocument current, UpdateDefinition<AgentDocument> update)
+		{
+			uint prevUpdateIndex = current.UpdateIndex++;
+			current.UpdateTime = DateTime.UtcNow;
+
+			Expression<Func<AgentDocument, bool>> filter = x => x.Id == current.Id && x.UpdateIndex == prevUpdateIndex;
+			UpdateDefinition<AgentDocument> updateWithIndex = update.Set(x => x.UpdateIndex, current.UpdateIndex).Set(x => x.UpdateTime, current.UpdateTime);
+
+			return await _agents.FindOneAndUpdateAsync<AgentDocument>(filter, updateWithIndex, new FindOneAndUpdateOptions<AgentDocument, AgentDocument> { ReturnDocument = ReturnDocument.After });
+		}
+
+		/// <summary>
+		/// Update a single document, with retries to ensure it's applied
+		/// If an outdated AgentDocument is supplied, it needs to be refreshed so that the UpdateIndex matches.
+		/// </summary>
+		/// <param name="agentDoc">The document to update</param>
+		/// <param name="update">The update definition</param>
+		/// <param name="numRetries">Number of retries to attempt</param>
+		/// <returns>The updated agent document or null if update failed</returns>
+		private async Task<AgentDocument?> TryUpdateAsyncWithRetries(AgentDocument agentDoc, UpdateDefinition<AgentDocument> update, int numRetries = 5)
+		{
+			AgentDocument current = agentDoc;
+			for (int retryCount = 0; retryCount < numRetries; retryCount++)
+			{
+				AgentDocument? result = await TryUpdateAsync(current, update);
+				if (result != null)
+				{
+					return result;
+				}
+
+				AgentId id = current.Id; // Capture id to avoid doing it inside closure (and warnings)
+				current = await _agents.Find(x => x.Id == id).FirstOrDefaultAsync();
+				if (current == null)
+				{
+					return null;
+				}
+			}
+
+			return null;
+		}
+
+		/// <inheritdoc/>
+		public async Task<IAgent?> TryUpdateSettingsAsync(IAgent agentInterface, bool? enabled = null, bool? requestConform = null, bool? requestFullConform = null, bool? requestRestart = null, bool? requestShutdown = null, string? shutdownReason = null, AgentSoftwareChannelName? channel = null, List<PoolId>? pools = null, Acl? acl = null, string? comment = null)
+		{
+			AgentDocument agent = (AgentDocument)agentInterface;
+
+			// Update the database
+			UpdateDefinitionBuilder<AgentDocument> updateBuilder = new UpdateDefinitionBuilder<AgentDocument>();
+
+			List<UpdateDefinition<AgentDocument>> updates = new List<UpdateDefinition<AgentDocument>>();
+			if (pools != null)
+			{
+				updates.Add(updateBuilder.Set(x => x.Pools, pools));
+			}
+			if (enabled != null)
+			{
+				updates.Add(updateBuilder.Set(x => x.Enabled, enabled.Value));
+			}
+			if (requestConform != null)
+			{
+				updates.Add(updateBuilder.Set(x => x.RequestConform, requestConform.Value));
+				updates.Add(updateBuilder.Unset(x => x.ConformAttemptCount));
+			}
+			if (requestFullConform != null)
+			{
+				updates.Add(updateBuilder.Set(x => x.RequestFullConform, requestFullConform.Value));
+				updates.Add(updateBuilder.Unset(x => x.ConformAttemptCount));
+			}
+			if (requestRestart != null)
+			{
+				if (requestRestart.Value)
+				{
+					updates.Add(updateBuilder.Set(x => x.RequestRestart, true));
+				}
+				else
+				{
+					updates.Add(updateBuilder.Unset(x => x.RequestRestart));
+				}
+			}
+			if (requestShutdown != null)
+			{
+				if (requestShutdown.Value)
+				{
+					updates.Add(updateBuilder.Set(x => x.RequestShutdown, true));
+				}
+				else
+				{
+					updates.Add(updateBuilder.Unset(x => x.RequestShutdown));
+				}
+			}
+
+			if (shutdownReason != null)
+			{
+				updates.Add(updateBuilder.Set(x => x.LastShutdownReason, shutdownReason));
+			}
+
+			if (channel != null)
+			{
+				if (channel.Value == AgentSoftwareService.DefaultChannelName)
+				{
+					updates.Add(updateBuilder.Unset(x => x.Channel));
+				}
+				else
+				{
+					updates.Add(updateBuilder.Set(x => x.Channel, channel));
+				}
+			}
+			if (acl != null)
+			{
+				updates.Add(Acl.CreateUpdate<AgentDocument>(x => x.Acl!, acl));
+			}
+			if (comment != null)
+			{
+				updates.Add(updateBuilder.Set(x => x.Comment, comment));
+			}
+
+			// Apply the update
+			IAgent? newAgent = await TryUpdateAsyncWithRetries(agent, updateBuilder.Combine(updates));
+			if (newAgent != null)
+			{
+				if (newAgent.RequestRestart != agent.RequestRestart || newAgent.RequestConform != agent.RequestConform || newAgent.RequestShutdown != agent.RequestShutdown || newAgent.Channel != agent.Channel)
+				{
+					await PublishUpdateEventAsync(agent.Id);
+				}
+			}
+			return newAgent;
+		}
+
+		/// <inheritdoc/>
+		public async Task<IAgent?> TryUpdateSessionAsync(IAgent agentInterface, AgentStatus? status, DateTime? sessionExpiresAt, IReadOnlyList<string>? properties, IReadOnlyDictionary<string, int>? resources, IReadOnlyList<PoolId>? dynamicPools, List<AgentLease>? leases)
+		{
+			AgentDocument agent = (AgentDocument)agentInterface;
+
+			// Create an update definition for the agent
+			UpdateDefinitionBuilder<AgentDocument> updateBuilder = Builders<AgentDocument>.Update;
+			List<UpdateDefinition<AgentDocument>> updates = new List<UpdateDefinition<AgentDocument>>();
+
+			if (status != null && agent.Status != status.Value)
+			{
+				updates.Add(updateBuilder.Set(x => x.Status, status.Value));
+			}
+			if (sessionExpiresAt != null)
+			{
+				updates.Add(updateBuilder.Set(x => x.SessionExpiresAt, sessionExpiresAt.Value));
+			}
+			if (properties != null)
+			{
+				List<string> newProperties = properties.OrderBy(x => x, StringComparer.OrdinalIgnoreCase).ToList();
+				if (!agentInterface.Properties.SequenceEqual(newProperties, StringComparer.Ordinal))
+				{
+					updates.Add(updateBuilder.Set(x => x.Properties, newProperties));
+				}
+			}
+			if (resources != null && !ResourcesEqual(resources, agentInterface.Resources))
+			{
+				updates.Add(updateBuilder.Set(x => x.Resources, new Dictionary<string, int>(resources)));
+			}
+			if (dynamicPools != null && !dynamicPools.SequenceEqual(agent.DynamicPools))
+			{
+				updates.Add(updateBuilder.Set(x => x.DynamicPools, new List<PoolId>(dynamicPools)));
+			}
+			if (leases != null)
+			{
+				foreach (AgentLease lease in leases)
+				{
+					if (lease.Payload != null && (agent.Leases == null || !agent.Leases.Any(x => x.Id == lease.Id)))
+					{
+						Any payload = Any.Parser.ParseFrom(lease.Payload.ToArray());
+						if (payload.TryUnpack(out ConformTask conformTask))
+						{
+							int newConformAttemptCount = (agent.ConformAttemptCount ?? 0) + 1;
+							updates.Add(updateBuilder.Set(x => x.ConformAttemptCount, newConformAttemptCount));
+							updates.Add(updateBuilder.Set(x => x.LastConformTime, DateTime.UtcNow));
+						}
+						else if (payload.TryUnpack(out UpgradeTask upgradeTask))
+						{
+							updates.Add(updateBuilder.Set(x => x.LastUpgradeVersion, upgradeTask.SoftwareId));
+							updates.Add(updateBuilder.Set(x => x.LastUpgradeTime, DateTime.UtcNow));
+						}
+					}
+				}
+
+				updates.Add(updateBuilder.Set(x => x.Leases, leases));
+			}
+
+			// If there are no new updates, return immediately. This is important for preventing UpdateSession calls from returning immediately.
+			if (updates.Count == 0)
+			{
+				return agent;
+			}
+
+			// Update the agent, and try to create new lease documents if we succeed
+			return await TryUpdateAsyncWithRetries(agent, updateBuilder.Combine(updates));
+		}
+
+		static bool ResourcesEqual(IReadOnlyDictionary<string, int> dictA, IReadOnlyDictionary<string, int> dictB)
+		{
+			if (dictA.Count != dictB.Count)
+			{
+				return false;
+			}
+
+			foreach (KeyValuePair<string, int> pair in dictA)
+			{
+				int value;
+				if (!dictB.TryGetValue(pair.Key, out value) || value != pair.Value)
+				{
+					return false;
+				}
+			}
+
+			return true;
+		}
+
+		/// <inheritdoc/>
+		public async Task<IAgent?> TryUpdateWorkspacesAsync(IAgent agentInterface, List<AgentWorkspace> workspaces, bool requestConform)
+		{
+			AgentDocument agent = (AgentDocument)agentInterface;
+			DateTime lastConformTime = DateTime.UtcNow;
+
+			// Set the new workspaces
+			UpdateDefinition<AgentDocument> update = Builders<AgentDocument>.Update.Set(x => x.Workspaces, workspaces);
+			update = update.Set(x => x.LastConformTime, lastConformTime);
+			update = update.Unset(x => x.ConformAttemptCount);
+			if (!requestConform)
+			{
+				update = update.Unset(x => x.RequestConform);
+				update = update.Unset(x => x.RequestFullConform);
+			}
+
+			// Update the agent
+			return await TryUpdateAsyncWithRetries(agent, update);
+		}
+
+		/// <inheritdoc/>
+		public async Task<IAgent?> TryStartSessionAsync(IAgent agentInterface, SessionId sessionId, DateTime sessionExpiresAt, AgentStatus status, IReadOnlyList<string> properties, IReadOnlyDictionary<string, int> resources, IReadOnlyList<PoolId> dynamicPools, string? version)
+		{
+			AgentDocument agent = (AgentDocument)agentInterface;
+			List<string> newProperties = properties.OrderBy(x => x, StringComparer.OrdinalIgnoreCase).ToList();
+			Dictionary<string, int> newResources = new Dictionary<string, int>(resources);
+			List<PoolId> newDynamicPools = new List<PoolId>(dynamicPools);
+
+			// Reset the agent to use the new session
+			UpdateDefinitionBuilder<AgentDocument> updateBuilder = Builders<AgentDocument>.Update;
+
+			List<UpdateDefinition<AgentDocument>> updates = new List<UpdateDefinition<AgentDocument>>();
+			updates.Add(updateBuilder.Set(x => x.SessionId, sessionId));
+			updates.Add(updateBuilder.Set(x => x.SessionExpiresAt, sessionExpiresAt));
+			updates.Add(updateBuilder.Set(x => x.Status, status));
+			updates.Add(updateBuilder.Unset(x => x.Leases));
+			updates.Add(updateBuilder.Unset(x => x.Deleted));
+			updates.Add(updateBuilder.Set(x => x.Properties, newProperties));
+			updates.Add(updateBuilder.Set(x => x.Resources, newResources));
+			updates.Add(updateBuilder.Set(x => x.DynamicPools, newDynamicPools));
+			updates.Add(updateBuilder.Set(x => x.Version, version));
+			updates.Add(updateBuilder.Unset(x => x.RequestRestart));
+			updates.Add(updateBuilder.Unset(x => x.RequestShutdown));
+			updates.Add(updateBuilder.Set(x => x.LastShutdownReason, "Unexpected"));
+
+			// Apply the update
+			return await TryUpdateAsyncWithRetries(agent, updateBuilder.Combine(updates));
+		}
+
+		/// <inheritdoc/>
+		public async Task<IAgent?> TryTerminateSessionAsync(IAgent agentInterface)
+		{
+			AgentDocument agent = (AgentDocument)agentInterface;
+			UpdateDefinition<AgentDocument> update = new BsonDocument();
+
+			update = update.Unset(x => x.SessionId);
+			update = update.Unset(x => x.SessionExpiresAt);
+			update = update.Unset(x => x.Leases);
+			update = update.Set(x => x.Status, AgentStatus.Stopped);
+
+			bool bDeleted = agent.Deleted || agent.Ephemeral;
+			if (bDeleted != agent.Deleted)
+			{
+				update = update.Set(x => x.Deleted, agent.Deleted);
+			}
+
+			return await TryUpdateAsyncWithRetries(agent, update);
+		}
+
+		/// <inheritdoc/>
+		public async Task<IAgent?> TryAddLeaseAsync(IAgent agentInterface, AgentLease newLease)
+		{
+			AgentDocument agent = (AgentDocument)agentInterface;
+
+			List<AgentLease> leases = new List<AgentLease>();
+			if (agent.Leases != null)
+			{
+				leases.AddRange(agent.Leases);
+			}
+			leases.Add(newLease);
+
+			UpdateDefinition<AgentDocument> update = Builders<AgentDocument>.Update.Set(x => x.Leases, leases);
+			return await TryUpdateAsyncWithRetries(agent, update);
+		}
+
+		/// <inheritdoc/>
+		public async Task<IAgent?> TryCancelLeaseAsync(IAgent agentInterface, int leaseIdx)
+		{
+			AgentDocument agent = (AgentDocument)agentInterface;
+
+			UpdateDefinition<AgentDocument> update = Builders<AgentDocument>.Update.Set(x => x.Leases![leaseIdx].State, LeaseState.Cancelled);
+			IAgent? newAgent = await TryUpdateAsyncWithRetries(agent, update);
+			if (newAgent != null)
+			{
+				await PublishUpdateEventAsync(agent.Id);
+			}
+			return newAgent;
+		}
+
+		/// <inheritdoc/>
+		public IAuditLogChannel<AgentId> GetLogger(AgentId agentId)
+		{
+			return _auditLog[agentId];
+		}
+
+		/// <inheritdoc/>
+		Task PublishUpdateEventAsync(AgentId agentId)
+		{
+			return _redisService.Database.PublishAsync(_updateEventChannel, agentId);
+		}
+
+		/// <inheritdoc/>
+		public async Task<IDisposable> SubscribeToUpdateEventsAsync(Action<AgentId> onUpdate)
+		{
+			return await _redisService.Multiplexer.GetSubscriber().SubscribeAsync(_updateEventChannel, (channel, agentId) => onUpdate(agentId));
+		}
+	}
+}
diff --git a/Engine/Source/Programs/Horde/Horde.Build/Agents/AgentService.cs b/Engine/Source/Programs/Horde/Horde.Build/Agents/AgentService.cs
new file mode 100644
--- /dev/null
+++ b/Engine/Source/Programs/Horde/Horde.Build/Agents/AgentService.cs
@@ -0,0 +1,1025 @@
+// Copyright Epic Games, Inc. All Rights Reserved.
+
+using System;
+using System.Collections.Generic;
+using System.Diagnostics;
+using System.Linq;
+using System.Security.Claims;
+using System.Threading;
+using System.Threading.Tasks;
+using EpicGames.Core;
+using EpicGames.Redis;
+using EpicGames.Serialization;
+using Google.Protobuf.WellKnownTypes;
+using Horde.Build.Acls;
+using Horde.Build.Collections;
+using Horde.Build.Models;
+using Horde.Build.Server;
+using Horde.Build.Tasks;
+using Horde.Build.Utilities;
+using HordeCommon;
+using Microsoft.Extensions.Hosting;
+using Microsoft.Extensions.Logging;
+using MongoDB.Driver;
+using OpenTracing;
+using OpenTracing.Util;
+using StackExchange.Redis;
+using StatsdClient;
+
+namespace Horde.Build.Services
+{
+	using AgentSoftwareChannelName = StringId<AgentSoftwareChannels>;
+	using LeaseId = ObjectId<ILease>;
+	using PoolId = StringId<IPool>;
+	using SessionId = ObjectId<ISession>;
+
+	/// <summary>
+	/// Singleton used to store agent costs
+	/// </summary>
+	[RedisConverter(typeof(RedisCbConverter<>))]
+	public class AgentRateTable
+	{
+		/// <summary>
+		/// List of costs for different agent types
+		/// </summary>
+		[CbField]
+		public List<AgentRateConfig> Entries { get; set; } = new List<AgentRateConfig>();
+	}
+
+	/// <summary>
+	/// Wraps funtionality for manipulating agents
+	/// </summary>
+	public sealed class AgentService : IHostedService, IDisposable
+	{
+		/// <summary>
+		/// Maximum time between updates for an agent to be considered online
+		/// </summary>
+		public static readonly TimeSpan SessionExpiryTime = TimeSpan.FromMinutes(5);
+
+		/// <summary>
+		/// Time before a session expires that we will poll until
+		/// </summary>
+		public static readonly TimeSpan SessionLongPollTime = TimeSpan.FromSeconds(55);
+
+		/// <summary>
+		/// Time after which a session will be renewed
+		/// </summary>
+		public static readonly TimeSpan SessionRenewTime = TimeSpan.FromSeconds(50);
+		readonly AclService _aclService;
+		readonly IDowntimeService _downtimeService;
+
+		/// <summary>
+		/// Collection of agent documents
+		/// </summary>
+		public IAgentCollection Agents { get; }
+
+		readonly ILeaseCollection _leases;
+		readonly ISessionCollection _sessions;
+		readonly IDogStatsd _dogStatsd;
+		readonly ITaskSource[] _taskSources;
+		readonly IHostApplicationLifetime _applicationLifetime;
+		readonly IClock _clock;
+		readonly ILogger _logger;
+		readonly ITicker _ticker;
+		
+		readonly RedisString<AgentRateTable> _agentRateTableData;
+		readonly RedisService _redisService;
+
+		// Lazily updated costs for different agent types
+		readonly AsyncCachedValue<AgentRateTable?> _agentRateTable;
+
+		// Lazily updated list of current pools
+		readonly AsyncCachedValue<List<IPool>> _poolsList;
+
+		// All the agents currently performing a long poll for work on this server
+		readonly Dictionary<AgentId, CancellationTokenSource> _waitingAgents = new Dictionary<AgentId, CancellationTokenSource>();
+
+		// Subscription for update events
+		readonly IDisposable _subscription;
+
+		/// <summary>
+		/// Constructor
+		/// </summary>
+		public AgentService(IAgentCollection agents, ILeaseCollection leases, ISessionCollection sessions, AclService aclService, IDowntimeService downtimeService, IPoolCollection poolCollection, IDogStatsd dogStatsd, IEnumerable<ITaskSource> taskSources, RedisService redisService, IHostApplicationLifetime applicationLifetime, IClock clock, ILogger<AgentService> logger)
+		{
+			Agents = agents;
+			_leases = leases;
+			_sessions = sessions;
+			_aclService = aclService;
+			_downtimeService = downtimeService;
+			_agentRateTableData = new RedisString<AgentRateTable>(redisService.Database, "agent-rates");
+			_agentRateTable = new AsyncCachedValue<AgentRateTable?>(() => _agentRateTableData.GetAsync(), TimeSpan.FromSeconds(2.0));//.FromMinutes(5.0));
+			_poolsList = new AsyncCachedValue<List<IPool>>(() => poolCollection.GetAsync(), TimeSpan.FromSeconds(30.0));
+			_dogStatsd = dogStatsd;
+			_taskSources = taskSources.ToArray();
+			_applicationLifetime = applicationLifetime;
+			_redisService = redisService;
+			_clock = clock;
+			_ticker = clock.AddTicker<AgentService>(TimeSpan.FromSeconds(30.0), TickAsync, logger);
+			_logger = logger;
+
+			_subscription = agents.SubscribeToUpdateEventsAsync(OnAgentUpdate).Result;
+		}
+
+		/// <inheritdoc/>
+		public Task StartAsync(CancellationToken cancellationToken) => _ticker.StartAsync();
+
+		/// <inheritdoc/>
+		public Task StopAsync(CancellationToken cancellationToken) => _ticker.StopAsync();
+
+		/// <inheritdoc/>
+		public void Dispose()
+		{
+			_subscription.Dispose();
+			_ticker.Dispose();
+		}
+
+		/// <summary>
+		/// Gets user-readable payload information
+		/// </summary>
+		/// <param name="payload">The payload data</param>
+		/// <returns>Dictionary of key/value pairs for the payload</returns>
+		public Dictionary<string, string>? GetPayloadDetails(ReadOnlyMemory<byte>? payload)
+		{
+			Dictionary<string, string>? details = null;
+			if (payload != null)
+			{
+				Any basePayload = Any.Parser.ParseFrom(payload.Value.ToArray());
+				foreach (ITaskSource taskSource in _taskSources)
+				{
+					if (basePayload.Is(taskSource.Descriptor))
+					{
+						details = new Dictionary<string, string>();
+						taskSource.GetLeaseDetails(basePayload, details);
+						break;
+					}
+				}
+			}
+			return details;
+		}
+
+		/// <summary>
+		/// Issues a bearer token for the given session id
+		/// </summary>
+		/// <param name="agentId">The agent id</param>
+		/// <param name="sessionId">The session id</param>
+		/// <returns>Bearer token for the agent</returns>
+		public string IssueSessionToken(AgentId agentId, SessionId sessionId)
+		{
+			List<AclClaim> claims = new List<AclClaim>();
+			claims.Add(AclService.AgentRoleClaim);
+			claims.Add(AclService.GetAgentClaim(agentId));
+			claims.Add(AclService.GetSessionClaim(sessionId));
+			return _aclService.IssueBearerToken(claims, null);
+		}
+
+		/// <summary>
+		/// Register a new agent
+		/// </summary>
+		/// <param name="name">Name of the agent</param>
+		/// <param name="bEnabled">Whether the agent is currently enabled</param>
+		/// <param name="channel">Override for the desired software version</param>
+		/// <param name="pools">Pools for this agent</param>
+		/// <returns>Unique id for the agent</returns>
+		public Task<IAgent> CreateAgentAsync(string name, bool bEnabled, AgentSoftwareChannelName? channel, List<PoolId>? pools)
+		{
+			return Agents.AddAsync(new AgentId(name), bEnabled, channel, pools);
+		}
+
+		/// <summary>
+		/// Gets an agent by ID
+		/// </summary>
+		/// <param name="agentId">Unique id of the agent</param>
+		/// <returns>The agent document</returns>
+		public Task<IAgent?> GetAgentAsync(AgentId agentId)
+		{
+			return Agents.GetAsync(agentId);
+		}
+
+		/// <summary>
+		/// Finds all agents matching certain criteria
+		/// </summary>
+		/// <param name="poolId">The pool containing the agent</param>
+		/// <param name="modifiedAfter">If set, only returns agents modified after this time</param>
+		/// <param name="index">Index within the list of results</param>
+		/// <param name="count">Number of results to return</param>
+		/// <returns>List of agents matching the given criteria</returns>
+		public Task<List<IAgent>> FindAgentsAsync(PoolId? poolId, DateTime? modifiedAfter, int? index, int? count)
+		{
+			return Agents.FindAsync(poolId, modifiedAfter, null, index, count);
+		}
+
+		/// <summary>
+		/// Update the current workspaces for an agent.
+		/// </summary>
+		/// <param name="agent">The agent to update</param>
+		/// <param name="workspaces">Current list of workspaces</param>
+		/// <param name="bPendingConform">Whether the agent still needs to run another conform</param>
+		/// <returns>New agent state</returns>
+		public async Task<bool> TryUpdateWorkspacesAsync(IAgent agent, List<AgentWorkspace> workspaces, bool bPendingConform)
+		{
+			IAgent? newAgent = await Agents.TryUpdateWorkspacesAsync(agent, workspaces, bPendingConform);
+			return newAgent != null;
+		}
+
+		/// <summary>
+		/// Marks the agent as deleted
+		/// </summary>
+		/// <param name="agent">The agent to delete</param>
+		/// <returns>Async task</returns>
+		public async Task DeleteAgentAsync(IAgent? agent)
+		{
+			while (agent != null && !agent.Deleted)
+			{
+				IAgent? newAgent = await Agents.TryDeleteAsync(agent);
+				if(newAgent != null)
+				{
+					break;
+				}
+				agent = await GetAgentAsync(agent.Id);
+			}
+		}
+
+		async ValueTask<List<PoolId>> GetDynamicPoolsAsync(IAgent agent)
+		{
+			List<PoolId> newDynamicPools = new List<PoolId>();
+
+			List<IPool> pools = await _poolsList.GetAsync();
+			foreach (IPool pool in pools)
+			{
+				if (pool.Condition != null && agent.SatisfiesCondition(pool.Condition))
+				{
+					newDynamicPools.Add(pool.Id);
+				}
+			}
+
+			return newDynamicPools;
+		}
+
+		/// <summary>
+		/// Callback for an agents 
+		/// </summary>
+		/// <param name="agentId"></param>
+		void OnAgentUpdate(AgentId agentId)
+		{
+			lock (_waitingAgents)
+			{
+				if (_waitingAgents.TryGetValue(agentId, out CancellationTokenSource? cancellationSource))
+				{
+					cancellationSource.Cancel();
+				}
+			}
+		}
+
+		/// <summary>
+		/// Creates a new agent session
+		/// </summary>
+		/// <param name="agent">The agent to create a session for</param>
+		/// <param name="status">Current status of the agent</param>
+		/// <param name="properties">Properties for the agent</param>
+		/// <param name="resources">Resources which the agent has</param>
+		/// <param name="version">Version of the software that's running</param>
+		/// <returns>New agent state</returns>
+		public async Task<IAgent> CreateSessionAsync(IAgent agent, AgentStatus status, IReadOnlyList<string> properties, IReadOnlyDictionary<string, int> resources, string? version)
+		{
+			for (; ; )
+			{
+				IAuditLogChannel<AgentId> agentLogger = Agents.GetLogger(agent.Id);
+
+				// Check if there's already a session running for this agent.
+				IAgent? newAgent;
+				if (agent.SessionId != null)
+				{
+					// Try to terminate the current session
+					await TryTerminateSessionAsync(agent);
+				}
+				else
+				{
+					DateTime utcNow = _clock.UtcNow;
+
+					// Remove any outstanding leases
+					foreach (AgentLease lease in agent.Leases)
+					{
+						agentLogger.LogInformation("Removing outstanding lease {LeaseId}", lease.Id);
+						await RemoveLeaseAsync(agent, lease, utcNow, LeaseOutcome.Failed, null);
+					}
+
+					// Create a new session document
+					ISession newSession = await _sessions.AddAsync(SessionId.GenerateNewId(), agent.Id, _clock.UtcNow, properties, resources, version);
+					DateTime sessionExpiresAt = utcNow + SessionExpiryTime;
+
+					// Get the new dynamic pools for the agent
+					List<PoolId> dynamicPools = await GetDynamicPoolsAsync(agent);
+
+					// Reset the agent to use the new session
+					newAgent = await Agents.TryStartSessionAsync(agent, newSession.Id, sessionExpiresAt, status, properties, resources, dynamicPools, version);
+					if(newAgent != null)
+					{
+						agent = newAgent;
+						agentLogger.LogInformation("Session {SessionId} started", newSession.Id);
+						break;
+					}
+
+					// Remove the session we didn't use
+					await _sessions.DeleteAsync(newSession.Id);
+				}
+
+				// Get the current agent state
+				newAgent = await GetAgentAsync(agent.Id);
+				if (newAgent == null)
+				{
+					throw new InvalidOperationException($"Invalid agent id '{agent.Id}'");
+				}
+				agent = newAgent;
+			}
+			return agent;
+		}
+
+		/// <summary>
+		/// Determines whether a task source can currently issue tasks
+		/// </summary>
+		bool CanUseTaskSource(IAgent agent, ITaskSource taskSource)
+		{
+			TaskSourceFlags flags = taskSource.Flags;
+			if ((flags & TaskSourceFlags.AllowWhenDisabled) == 0 && !agent.Enabled)
+			{
+				return false;
+			}
+			if ((flags & TaskSourceFlags.AllowDuringDowntime) == 0 && _downtimeService.IsDowntimeActive)
+			{
+				return false;
+			}
+			return true;
+		}
+
+		/// <summary>
+		/// Waits for a lease to be assigned to an agent
+		/// </summary>
+		/// <param name="agent">The agent to assign a lease to</param>
+		/// <param name="cancellationToken"></param>
+		/// <returns>True if a lease was assigned, false otherwise</returns>
+		public async Task<IAgent?> WaitForLeaseAsync(IAgent? agent, CancellationToken cancellationToken)
+		{
+			while (agent != null)
+			{
+				if (!agent.SessionExpiresAt.HasValue)
+				{
+					break;
+				}
+
+				// Check we have some time to wait
+				DateTime utcNow = _clock.UtcNow;
+				TimeSpan maxWaitTime = (agent.SessionExpiresAt.Value - SessionExpiryTime + SessionLongPollTime) - utcNow;
+				if (maxWaitTime <= TimeSpan.Zero)
+				{
+					break;
+				}
+
+				// Create a cancellation token that will expire with the session
+				using CancellationTokenSource cancellationSource = CancellationTokenSource.CreateLinkedTokenSource(_applicationLifetime.ApplicationStopping, cancellationToken);
+				cancellationSource.CancelAfter(maxWaitTime);
+
+				// Assign a new lease
+				(ITaskSource, AgentLease)? result = null;
+				try
+				{
+					// Add the cancellation source to the set of waiting agents
+					lock (_waitingAgents)
+					{
+						_waitingAgents[agent.Id] = cancellationSource;
+					}
+
+					// Create all the tasks to wait for
+					List<Task<(ITaskSource, AgentLease)?>> tasks = new List<Task<(ITaskSource, AgentLease)?>>();
+					foreach (ITaskSource taskSource in _taskSources)
+					{
+						if (CanUseTaskSource(agent, taskSource) && !cancellationSource.IsCancellationRequested)
+						{
+							Task<(ITaskSource, AgentLease)?> task = await GuardedAssignLeaseAsync(taskSource, agent, cancellationSource);
+							tasks.Add(task);
+						}
+					}
+
+					// If no task source is valid, just add a delay
+					if (tasks.Count == 0)
+					{
+						await AsyncUtils.DelayNoThrow(maxWaitTime, cancellationToken);
+						break;
+					}
+
+					// Wait for all the tasks to complete. Once the first task completes it will set the cancellation source, triggering the 
+					// others to terminate.
+					await Task.WhenAll(tasks);
+
+					// Find the first result
+					foreach (Task<(ITaskSource, AgentLease)?> task in tasks)
+					{
+						(ITaskSource, AgentLease)? taskResult;
+						if (task.TryGetResult(out taskResult) && taskResult != null)
+						{
+							(ITaskSource taskSource, AgentLease taskLease) = taskResult.Value;
+							if (result == null)
+							{
+								result = (taskSource, taskLease);
+							}
+							else
+							{
+								await taskSource.CancelLeaseAsync(agent, taskLease.Id, Any.Parser.ParseFrom(taskLease.Payload));
+							}
+						}
+					}
+				}
+				finally
+				{
+					lock (_waitingAgents)
+					{
+						_waitingAgents.Remove(agent.Id);
+					}
+				}
+
+				// Exit if we didn't find any work to do. It may be that all the task sources returned null, in which case wait for the time period to expire.
+				if (result == null)
+				{
+					if (!cancellationSource.IsCancellationRequested)
+					{
+						await cancellationSource.Token.AsTask();
+					}
+					break;
+				}
+
+				// Get the resulting lease
+				(ITaskSource source, AgentLease lease) = result.Value;
+
+				// Add the new lease to the agent
+				IAgent? newAgent = await Agents.TryAddLeaseAsync(agent, lease);
+				if (newAgent != null)
+				{
+					await source.OnLeaseStartedAsync(newAgent, lease.Id, Any.Parser.ParseFrom(lease.Payload), Agents.GetLogger(agent.Id));
+					await CreateLeaseAsync(agent, lease);
+					return newAgent;
+				}
+
+				// Update the agent
+				agent = await GetAgentAsync(agent.Id);
+			}
+			return agent;
+		}
+
+		async Task<Task<(ITaskSource, AgentLease)?>> GuardedAssignLeaseAsync(ITaskSource source, IAgent agent, CancellationTokenSource cancellationSource)
+		{
+			CancellationToken cancellationToken = cancellationSource.Token;
+			try
+			{
+				Task<AgentLease?> task = await source.AssignLeaseAsync(agent, cancellationToken);
+				return task.ContinueWith(x => WrapAssignedLease(source, x, cancellationSource), TaskScheduler.Default);
+			}
+			catch (TaskCanceledException) when (cancellationToken.IsCancellationRequested)
+			{
+				return Task.FromResult<(ITaskSource, AgentLease)?>(null);
+			}
+			catch (Exception ex)
+			{
+				_logger.LogError(ex, "Exception while trying to assign lease");
+				return Task.FromResult<(ITaskSource, AgentLease)?>(null);
+			}
+		}
+
+		(ITaskSource, AgentLease)? WrapAssignedLease(ITaskSource source, Task<AgentLease?> task, CancellationTokenSource cancellationSource)
+		{
+			if (task.IsCanceled)
+			{
+				return null;
+			}
+			else if (task.TryGetResult(out AgentLease? lease))
+			{
+				if (lease == null)
+				{
+					return null;
+				}
+				else
+				{
+					cancellationSource.Cancel();
+					return (source, lease);
+				}
+			}
+			else if (task.IsFaulted)
+			{
+				_logger.LogError(task.Exception, "Exception while trying to assign lease");
+				return null;
+			}
+			else
+			{
+				_logger.LogWarning("Unhandled task state: {Status}", task.Status);
+				return null;
+			}
+		}
+
+		/// <summary>
+		/// Cancels the specified agent lease
+		/// </summary>
+		/// <param name="agent">Agent to cancel the lease on</param>
+		/// <param name="leaseId">The lease id to cancel</param>
+		/// <returns></returns>
+		public async Task<bool> CancelLeaseAsync(IAgent agent, LeaseId leaseId)
+		{		
+			int index = 0;
+			while (index < agent.Leases.Count && agent.Leases[index].Id != leaseId)
+			{
+				index++;
+			}
+
+			if (index == agent.Leases.Count)
+			{
+				return false;
+			}
+
+			if (agent.Leases[index].State == LeaseState.Cancelled)
+			{
+				return false;
+			}
+
+			await Agents.TryCancelLeaseAsync(agent, index);
+			return true;
+		}
+
+		/// <summary>
+		/// 
+		/// </summary>
+		public async Task<IAgent?> UpdateSessionWithWaitAsync(IAgent inAgent, SessionId sessionId, AgentStatus status, IReadOnlyList<string>? properties, IReadOnlyDictionary<string, int>? resources, IList<HordeCommon.Rpc.Messages.Lease> newLeases, CancellationToken cancellationToken)
+		{
+			IAgent? agent = inAgent;
+
+			// Capture the current agent update index. This allows us to detect if anything has changed.
+			uint updateIndex = agent.UpdateIndex;
+
+			// Update the agent session and return to the caller if anything changes
+			agent = await UpdateSessionAsync(agent, sessionId, status, properties, resources, newLeases);
+			if (agent != null && agent.UpdateIndex == updateIndex && (agent.Leases.Count > 0 || agent.Status != AgentStatus.Stopping))
+			{
+				agent = await WaitForLeaseAsync(agent, cancellationToken);
+			}
+			return agent;
+		}
+
+		/// <summary>
+		/// Updates the state of the current agent session
+		/// </summary>
+		/// <param name="inAgent">The current agent state</param>
+		/// <param name="sessionId">Id of the session</param>
+		/// <param name="status">New status for the agent</param>
+		/// <param name="properties">New agent properties</param>
+		/// <param name="resources">New agent resources</param>
+		/// <param name="newLeases">New list of leases for this session</param>
+		/// <returns>Updated agent state</returns>
+		public async Task<IAgent?> UpdateSessionAsync(IAgent inAgent, SessionId sessionId, AgentStatus status, IReadOnlyList<string>? properties, IReadOnlyDictionary<string, int>? resources, IList<HordeCommon.Rpc.Messages.Lease> newLeases)
+		{
+			DateTime utcNow = _clock.UtcNow;
+
+			Stopwatch timer = Stopwatch.StartNew();
+
+			IAgent? agent = inAgent;
+			while (agent != null)
+			{
+				// If the agent is stopping and doesn't have any leases, we can terminate the current session.
+				if (status == AgentStatus.Stopping && newLeases.Count == 0)
+				{
+					// If we've already decided to terminate the session, this update is redundant but harmless
+					if (agent.SessionId != sessionId)
+					{
+						break;
+					}
+
+					// If the session is valid, we can terminate once the agent leases are also empty
+					if (agent.Leases.Count == 0)
+					{
+						if (!await TryTerminateSessionAsync(agent))
+						{
+							agent = await GetAgentAsync(agent.Id);
+							continue;
+						}
+						break;
+					}
+				}
+
+				// Check the session id is correct.
+				if (agent.SessionId != sessionId)
+				{
+					throw new InvalidOperationException($"Invalid agent session {sessionId}");
+				}
+
+				// Check the session hasn't expired
+				if (!agent.IsSessionValid(utcNow))
+				{
+					throw new InvalidOperationException("Session has already expired");
+				}
+
+				// Extend the current session time if we're within a time period of the current time expiring. This reduces
+				// unnecessary DB writes a little, but it also allows us to skip the update and jump into a long poll state
+				// if there's still time on the session left.
+				DateTime? sessionExpiresAt = null;
+				if (!agent.SessionExpiresAt.HasValue || utcNow > (agent.SessionExpiresAt - SessionExpiryTime) + SessionRenewTime)
+				{
+					sessionExpiresAt = utcNow + SessionExpiryTime;
+				}
+
+				// Flag for whether the leases array should be updated
+				bool bUpdateLeases = false;
+				List<AgentLease> leases = new List<AgentLease>(agent.Leases);
+
+				// Remove any completed leases from the agent
+				Dictionary<LeaseId, HordeCommon.Rpc.Messages.Lease> leaseIdToNewState = newLeases.ToDictionary(x => new LeaseId(x.Id), x => x);
+				for (int idx = 0; idx < leases.Count; idx++)
+				{
+					AgentLease lease = leases[idx];
+					if (lease.State == LeaseState.Cancelled)
+					{
+						HordeCommon.Rpc.Messages.Lease? newLease;
+						if (!leaseIdToNewState.TryGetValue(lease.Id, out newLease) || newLease.State == LeaseState.Cancelled || newLease.State == LeaseState.Completed)
+						{
+							await RemoveLeaseAsync(agent, lease, utcNow, LeaseOutcome.Cancelled, null);
+							leases.RemoveAt(idx--);
+							bUpdateLeases = true;
+						}
+					}
+					else
+					{
+						HordeCommon.Rpc.Messages.Lease? newLease;
+						if (leaseIdToNewState.TryGetValue(lease.Id, out newLease) && newLease.State != lease.State)
+						{
+							if (newLease.State == LeaseState.Cancelled || newLease.State == LeaseState.Completed)
+							{
+								await RemoveLeaseAsync(agent, lease, utcNow, newLease.Outcome, newLease.Output.ToByteArray());
+								leases.RemoveAt(idx--);
+							}
+							else if (newLease.State == LeaseState.Active && lease.State == LeaseState.Pending)
+							{
+								lease.State = LeaseState.Active;
+							}
+							bUpdateLeases = true;
+						}
+					}
+				}
+
+				// If the agent is stopping, cancel all the leases. Clear out the current session once it's complete.
+				if (status == AgentStatus.Stopping)
+				{
+					foreach (AgentLease lease in leases)
+					{
+						if (lease.State != LeaseState.Cancelled)
+						{
+							lease.State = LeaseState.Cancelled;
+							bUpdateLeases = true;
+						}
+					}
+				}
+
+				// Get the new dynamic pools for the agent
+				List<PoolId> dynamicPools = await GetDynamicPoolsAsync(agent);
+
+				// Update the agent, and try to create new lease documents if we succeed
+				IAgent? newAgent = await Agents.TryUpdateSessionAsync(agent, status, sessionExpiresAt, properties, resources, dynamicPools, bUpdateLeases ? leases : null);
+				if (newAgent != null)
+				{
+					agent = newAgent;
+					break;
+				}
+
+				// Fetch the agent again
+				agent = await GetAgentAsync(agent.Id);
+			}
+			return agent;
+		}
+
+		/// <summary>
+		/// Terminates an existing session. Does not update the agent itself, if it's currently 
+		/// </summary>
+		/// <param name="agent">The agent whose current session should be terminated</param>
+		/// <returns>True if the session was terminated</returns>
+		private async Task<bool> TryTerminateSessionAsync(IAgent agent)
+		{
+			// Make sure the agent has a valid session id
+			if (agent.SessionId == null)
+			{
+				return true;
+			}
+
+			// Get the time that the session finishes at
+			DateTime finishTime = _clock.UtcNow;
+			if (agent.SessionExpiresAt.HasValue && agent.SessionExpiresAt.Value < finishTime)
+			{
+				finishTime = agent.SessionExpiresAt.Value;
+			}
+
+			// Save off the session id and current leases
+			SessionId sessionId = agent.SessionId.Value;
+			List<AgentLease> leases = new List<AgentLease>(agent.Leases);
+
+			// Clear the current session
+			IAgent? newAgent = await Agents.TryTerminateSessionAsync(agent);
+			if (newAgent != null)
+			{
+				agent = newAgent;
+
+				// Remove any outstanding leases
+				foreach(AgentLease lease in leases)
+				{
+					Agents.GetLogger(agent.Id).LogInformation("Removing lease {LeaseId} during session terminate...", lease.Id);
+					await RemoveLeaseAsync(agent, lease, finishTime, LeaseOutcome.Failed, null);
+				}
+
+				// Update the session document
+				Agents.GetLogger(agent.Id).LogInformation("Terminated session {SessionId}", sessionId);
+				await _sessions.UpdateAsync(sessionId, finishTime, agent.Properties, agent.Resources);
+				return true;
+			}
+			return false;
+		}
+
+		/// <summary>
+		/// Creates a new lease document
+		/// </summary>
+		/// <param name="agent">Agent that will be executing the lease</param>
+		/// <param name="agentLease">The new agent lease</param>
+		/// <returns>New lease document</returns>
+		private Task<ILease> CreateLeaseAsync(IAgent agent, AgentLease agentLease)
+		{
+			try
+			{
+				return _leases.AddAsync(agentLease.Id, agentLease.Name, agent.Id, agent.SessionId!.Value, agentLease.StreamId, agentLease.PoolId, agentLease.LogId, agentLease.StartTime, agentLease.Payload!);
+			}
+			catch (Exception ex)
+			{
+				_logger.LogError(ex, "Unable to create lease {LeaseId} for agent {AgentId}; lease already exists?", agentLease.Id, agent.Id);
+				throw;
+			}
+		}
+
+		/// <summary>
+		/// Finds all leases matching a set of criteria
+		/// </summary>
+		/// <param name="agentId">Unqiue id of the agent executing this lease</param>
+		/// <param name="sessionId">Unique id of the agent session</param>
+		/// <param name="startTime">Start of the search window to return results for</param>
+		/// <param name="finishTime">End of the search window to return results for</param>
+		/// <param name="index">Index of the first result to return</param>
+		/// <param name="count">Number of results to return</param>
+		/// <returns>List of leases matching the given criteria</returns>
+		public Task<List<ILease>> FindLeasesAsync(AgentId? agentId, SessionId? sessionId, DateTime? startTime, DateTime? finishTime, int index, int count)
+		{
+			using IScope scope = GlobalTracer.Instance.BuildSpan($"{nameof(AgentService)}.{nameof(FindLeasesAsync)}").StartActive();
+			scope.Span.SetTag("AgentId", agentId?.ToString());
+			scope.Span.SetTag("SessionId", sessionId?.ToString());
+			scope.Span.SetTag("StartTime", startTime?.ToString());
+			scope.Span.SetTag("FinishTime", finishTime?.ToString());
+			scope.Span.SetTag("Index", index);
+			scope.Span.SetTag("Count", count);
+			return _leases.FindLeasesAsync(agentId, sessionId, startTime, finishTime, index, count);
+		}
+
+		/// <summary>
+		/// Finds all leases by finish time
+		/// </summary>
+		/// <param name="minFinishTime">Start of the search window to return results for</param>
+		/// <param name="maxFinishTime">End of the search window to return results for</param>
+		/// <param name="index">Index of the first result to return</param>
+		/// <param name="count">Number of results to return</param>
+		/// <returns>List of leases matching the given criteria</returns>
+		public Task<List<ILease>> FindLeasesByFinishTimeAsync(DateTime? minFinishTime, DateTime? maxFinishTime, int? index, int? count)
+		{
+			return _leases.FindLeasesByFinishTimeAsync(minFinishTime, maxFinishTime, index, count, null, false);
+		}
+
+		/// <summary>
+		/// Gets a specific lease
+		/// </summary>
+		/// <param name="leaseId">Unique id of the lease</param>
+		/// <returns>The lease that was found, or null if it does not exist</returns>
+		public Task<ILease?> GetLeaseAsync(LeaseId leaseId)
+		{
+			return _leases.GetAsync(leaseId);
+		}
+
+		/// <summary>
+		/// Removes a lease with the given id. Updates the lease state in the database, and removes the item from the agent's leases array.
+		/// </summary>
+		/// <param name="agent">The agent to remove a lease from</param>
+		/// <param name="lease">The lease to cancel</param>
+		/// <param name="utcNow">The current time</param>
+		/// <param name="outcome">Final status of the lease</param>
+		/// <param name="output">Output from executing the task</param>
+		/// <returns>Async task</returns>
+		private async Task RemoveLeaseAsync(IAgent agent, AgentLease lease, DateTime utcNow, LeaseOutcome outcome, byte[]? output)
+		{
+			// Make sure the lease is terminated correctly
+			if (lease.Payload == null)
+			{
+				_logger.LogWarning("Removing lease {LeaseId} (no payload)", lease.Id);
+			}
+			else
+			{
+				Any any = Any.Parser.ParseFrom(lease.Payload);
+				_logger.LogInformation("Removing lease {LeaseId} ({LeaseType})", lease.Id, any.TypeUrl);
+
+				foreach (ITaskSource taskSource in _taskSources)
+				{
+					if (any.Is(taskSource.Descriptor))
+					{
+						await taskSource.OnLeaseFinishedAsync(agent, lease.Id, any, outcome, output, Agents.GetLogger(agent.Id));
+						break;
+					}
+				}
+			}
+
+			// Figure out what time the lease finished
+			DateTime finishTime = utcNow;
+			if (agent.SessionExpiresAt.HasValue && agent.SessionExpiresAt.Value < finishTime)
+			{
+				finishTime = agent.SessionExpiresAt.Value;
+			}
+			if (lease.ExpiryTime.HasValue && lease.ExpiryTime.Value < finishTime)
+			{
+				finishTime = lease.ExpiryTime.Value;
+			}
+
+			// Update the lease
+			await _leases.TrySetOutcomeAsync(lease.Id, finishTime, outcome, output);
+		}
+
+		/// <summary>
+		/// Gets information about a particular session
+		/// </summary>
+		/// <param name="sessionId">The unique session id</param>
+		/// <returns>The session information</returns>
+		public Task<ISession?> GetSessionAsync(SessionId sessionId)
+		{
+			return _sessions.GetAsync(sessionId);
+		}
+
+		/// <summary>
+		/// Find sessions for the given agent
+		/// </summary>
+		/// <param name="agentId">The unique agent id</param>
+		/// <param name="startTime">Start time to include in the search</param>
+		/// <param name="finishTime">Finish time to include in the search</param>
+		/// <param name="index">Index of the first result to return</param>
+		/// <param name="count">Number of results to return</param>
+		/// <returns>List of sessions matching the given criteria</returns>
+		public Task<List<ISession>> FindSessionsAsync(AgentId agentId, DateTime? startTime, DateTime? finishTime, int index, int count)
+		{
+			return _sessions.FindAsync(agentId, startTime, finishTime, index, count);
+		}
+
+		/// <summary>
+		/// Update the rate table for different agent types
+		/// </summary>
+		/// <param name="entries">New entries for the rate table</param>
+		/// <returns></returns>
+		public async Task UpdateRateTableAsync(List<AgentRateConfig> entries)
+		{
+			await _agentRateTableData.SetAsync(new AgentRateTable { Entries = entries });
+		}
+
+		/// <summary>
+		/// Gets the rate for the given agent
+		/// </summary>
+		/// <param name="agentId">Agent id to query</param>
+		/// <returns>Hourly rate of running the given agent</returns>
+		public async ValueTask<double?> GetRateAsync(AgentId agentId)
+		{
+			RedisKey key = $"agent-rate/{agentId}";
+
+			// Try to get the current value
+			RedisValue value = await _redisService.Database.StringGetAsync(key);
+			if (!value.IsNull)
+			{
+				double rate = (double)value;
+				if (rate == 0.0)
+				{
+					return null;
+				}
+				else
+				{
+					return rate;
+				}
+			}
+			else
+			{
+				double rate = 0.0;
+
+				// Get the rate table
+				AgentRateTable? rateTable = await _agentRateTable.GetAsync();
+				if (rateTable != null && rateTable.Entries.Count > 0)
+				{
+					IAgent? agent = await GetAgentAsync(agentId);
+					if (agent != null)
+					{
+						foreach (AgentRateConfig config in rateTable.Entries)
+						{
+							if (config.Condition != null && config.Condition.Evaluate(x => agent.GetPropertyValues(x)))
+							{
+								rate = config.Rate;
+								break;
+							}
+						}
+					}
+				}
+
+				// Cache it for future reference
+				await _redisService.Database.StringSetAsync(key, rate, TimeSpan.FromMinutes(5.0), flags: CommandFlags.FireAndForget);
+				return rate;
+			}
+		}
+
+		/// <summary>
+		/// Terminate any sessions for agents that are offline
+		/// </summary>
+		/// <param name="stoppingToken">Token indicating the service is shutting down</param>
+		/// <returns>Async task</returns>
+		async ValueTask TickAsync(CancellationToken stoppingToken)
+		{
+			while (!stoppingToken.IsCancellationRequested)
+			{
+				// Find all the agents which are ready to be expired
+				const int MaxAgents = 100;
+				DateTime utcNow = _clock.UtcNow;
+				List<IAgent> expiredAgents = await Agents.FindExpiredAsync(utcNow, MaxAgents);
+
+				// Transition each agent to being offline
+				foreach (IAgent expiredAgent in expiredAgents)
+				{
+					stoppingToken.ThrowIfCancellationRequested();
+					_logger.LogDebug("Terminating session {SessionId} for agent {Agent}", expiredAgent.SessionId, expiredAgent.Id);
+					await TryTerminateSessionAsync(expiredAgent);
+				}
+
+				// Try again if we didn't fetch everything
+				if(expiredAgents.Count < MaxAgents)
+				{
+					break;
+				}
+			}
+
+			await CollectMetrics();
+		}
+
+		private async Task CollectMetrics()
+		{
+			List<IAgent> agentList = await Agents.FindAsync();
+			int numAgentsTotal = agentList.Count;
+			int numAgentsTotalEnabled = agentList.Count(a => a.Enabled);
+			int numAgentsTotalDisabled = agentList.Count(a => !a.Enabled);
+			int numAgentsTotalOk = agentList.Count(a => a.Enabled && a.Status == AgentStatus.Ok);
+			int numAgentsTotalStopping = agentList.Count(a => a.Enabled && a.Status == AgentStatus.Stopping);
+			int numAgentsTotalUnhealthy = agentList.Count(a => a.Enabled && a.Status == AgentStatus.Unhealthy);
+			int numAgentsTotalUnspecified = agentList.Count(a => a.Enabled && a.Status == AgentStatus.Unspecified);
+			
+			// TODO: utilize tags argument in a smarter way below
+			_dogStatsd.Gauge("agents.total.count", numAgentsTotal);
+			_dogStatsd.Gauge("agents.total.enabled.count", numAgentsTotalEnabled);
+			_dogStatsd.Gauge("agents.total.disabled.count", numAgentsTotalDisabled);
+			_dogStatsd.Gauge("agents.total.ok.count", numAgentsTotalOk);
+			_dogStatsd.Gauge("agents.total.stopping.count", numAgentsTotalStopping);
+			_dogStatsd.Gauge("agents.total.unhealthy.count", numAgentsTotalUnhealthy);
+			_dogStatsd.Gauge("agents.total.unspecified.count", numAgentsTotalUnspecified);
+		}
+
+		/// <summary>
+		/// Determines if the user is authorized to perform an action on a particular agent
+		/// </summary>
+		/// <param name="agent">The agent to check</param>
+		/// <param name="action">The action being performed</param>
+		/// <param name="user">The principal to authorize</param>
+		/// <param name="cache">The permissions cache</param>
+		/// <returns>True if the action is authorized</returns>
+		public async Task<bool> AuthorizeAsync(IAgent agent, AclAction action, ClaimsPrincipal user, GlobalPermissionsCache? cache)
+		{
+			bool? result = agent.Acl?.Authorize(action, user);
+			if (result == null)
+			{
+				return await _aclService.AuthorizeAsync(action, user, cache);
+			}
+			else
+			{
+				return result.Value;
+			}
+		}
+
+		/// <summary>
+		/// Determines if the user is authorized to perform an action on a particular agent
+		/// </summary>
+		/// <param name="agent">The agent to check</param>
+		/// <param name="user">The principal to authorize</param>
+		/// <returns>True if the action is authorized</returns>
+		public bool AuthorizeSession(IAgent agent, ClaimsPrincipal user)
+		{
+			if (agent.SessionId != null && user.HasSessionClaim(agent.SessionId.Value) && agent.IsSessionValid(_clock.UtcNow))
+			{
+				return true;
+			}
+			else
+			{
+				return false;
+			}
+		}
+	}
+}
diff --git a/Engine/Source/Programs/Horde/Horde.Build/Agents/AgentsController.Messages.cs b/Engine/Source/Programs/Horde/Horde.Build/Agents/AgentsController.Messages.cs
new file mode 100644
--- /dev/null
+++ b/Engine/Source/Programs/Horde/Horde.Build/Agents/AgentsController.Messages.cs
@@ -0,0 +1,497 @@
+// Copyright Epic Games, Inc. All Rights Reserved.
+
+using System;
+using System.Collections.Generic;
+using System.ComponentModel.DataAnnotations;
+using System.Linq;
+using EpicGames.Core;
+using Horde.Build.Models;
+using HordeCommon;
+
+namespace Horde.Build.Api
+{
+	/// <summary>
+	/// Parameters to register a new agent
+	/// </summary>
+	public class CreateAgentRequest
+	{
+		/// <summary>
+		/// Friendly name for the agent
+		/// </summary>
+		[Required]
+		public string Name { get; set; } = null!; // Enforced by [required] attribute
+
+		/// <summary>
+		/// Whether the agent is currently enabled
+		/// </summary>
+		public bool Enabled { get; set; } = true;
+
+		/// <summary>
+		/// Whether the agent is ephemeral (ie. should not be shown when inactive)
+		/// </summary>
+		public bool Ephemeral { get; set; }
+
+		/// <summary>
+		/// Per-agent override for the desired agent software channel
+		/// </summary>
+		public string? Channel { get; set; }
+
+		/// <summary>
+		/// Pools for this agent
+		/// </summary>
+		public List<string>? Pools { get; set; }
+	}
+
+	/// <summary>
+	/// Response from creating an agent
+	/// </summary>
+	public class CreateAgentResponse
+	{
+		/// <summary>
+		/// Unique id for the new agent
+		/// </summary>
+		public string Id { get; set; }
+
+		/// <summary>
+		/// Constructor
+		/// </summary>
+		/// <param name="id">Unique id for this agent</param>
+		public CreateAgentResponse(string id)
+		{
+			Id = id;
+		}
+	}
+
+	/// <summary>
+	/// Parameters to update an agent
+	/// </summary>
+	public class UpdateAgentRequest
+	{
+		/// <summary>
+		/// Whether the agent is currently enabled
+		/// </summary>
+		public bool? Enabled { get; set; }
+
+		/// <summary>
+		/// Whether the agent is ephemeral
+		/// </summary>
+		public bool? Ephemeral { get; set; }
+
+		/// <summary>
+		/// Request a conform be performed using the current agent set
+		/// </summary>
+		public bool? RequestConform { get; set; }
+
+		/// <summary>
+		/// Request that a full conform be performed, removing all intermediate files
+		/// </summary>
+		public bool? RequestFullConform { get; set; }
+
+		/// <summary>
+		/// Request the machine be restarted
+		/// </summary>
+		public bool? RequestRestart { get; set; }
+
+		/// <summary>
+		/// Request the machine be shut down
+		/// </summary>
+		public bool? RequestShutdown { get; set; }
+
+		/// <summary>
+		/// Per-agent override for the desired agent software channel
+		/// </summary>
+		public string? Channel { get; set; }
+
+		/// <summary>
+		/// Pools for this agent
+		/// </summary>
+		public List<string>? Pools { get; set; }
+
+		/// <summary>
+		/// New ACL for this agent
+		/// </summary>
+		public UpdateAclRequest? Acl { get; set; }
+		
+		/// <summary>
+		/// New comment
+		/// </summary>
+		public string? Comment { get; set; }
+	}
+
+	/// <summary>
+	/// Response for queries to find a particular lease within an agent
+	/// </summary>
+	public class GetAgentLeaseResponse
+	{
+		/// <summary>
+		/// Identifier for the lease
+		/// </summary>
+		public string Id { get; set; }
+
+		/// <summary>
+		/// The agent id
+		/// </summary>
+		public string? AgentId { get; set; }
+
+		/// <summary>
+		/// Cost of this agent, per hour
+		/// </summary>
+		public double? AgentRate { get; set; }
+
+		/// <summary>
+		/// Name of the lease
+		/// </summary>
+		public string Name { get; set; }
+
+		/// <summary>
+		/// Log id for this lease
+		/// </summary>
+		public string? LogId { get; set; }
+
+		/// <summary>
+		/// Time at which the lease started (UTC)
+		/// </summary>
+		public DateTime StartTime { get; set; }
+
+		/// <summary>
+		/// Time at which the lease started (UTC)
+		/// </summary>
+		public DateTime? FinishTime { get; set; }
+
+		/// <summary>
+		/// Whether this lease has started executing on the agent yet
+		/// </summary>
+		public bool Executing { get; set; }
+
+		/// <summary>
+		/// Details of the payload being executed
+		/// </summary>
+		public Dictionary<string, string>? Details { get; set; }
+
+		/// <summary>
+		/// Outcome of the lease
+		/// </summary>
+		public LeaseOutcome? Outcome { get; set; }
+
+		/// <summary>
+		/// State of the lease (for AgentLeases)
+		/// </summary>
+		public LeaseState? State { get; set; }
+
+		/// <summary>
+		/// Constructor
+		/// </summary>
+		/// <param name="lease">The lease to initialize from</param>
+		/// <param name="details">The payload details</param>
+		public GetAgentLeaseResponse(AgentLease lease, Dictionary<string, string>? details)
+		{
+			Id = lease.Id.ToString();
+			Name = lease.Name;
+			LogId = lease.LogId?.ToString();
+			State = lease.State;
+			StartTime = lease.StartTime;
+			Executing = lease.Active;
+			FinishTime = lease.ExpiryTime;
+			Details = details;
+		}
+
+		/// <summary>
+		/// Converts this lease to a response object
+		/// </summary>
+		/// <param name="lease">The lease to initialize from</param>
+		/// <param name="details">The payload details</param>
+		/// <param name="agentRate">Rate for running this agent</param>
+		public GetAgentLeaseResponse(ILease lease, Dictionary<string, string>? details, double? agentRate)
+		{
+			Id = lease.Id.ToString();
+			AgentId = lease.AgentId.ToString();
+			AgentRate = agentRate;
+			Name = lease.Name;
+			LogId = lease.LogId?.ToString();
+			StartTime = lease.StartTime;
+			Executing = (lease.FinishTime == null);
+			FinishTime = lease.FinishTime;
+			Details = details;
+			Outcome = lease.Outcome;
+		}
+	}
+
+	/// <summary>
+	/// Information about an agent session
+	/// </summary>
+	public class GetAgentSessionResponse
+	{
+		/// <summary>
+		/// Unique id for this session
+		/// </summary>
+		[Required]
+		public string Id { get; set; }
+
+		/// <summary>
+		/// Start time for this session
+		/// </summary>
+		[Required]
+		public DateTime StartTime { get; set; }
+
+		/// <summary>
+		/// Finishing time for this session
+		/// </summary>
+		public DateTime? FinishTime { get; set; }
+
+		/// <summary>
+		/// Properties of this agent
+		/// </summary>
+		public List<string>? Properties { get; set; }
+
+		/// <summary>
+		/// Version of the software running during this session
+		/// </summary>
+		public string? Version { get; set; }
+
+		/// <summary>
+		/// Constructor
+		/// </summary>
+		/// <param name="session">The session to construct from</param>
+		public GetAgentSessionResponse(ISession session)
+		{
+			Id = session.Id.ToString();
+			StartTime = session.StartTime;
+			FinishTime = session.FinishTime;
+			Properties = (session.Properties != null) ? new List<string>(session.Properties) : null;
+			Version = session.Version?.ToString();
+		}
+	}
+
+	/// <summary>
+	/// Information about a workspace synced on an agent
+	/// </summary>
+	public class GetAgentWorkspaceResponse
+	{
+		/// <summary>
+		/// The Perforce server and port to connect to
+		/// </summary>
+		public string? Cluster { get; set; }
+
+		/// <summary>
+		/// User to log into Perforce with (eg. buildmachine)
+		/// </summary>
+		public string? UserName { get; set; }
+
+		/// <summary>
+		/// Identifier to distinguish this workspace from other workspaces
+		/// </summary>
+		public string Identifier { get; set; }
+
+		/// <summary>
+		/// The stream to sync
+		/// </summary>
+		public string Stream { get; set; }
+
+		/// <summary>
+		/// Custom view for the workspace
+		/// </summary>
+		public List<string>? View { get; set; }
+
+		/// <summary>
+		/// Whether to use an incremental workspace
+		/// </summary>
+		public bool BIncremental { get; set; }
+
+		/// <summary>
+		/// Constructor
+		/// </summary>
+		/// <param name="workspace">The workspace to construct from</param>
+		public GetAgentWorkspaceResponse(AgentWorkspace workspace)
+		{
+			Cluster = workspace.Cluster;
+			UserName = workspace.UserName;
+			Identifier = workspace.Identifier;
+			Stream = workspace.Stream;
+			View = workspace.View;
+			BIncremental = workspace.BIncremental;
+		}
+	}
+
+	/// <summary>
+	/// Information about an agent
+	/// </summary>
+	public class GetAgentResponse
+	{
+		/// <summary>
+		/// The agent's unique ID
+		/// </summary>
+		public string Id { get; set; }
+
+		/// <summary>
+		/// Friendly name of the agent
+		/// </summary>
+		public string Name { get; set; }
+
+		/// <summary>
+		/// Whether the agent is currently enabled
+		/// </summary>
+		public bool Enabled { get; set; }
+
+		/// <summary>
+		/// Cost estimate per-hour for this agent
+		/// </summary>
+		public double? Rate { get; set; }
+
+		/// <summary>
+		/// The current session id
+		/// </summary>
+		public string? SessionId { get; set; }
+
+		/// <summary>
+		/// Whether the agent is ephemeral
+		/// </summary>
+		public bool Ephemeral { get; set; }
+
+		/// <summary>
+		/// Whether the agent is currently online
+		/// </summary>
+		public bool Online { get; set; }
+
+		/// <summary>
+		/// Whether this agent has expired
+		/// </summary>
+		public bool Deleted { get; set; }
+
+		/// <summary>
+		/// Whether a conform job is pending
+		/// </summary>
+		public bool PendingConform { get; set; }
+
+		/// <summary>
+		/// Whether a full conform job is pending
+		/// </summary>
+		public bool PendingFullConform { get; set; }
+
+		/// <summary>
+		/// Whether a restart is pending
+		/// </summary>
+		public bool PendingRestart { get; set; }
+		
+		/// <summary>
+		/// Whether a shutdown is pending
+		/// </summary>
+		public bool PendingShutdown { get; set; }
+
+		/// <summary>
+		/// The reason for the last shutdown
+		/// </summary>
+		public string LastShutdownReason { get; set; }
+
+		/// <summary>
+		/// Last time a conform was attempted
+		/// </summary>
+		public DateTime LastConformTime { get; set; }
+
+		/// <summary>
+		/// Number of times a conform has been attempted
+		/// </summary>
+		public int? ConformAttemptCount { get; set; }
+
+		/// <summary>
+		/// Last time a conform was attempted
+		/// </summary>
+		public DateTime? NextConformTime { get; set; }
+
+		/// <summary>
+		/// The current client version
+		/// </summary>
+		public string? Version { get; set; }
+
+		/// <summary>
+		/// Per-agent override for the desired agent software channel
+		/// </summary>
+		public string? Channel { get; set; }
+
+		/// <summary>
+		/// Properties for the agent
+		/// </summary>
+		public List<string> Properties { get; set; }
+
+		/// <summary>
+		/// Resources for the agent
+		/// </summary>
+		public Dictionary<string, int> Resources { get; set; }
+
+		/// <summary>
+		/// Last update time of this agent
+		/// </summary>
+		public DateTime? UpdateTime { get; set; }
+
+		/// <summary>
+		/// Pools for this agent
+		/// </summary>
+		public List<string>? Pools { get; set; }
+
+		/// <summary>
+		/// List of workspaces currently synced to this machine
+		/// </summary>
+		public List<GetAgentWorkspaceResponse> Workspaces { get; set; } = new List<GetAgentWorkspaceResponse>();
+
+		/// <summary>
+		/// Capabilities of this agent
+		/// </summary>
+		public object? Capabilities { get; }
+
+		/// <summary>
+		/// Array of active leases.
+		/// </summary>
+		public List<GetAgentLeaseResponse> Leases { get; }
+
+		/// <summary>
+		/// Per-object permissions
+		/// </summary>
+		public GetAclResponse? Acl { get; set; }
+		
+		/// <summary>
+		/// Comment for this agent
+		/// </summary>
+		public string? Comment { get; set; }
+
+		/// <summary>
+		/// Constructor
+		/// </summary>
+		/// <param name="agent">The agent to construct from</param>
+		/// <param name="leases">Active leases</param>
+		/// <param name="rate">Rate for this agent</param>
+		/// <param name="bIncludeAcl">Whether to include the ACL in the response</param>
+		public GetAgentResponse(IAgent agent, List<GetAgentLeaseResponse> leases, double? rate, bool bIncludeAcl)
+		{
+			Id = agent.Id.ToString();
+			Name = agent.Id.ToString();
+			Enabled = agent.Enabled;
+			Rate = rate;
+			Properties = new List<string>(agent.Properties);
+			Resources = new Dictionary<string, int>(agent.Resources);
+			SessionId = agent.SessionId?.ToString();
+			Online = agent.IsSessionValid(DateTime.UtcNow);
+			Deleted = agent.Deleted;
+			PendingConform = agent.RequestConform;
+			PendingFullConform = agent.RequestFullConform;
+			PendingRestart = agent.RequestRestart;
+			PendingShutdown = agent.RequestShutdown;
+			LastShutdownReason = agent.LastShutdownReason ?? "Unknown";
+			LastConformTime = agent.LastConformTime;
+			NextConformTime =   agent.LastConformTime;
+			ConformAttemptCount = agent.ConformAttemptCount;
+			Version = agent.Version?.ToString() ?? "Unknown";
+			if(agent.Channel != null)
+			{
+				Version += $" ({agent.Channel})";
+			}
+			Version = agent.Version?.ToString();
+			Channel = agent.Channel?.ToString();
+			UpdateTime = agent.UpdateTime;
+			Pools = agent.GetPools().Select(x => x.ToString()).ToList();
+			Workspaces = agent.Workspaces.ConvertAll(x => new GetAgentWorkspaceResponse(x));
+			Capabilities = new { Devices = new[] { new { agent.Properties, agent.Resources } } };
+			Leases = leases;
+			Acl = (bIncludeAcl && agent.Acl != null) ? new GetAclResponse(agent.Acl) : null;
+			Comment = agent.Comment;
+		}
+	}
+}
diff --git a/Engine/Source/Programs/Horde/Horde.Build/Agents/AgentsController.cs b/Engine/Source/Programs/Horde/Horde.Build/Agents/AgentsController.cs
new file mode 100644
--- /dev/null
+++ b/Engine/Source/Programs/Horde/Horde.Build/Agents/AgentsController.cs
@@ -0,0 +1,388 @@
+// Copyright Epic Games, Inc. All Rights Reserved.
+
+using System;
+using System.Collections.Generic;
+using System.Linq;
+using System.Threading.Tasks;
+using Horde.Build.Acls;
+using Horde.Build.Api;
+using Horde.Build.Collections;
+using Horde.Build.Models;
+using Horde.Build.Services;
+using Horde.Build.Utilities;
+using Microsoft.AspNetCore.Authorization;
+using Microsoft.AspNetCore.Mvc;
+using Microsoft.Extensions.Logging;
+
+namespace Horde.Build.Controllers
+{
+	using AgentSoftwareChannelName = StringId<AgentSoftwareChannels>;
+	using LeaseId = ObjectId<ILease>;
+	using PoolId = StringId<IPool>;
+	using SessionId = ObjectId<ISession>;
+
+	/// <summary>
+	/// Controller for the /api/v1/agents endpoint
+	/// </summary>
+	[ApiController]
+	[Authorize]
+	[Route("[controller]")]
+	public class AgentsController : HordeControllerBase
+	{
+		readonly AclService _aclService;
+		readonly AgentService _agentService;
+
+		/// <summary>
+		/// Constructor
+		/// </summary>
+		/// <param name="aclService">The ACL service singleton</param>
+		/// <param name="agentService">The agent service</param>
+		public AgentsController(AclService aclService, AgentService agentService)
+		{
+			_aclService = aclService;
+			_agentService = agentService;
+		}
+
+		/// <summary>
+		/// Register an agent to perform remote work.
+		/// </summary>
+		/// <param name="request">Request parameters</param>
+		/// <returns>Information about the registered agent</returns>
+		[HttpPost]
+		[Route("/api/v1/agents")]
+		public async Task<ActionResult<CreateAgentResponse>> CreateAgentAsync([FromBody] CreateAgentRequest request)
+		{
+			if(!await _aclService.AuthorizeAsync(AclAction.CreateAgent, User))
+			{
+				return Forbid(AclAction.CreateAgent);
+			}
+
+			AgentSoftwareChannelName? channel = String.IsNullOrEmpty(request.Channel) ? (AgentSoftwareChannelName?)null : new AgentSoftwareChannelName(request.Channel);
+			IAgent agent = await _agentService.CreateAgentAsync(request.Name, request.Enabled, channel, request.Pools?.ConvertAll(x => new PoolId(x)));
+
+			return new CreateAgentResponse(agent.Id.ToString());
+		}
+
+		/// <summary>
+		/// Finds the agents matching specified criteria.
+		/// </summary>
+		/// <param name="poolId">The pool containing the agent</param>
+		/// <param name="index">First result to return</param>
+		/// <param name="count">Number of results to return</param>
+		/// <param name="modifiedAfter">If set, only returns agents modified after this time</param>
+		/// <param name="filter">Filter for the properties to return</param>
+		/// <returns>List of matching agents</returns>
+		[HttpGet]
+		[Route("/api/v1/agents")]
+		[ProducesResponseType(typeof(List<GetAgentResponse>), 200)]
+		public async Task<ActionResult<List<object>>> FindAgentsAsync([FromQuery] PoolId? poolId = null, [FromQuery] int? index = null, [FromQuery] int? count = null, [FromQuery] DateTimeOffset? modifiedAfter = null, [FromQuery] PropertyFilter? filter = null)
+		{
+			GlobalPermissionsCache permissionsCache = new GlobalPermissionsCache();
+			List<IAgent> agents = await _agentService.FindAgentsAsync(poolId, modifiedAfter?.UtcDateTime, index, count);
+
+			List<object> responses = new List<object>();
+			foreach (IAgent agent in agents)
+			{
+				responses.Add(await GetAgentResponseAsync(agent, permissionsCache, filter));
+			}
+
+			return responses;
+		}
+
+		/// <summary>
+		/// Retrieve information about a specific agent
+		/// </summary>
+		/// <param name="agentId">Id of the agent to get information about</param>
+		/// <param name="filter">Filter for the properties to return</param>
+		/// <returns>Information about the requested agent</returns>
+		[HttpGet]
+		[Route("/api/v1/agents/{agentId}")]
+		[ProducesResponseType(typeof(GetAgentResponse), 200)]
+		public async Task<ActionResult<object>> GetAgentAsync(AgentId agentId, [FromQuery] PropertyFilter? filter = null)
+		{
+			IAgent? agent = await _agentService.GetAgentAsync(agentId);
+			if (agent == null)
+			{
+				return NotFound(agentId);
+			}
+
+			return await GetAgentResponseAsync(agent, new GlobalPermissionsCache(), filter);
+		}
+
+		/// <summary>
+		/// Gets an individual agent response
+		/// </summary>
+		async ValueTask<object> GetAgentResponseAsync(IAgent agent, GlobalPermissionsCache permissionsCache, PropertyFilter? filter = null)
+		{
+			double? rate = null;
+			if (await _agentService.AuthorizeAsync(agent, AclAction.ViewCosts, User, permissionsCache))
+			{
+				rate = await _agentService.GetRateAsync(agent.Id);
+			}
+
+			List<GetAgentLeaseResponse> leases = new List<GetAgentLeaseResponse>();
+			foreach (AgentLease lease in agent.Leases)
+			{
+				Dictionary<string, string>? details = _agentService.GetPayloadDetails(lease.Payload);
+				leases.Add(new GetAgentLeaseResponse(lease, details));
+			}
+
+			bool bIncludeAcl = await _agentService.AuthorizeAsync(agent, AclAction.ViewPermissions, User, permissionsCache);
+			return new GetAgentResponse(agent, leases, rate, bIncludeAcl).ApplyFilter(filter);
+		}
+
+		/// <summary>
+		/// Update an agent's properties.
+		/// </summary>
+		/// <param name="agentId">Id of the agent to update.</param>
+		/// <param name="update">Properties on the agent to update.</param>
+		/// <returns>Http result code</returns>
+		[HttpPut]
+		[Route("/api/v1/agents/{agentId}")]
+		public async Task<ActionResult> UpdateAgentAsync(AgentId agentId, [FromBody] UpdateAgentRequest update)
+		{
+			List<PoolId>? updatePools = update.Pools?.ConvertAll(x => new PoolId(x));
+			AgentSoftwareChannelName? channel = String.IsNullOrEmpty(update.Channel) ? (AgentSoftwareChannelName?)null : new AgentSoftwareChannelName(update.Channel);
+
+			string userName = User.GetUserName() ?? "Unknown";
+
+			GlobalPermissionsCache cache = new GlobalPermissionsCache();
+			for (; ; )
+			{
+				IAgent? agent = await _agentService.GetAgentAsync(agentId);
+				if (agent == null)
+				{
+					return NotFound(agentId);
+				}
+				if (!await _agentService.AuthorizeAsync(agent, AclAction.UpdateAgent, User, cache))
+				{
+					return Forbid(AclAction.UpdateAgent, agentId);
+				}
+				if (update.Acl != null && !await _agentService.AuthorizeAsync(agent, AclAction.ChangePermissions, User, cache))
+				{
+					return Forbid(AclAction.ChangePermissions, agentId);
+				}
+
+				IAgent? newAgent = await _agentService.Agents.TryUpdateSettingsAsync(agent, update.Enabled, update.RequestConform, update.RequestFullConform, update.RequestRestart, update.RequestShutdown, $"Manual ({userName})", channel, update.Pools?.ConvertAll(x => new PoolId(x)), Acl.Merge(agent.Acl, update.Acl), update.Comment);
+				if (newAgent == null)
+				{
+					continue;
+				}
+
+				IAuditLogChannel<AgentId> logger = _agentService.Agents.GetLogger(agent.Id);
+				if (agent.Enabled != newAgent.Enabled)
+				{
+					logger.LogInformation("Setting changed: Enabled = {State} ({UserName})", newAgent.Enabled, userName);
+				}
+				if (agent.RequestConform != newAgent.RequestConform)
+				{
+					logger.LogInformation("Setting changed: RequestConform = {State} ({UserName})", newAgent.RequestConform, userName);
+				}
+				if (agent.RequestFullConform != newAgent.RequestFullConform)
+				{
+					logger.LogInformation("Setting changed: RequestFullConform = {State} ({UserName})", newAgent.RequestFullConform, userName);
+				}
+				if (agent.RequestRestart != newAgent.RequestRestart)
+				{
+					logger.LogInformation("Setting changed: RequestRestart = {State} ({UserName})", newAgent.RequestRestart, userName);
+				}
+				if (agent.RequestShutdown != newAgent.RequestShutdown)
+				{
+					logger.LogInformation("Setting changed: RequestShutdown = {State} ({UserName})", newAgent.RequestShutdown, userName);
+				}
+				if (agent.Comment != newAgent.Comment)
+				{
+					logger.LogInformation("Setting changed: Comment = \"{Comment}\" ({UserName})", update.Comment, userName);
+				}
+				foreach (PoolId addedPool in newAgent.ExplicitPools.Except(agent.ExplicitPools))
+				{
+					logger.LogInformation("Added to pool {PoolId} ({UserName})", addedPool, userName);
+				}
+				foreach (PoolId removedPool in agent.ExplicitPools.Except(newAgent.ExplicitPools))
+				{
+					logger.LogInformation("Removed from pool {PoolId} ({UserName})", removedPool, userName);
+				}
+				break;
+			}
+			return Ok();
+		}
+
+		/// <summary>
+		/// Remove a registered agent.
+		/// </summary>
+		/// <param name="agentId">Id of the agent to delete.</param>
+		/// <returns>Http result code</returns>
+		[HttpDelete]
+		[Route("/api/v1/agents/{agentId}")]
+		public async Task<ActionResult> DeleteAgentAsync(AgentId agentId)
+		{
+			IAgent? agent = await _agentService.GetAgentAsync(agentId);
+			if (agent == null)
+			{
+				return NotFound(agentId);
+			}
+			if (!await _agentService.AuthorizeAsync(agent, AclAction.DeleteAgent, User, null))
+			{
+				return Forbid(AclAction.DeleteAgent, agentId);
+			}
+
+			await _agentService.DeleteAgentAsync(agent);
+			return new OkResult();
+		}
+
+		/// <summary>
+		/// Retrieve historical information about a specific agent
+		/// </summary>
+		/// <param name="agentId">Id of the agent to get information about</param>
+		/// <param name="minTime">Minimum time for records to return</param>
+		/// <param name="maxTime">Maximum time for records to return</param>
+		/// <param name="index">Offset of the first result</param>
+		/// <param name="count">Number of records to return</param>
+		/// <returns>Information about the requested agent</returns>
+		[HttpGet]
+		[Route("/api/v1/agents/{agentId}/history")]
+		public async Task GetAgentHistoryAsync(AgentId agentId, [FromQuery] DateTime? minTime = null, [FromQuery] DateTime? maxTime = null, [FromQuery] int index = 0, [FromQuery] int count = 50)
+		{
+			Response.ContentType = "application/json";
+			Response.StatusCode = 200;
+			await Response.StartAsync();
+			await _agentService.Agents.GetLogger(agentId).FindAsync(Response.BodyWriter, minTime, maxTime, index, count);
+		}
+
+		/// <summary>
+		/// Find all the sessions of a particular agent
+		/// </summary>
+		/// <param name="agentId">Unique id of the agent to find</param>
+		/// <param name="startTime">Start time to include in the search</param>
+		/// <param name="finishTime">Finish time to include in the search</param>
+		/// <param name="index">Index of the first result to return</param>
+		/// <param name="count">Number of results to return</param>
+		/// <returns>Sessions </returns>
+		[HttpGet]
+		[Route("/api/v1/agents/{agentId}/sessions")]
+		public async Task<ActionResult<List<GetAgentSessionResponse>>> FindSessionsAsync(AgentId agentId, [FromQuery] DateTimeOffset? startTime, [FromQuery] DateTimeOffset? finishTime, [FromQuery] int index = 0, [FromQuery] int count = 50)
+		{
+			IAgent? agent = await _agentService.GetAgentAsync(agentId);
+			if (agent == null)
+			{
+				return NotFound(agentId);
+			}
+			if (!await _agentService.AuthorizeAsync(agent, AclAction.ViewSession, User, null))
+			{
+				return Forbid(AclAction.ViewSession, agentId);
+			}
+
+			List<ISession> sessions = await _agentService.FindSessionsAsync(agentId, startTime?.UtcDateTime, finishTime?.UtcDateTime, index, count);
+			return sessions.ConvertAll(x => new GetAgentSessionResponse(x));
+		}
+
+		/// <summary>
+		/// Find all the sessions of a particular agent
+		/// </summary>
+		/// <param name="agentId">Unique id of the agent to find</param>
+		/// <param name="sessionId">Unique id of the session</param>
+		/// <returns>Sessions </returns>
+		[HttpGet]
+		[Route("/api/v1/agents/{agentId}/sessions/{sessionId}")]
+		public async Task<ActionResult<GetAgentSessionResponse>> GetSessionAsync(AgentId agentId, SessionId sessionId)
+		{
+			IAgent? agent = await _agentService.GetAgentAsync(agentId);
+			if (agent == null)
+			{
+				return NotFound(agentId);
+			}
+			if (!await _agentService.AuthorizeAsync(agent, AclAction.ViewSession, User, null))
+			{
+				return Forbid(AclAction.ViewSession, agentId);
+			}
+
+			ISession? session = await _agentService.GetSessionAsync(sessionId);
+			if(session == null || session.AgentId != agentId)
+			{
+				return NotFound();
+			}
+
+			return new GetAgentSessionResponse(session);
+		}
+
+		/// <summary>
+		/// Find all the leases for a particular agent
+		/// </summary>
+		/// <param name="agentId">Unique id of the agent to find</param>
+		/// <param name="sessionId">The session to query</param>
+		/// <param name="startTime">Start of the time window to consider</param>
+		/// <param name="finishTime">End of the time window to consider</param>
+		/// <param name="index">Index of the first result to return</param>
+		/// <param name="count">Number of results to return</param>
+		/// <param name="filter">Filter to apply to the properties</param>
+		/// <returns>Sessions </returns>
+		[HttpGet]
+		[Route("/api/v1/agents/{agentId}/leases")]
+		[ProducesResponseType(200, Type = typeof(List<GetAgentLeaseResponse>))]
+		public async Task<ActionResult<List<object>>> FindLeasesAsync(AgentId agentId, [FromQuery] SessionId? sessionId, [FromQuery] DateTimeOffset? startTime, [FromQuery] DateTimeOffset? finishTime, [FromQuery] int index = 0, [FromQuery] int count = 1000, [FromQuery] PropertyFilter? filter = null)
+		{
+			IAgent? agent = await _agentService.GetAgentAsync(agentId);
+			if (agent == null)
+			{
+				return NotFound(agentId);
+			}
+			if (!await _agentService.AuthorizeAsync(agent, AclAction.ViewLeases, User, null))
+			{
+				return Forbid(AclAction.ViewLeases, agentId);
+			}
+
+			List<ILease> leases = await _agentService.FindLeasesAsync(agentId, sessionId, startTime?.UtcDateTime, finishTime?.UtcDateTime, index, count);
+
+			double? agentRate = null;
+			if (await _aclService.AuthorizeAsync(AclAction.ViewCosts, User))
+			{
+				agentRate = await _agentService.GetRateAsync(agentId);
+			}
+
+			List<object> responses = new List<object>();
+			foreach(ILease lease in leases)
+			{
+				Dictionary<string, string>? details = _agentService.GetPayloadDetails(lease.Payload);
+				responses.Add(PropertyFilter.Apply(new GetAgentLeaseResponse(lease, details, agentRate), filter));
+			}
+
+			return responses;
+		}
+
+		/// <summary>
+		/// Get info about a particular lease
+		/// </summary>
+		/// <param name="agentId">Unique id of the agent to find</param>
+		/// <param name="leaseId">Unique id of the particular lease</param>
+		/// <returns>Lease matching the given id</returns>
+		[HttpGet]
+		[Route("/api/v1/agents/{agentId}/leases/{leaseId}")]
+		public async Task<ActionResult<GetAgentLeaseResponse>> GetLeaseAsync(AgentId agentId, LeaseId leaseId)
+		{
+			IAgent? agent = await _agentService.GetAgentAsync(agentId);
+			if(agent == null)
+			{
+				return NotFound(agentId);
+			}
+			if (!await _aclService.AuthorizeAsync(AclAction.ViewLeases, User))
+			{
+				return Forbid(AclAction.ViewLeases, agentId);
+			}
+
+			ILease? lease = await _agentService.GetLeaseAsync(leaseId);
+			if (lease == null || lease.AgentId != agentId)
+			{
+				return NotFound(agentId, leaseId);
+			}
+
+			double? agentRate = null;
+			if (await _aclService.AuthorizeAsync(AclAction.ViewCosts, User))
+			{
+				agentRate = await _agentService.GetRateAsync(agentId);
+			}
+
+			Dictionary<string, string>? details = _agentService.GetPayloadDetails(lease.Payload);
+			return new GetAgentLeaseResponse(lease, details, agentRate);
+		}
+	}
+}
