Horde: Finish re-org of Horde solution started a while ago.

diff --git a/Engine/Source/Programs/Horde/Horde.Build/Agents/Fleet/AutoscaleService.cs b/Engine/Source/Programs/Horde/Horde.Build/Agents/Fleet/AutoscaleService.cs
new file mode 100644
--- /dev/null
+++ b/Engine/Source/Programs/Horde/Horde.Build/Agents/Fleet/AutoscaleService.cs
@@ -0,0 +1,280 @@
+// Copyright Epic Games, Inc. All Rights Reserved.
+
+using System;
+using System.Collections.Generic;
+using System.Diagnostics;
+using System.Linq;
+using System.Threading;
+using System.Threading.Tasks;
+using Google.Protobuf.WellKnownTypes;
+using Horde.Build.Collections;
+using Horde.Build.Models;
+using Horde.Build.Server;
+using Horde.Build.Utilities;
+using HordeCommon;
+using HordeCommon.Rpc.Tasks;
+using Microsoft.Extensions.Hosting;
+using Microsoft.Extensions.Logging;
+using OpenTracing;
+using OpenTracing.Util;
+using StatsdClient;
+
+namespace Horde.Build.Services
+{
+	using PoolId = StringId<IPool>;
+
+	/// <summary>
+	/// Service for managing the autoscaling of agent pools
+	/// </summary>
+	public sealed class AutoscaleService : IHostedService, IDisposable
+	{
+		const int NumSamples = 10;
+		const int NumSamplesForResult = 9;
+		static readonly TimeSpan s_sampleTime = TimeSpan.FromMinutes(6.0);
+		static readonly TimeSpan s_shrinkPoolCoolDown = TimeSpan.FromMinutes(20.0);
+
+		struct UtilizationSample
+		{
+			public double _jobWork;
+			public double _otherWork;
+		}
+
+		class AgentData
+		{
+			public IAgent Agent { get; }
+			public UtilizationSample[] Samples { get; } = new UtilizationSample[NumSamples];
+			public List<ILease> Leases { get; } = new List<ILease>();
+
+			public AgentData(IAgent agent)
+			{
+				Agent = agent;
+			}
+
+			public void Add(double minT, double maxT, double jobWork, double otherWork)
+			{
+				int minIdx = Math.Clamp((int)minT, 0, NumSamples - 1);
+				int maxIdx = Math.Clamp((int)maxT, minIdx, NumSamples - 1);
+				for (int idx = minIdx; idx <= maxIdx; idx++)
+				{
+					double fraction = Math.Clamp(maxT - idx, 0.0, 1.0) - Math.Clamp(minT - idx, 0.0, 1.0);
+					Samples[idx]._jobWork += jobWork * fraction;
+					Samples[idx]._otherWork += otherWork * fraction;
+				}
+			}
+		}
+
+		class PoolData
+		{
+			public IPool Pool { get; }
+			public List<IAgent> Agents { get; } = new List<IAgent>();
+			public UtilizationSample[] Samples { get; } = new UtilizationSample[NumSamples];
+
+			public PoolData(IPool pool)
+			{
+				Pool = pool;
+			}
+
+			public void Add(AgentData agentData)
+			{
+				Agents.Add(agentData.Agent);
+				for (int idx = 0; idx < Samples.Length; idx++)
+				{
+					Samples[idx]._jobWork += agentData.Samples[idx]._jobWork;
+					Samples[idx]._otherWork += agentData.Samples[idx]._otherWork;
+				}
+			}
+		}
+
+		readonly IAgentCollection _agentCollection;
+		readonly IPoolCollection _poolCollection;
+		readonly ILeaseCollection _leaseCollection;
+		readonly IFleetManager _fleetManager;
+		readonly IDogStatsd _dogStatsd;
+		readonly IClock _clock;
+		readonly ITicker _ticker;
+		readonly ILogger<AutoscaleService> _logger;
+
+		/// <summary>
+		/// Constructor
+		/// </summary>
+		public AutoscaleService(IAgentCollection agentCollection, IPoolCollection poolCollection, ILeaseCollection leaseCollection, IFleetManager fleetManager, IDogStatsd dogStatsd, MongoService mongoService, IClock clock, ILogger<AutoscaleService> logger)
+		{
+			_agentCollection = agentCollection;
+			_poolCollection = poolCollection;
+			_leaseCollection = leaseCollection;
+			_fleetManager = fleetManager;
+			_dogStatsd = dogStatsd;
+			_clock = clock;
+			if (mongoService.ReadOnlyMode)
+			{
+				_ticker = new NullTicker();
+			}
+			else
+			{
+				_ticker = clock.AddSharedTicker<AutoscaleService>(TimeSpan.FromMinutes(5.0), TickLeaderAsync, logger);
+			}
+			_logger = logger;
+		}
+
+		/// <inheritdoc/>
+		public Task StartAsync(CancellationToken cancellationToken) => _ticker.StartAsync();
+
+		/// <inheritdoc/>
+		public Task StopAsync(CancellationToken cancellationToken) => _ticker.StopAsync();
+
+		/// <inheritdoc/>
+		public void Dispose() => _ticker.Dispose();
+
+		/// <inheritdoc/>
+		async ValueTask TickLeaderAsync(CancellationToken stoppingToken)
+		{
+			DateTime utcNow = _clock.UtcNow;
+
+			_logger.LogInformation("Autoscaling pools...");
+			Stopwatch stopwatch = Stopwatch.StartNew();
+
+			Dictionary<PoolId, PoolData> poolToData;
+			using (IScope scope = GlobalTracer.Instance.BuildSpan("Compute utilization").StartActive())
+			{
+				// Find all the current agents
+				List<IAgent> agents = await _agentCollection.FindAsync(status: AgentStatus.Ok);
+
+				// Query leases in last interval
+				DateTime maxTime = _clock.UtcNow;
+				DateTime minTime = maxTime - (s_sampleTime * NumSamples);
+				List<ILease> leases = await _leaseCollection.FindLeasesAsync(minTime, maxTime);
+
+				// Add all the leases to a data object for each agent
+				Dictionary<AgentId, AgentData> agentIdToData = agents.ToDictionary(x => x.Id, x => new AgentData(x));
+				foreach (ILease lease in leases)
+				{
+					AgentData? agentData;
+					if (agentIdToData.TryGetValue(lease.AgentId, out agentData) &&
+					    agentData.Agent.SessionId == lease.SessionId)
+					{
+						agentData.Leases.Add(lease);
+					}
+				}
+
+				// Compute utilization for each agent
+				foreach (AgentData agentData in agentIdToData.Values)
+				{
+					foreach (ILease lease in agentData.Leases.OrderBy(x => x.StartTime))
+					{
+						double minT = (lease.StartTime - minTime).TotalSeconds / s_sampleTime.TotalSeconds;
+						double maxT = (lease.FinishTime == null)
+							? NumSamples
+							: ((lease.FinishTime.Value - minTime).TotalSeconds / s_sampleTime.TotalSeconds);
+
+						Any payload = Any.Parser.ParseFrom(lease.Payload.ToArray());
+						if (payload.Is(ExecuteJobTask.Descriptor))
+						{
+							agentData.Add(minT, maxT, 1.0, 0.0);
+						}
+						else
+						{
+							agentData.Add(minT, maxT, 0.0, 1.0);
+						}
+					}
+				}
+
+				// Get all the pools
+				List<IPool> pools = await _poolCollection.GetAsync();
+				poolToData = pools.ToDictionary(x => x.Id, x => new PoolData(x));
+
+				// Find pool utilization over the query period
+				foreach (AgentData agentData in agentIdToData.Values)
+				{
+					foreach (PoolId poolId in agentData.Agent.GetPools())
+					{
+						PoolData? poolData;
+						if (poolToData.TryGetValue(poolId, out poolData))
+						{
+							poolData.Add(agentData);
+						}
+					}
+				}
+			}
+
+			// Output all the final stats
+			foreach (PoolData poolData in poolToData.Values.OrderByDescending(x => x.Agents.Count))
+			{
+				IPool pool = poolData.Pool;
+
+				int minAgents = pool.MinAgents ?? 1;
+				int numReserveAgents = pool.NumReserveAgents ?? 5;
+
+				double utilization = poolData.Samples.Select(x => x._jobWork).OrderByDescending(x => x).Skip(NumSamples - NumSamplesForResult).First();
+
+				int targetAgents = Math.Max((int)utilization + numReserveAgents, minAgents);
+				int delta = targetAgents - poolData.Agents.Count;
+
+				_logger.LogInformation("{PoolName,-48} Jobs=[{JobUtilization}] Total=[{OtherUtilization}] Min={Min,5:0.0} Max={Max,5:0.0} Avg={Avg,5:0.0} Pct={Pct,5:0.0} Current={Current,4} Target={Target,4} Delta={Delta}",
+					pool.Name,
+					GetDensityMap(poolData.Samples.Select(x => x._jobWork / Math.Max(1, poolData.Agents.Count))),
+					GetDensityMap(poolData.Samples.Select(x => x._otherWork / Math.Max(1, poolData.Agents.Count))),
+					poolData.Samples.Min(x => x._jobWork),
+					poolData.Samples.Max(x => x._jobWork),
+					utilization,
+					poolData.Samples.Sum(x => x._jobWork) / NumSamples,
+					poolData.Agents.Count,
+					targetAgents,
+					delta);
+
+				try
+				{
+					if (pool.EnableAutoscaling)
+					{
+						using (IScope scope = GlobalTracer.Instance.BuildSpan("Scaling pool").StartActive())
+						{
+							scope.Span.SetTag("poolName", pool.Name);
+							scope.Span.SetTag("delta", delta);
+
+							if (delta > 0)
+							{
+								await _fleetManager.ExpandPoolAsync(pool, poolData.Agents, delta);
+								await _poolCollection.TryUpdateAsync(pool, lastScaleUpTime: DateTime.UtcNow);
+							}
+
+							if (delta < 0)
+							{
+								bool bShrinkIsOnCoolDown = pool.LastScaleDownTime != null && pool.LastScaleDownTime + s_shrinkPoolCoolDown > DateTime.UtcNow;
+								if (!bShrinkIsOnCoolDown)
+								{
+									await _fleetManager.ShrinkPoolAsync(pool, poolData.Agents, -delta);
+									await _poolCollection.TryUpdateAsync(pool, lastScaleDownTime: DateTime.UtcNow);
+								}
+								else
+								{
+									_logger.LogDebug("Cannot shrink {PoolName} right now, it's on cool-down until {CoolDownTimeEnds}", pool.Name, pool.LastScaleDownTime + s_shrinkPoolCoolDown);
+								}
+							}
+						}
+					}
+				}
+				catch (Exception ex)
+				{
+					_logger.LogError(ex, "Failed to scale {PoolName}:\n{Exception}", pool.Name, ex);
+					continue;
+				}
+
+				_dogStatsd.Gauge("agentpools.autoscale.target", targetAgents, tags: new []{"pool:" + pool.Name});
+				_dogStatsd.Gauge("agentpools.autoscale.current", poolData.Agents.Count, tags: new []{"pool:" + pool.Name});
+			}
+			
+			stopwatch.Stop();
+			_logger.LogInformation("Autoscaling pools took {ElapsedTime} ms", stopwatch.ElapsedMilliseconds);
+		}
+
+		/// <summary>
+		/// Creates a string of characters indicating a sequence of 0-1 values over time
+		/// </summary>
+		/// <param name="values">Sequence of values</param>
+		/// <returns>Density map</returns>
+		static string GetDensityMap(IEnumerable<double> values)
+		{
+			const string Greyscale = " 123456789";
+			return new string(values.Select(x => Greyscale[Math.Clamp((int)(x * Greyscale.Length), 0, Greyscale.Length - 1)]).ToArray());
+		}
+	}
+}
diff --git a/Engine/Source/Programs/Horde/Horde.Build/Agents/Fleet/AutoscaleServiceV2.cs b/Engine/Source/Programs/Horde/Horde.Build/Agents/Fleet/AutoscaleServiceV2.cs
new file mode 100644
--- /dev/null
+++ b/Engine/Source/Programs/Horde/Horde.Build/Agents/Fleet/AutoscaleServiceV2.cs
@@ -0,0 +1,210 @@
+// Copyright Epic Games, Inc. All Rights Reserved.
+
+using System;
+using System.Collections.Generic;
+using System.Diagnostics;
+using System.Linq;
+using System.Threading;
+using System.Threading.Tasks;
+using Horde.Build.Collections;
+using Horde.Build.Models;
+using Horde.Build.Services;
+using Horde.Build.Utilities;
+using HordeCommon;
+using Microsoft.Extensions.Hosting;
+using Microsoft.Extensions.Logging;
+using Microsoft.Extensions.Options;
+using OpenTracing;
+using OpenTracing.Util;
+using StatsdClient;
+
+namespace Horde.Build.Fleet.Autoscale
+{
+	using PoolId = StringId<IPool>;
+
+	/// <summary>
+	/// Service for managing the autoscaling of agent pools
+	/// </summary>
+	public sealed class AutoscaleServiceV2 : IHostedService, IDisposable
+	{
+		private IPoolSizeStrategy _leaseUtilizationStrategy;
+		private IPoolSizeStrategy _jobQueueStrategy;
+		private IPoolSizeStrategy _noOpStrategy;
+		private readonly IAgentCollection _agentCollection;
+		private readonly IPoolCollection _poolCollection;
+		private readonly IFleetManager _fleetManager;
+		private readonly IDogStatsd _dogStatsd;
+		private readonly IClock _clock;
+		private readonly ITicker _ticker;
+		private readonly ServerSettings _settings;
+		private readonly TimeSpan _defaultScaleOutCooldown;
+		private readonly TimeSpan _defaultScaleInCooldown;
+		private readonly ILogger<AutoscaleServiceV2> _logger;
+
+		/// <summary>
+		/// Constructor
+		/// </summary>
+		public AutoscaleServiceV2(LeaseUtilizationStrategy leaseUtilizationStrategy, JobQueueStrategy jobQueueStrategy, NoOpPoolSizeStrategy noOpStrategy, IAgentCollection agentCollection, IPoolCollection poolCollection, IFleetManager fleetManager, IDogStatsd dogStatsd, IClock clock, IOptions<ServerSettings> settings, ILogger<AutoscaleServiceV2> logger)
+		{
+			_leaseUtilizationStrategy = leaseUtilizationStrategy;
+			_jobQueueStrategy = jobQueueStrategy;
+			_noOpStrategy = noOpStrategy;
+			_agentCollection = agentCollection;
+			_poolCollection = poolCollection;
+			_fleetManager = fleetManager;
+			_dogStatsd = dogStatsd;
+			_clock = clock;
+			_ticker = clock.AddSharedTicker<AutoscaleServiceV2>(TimeSpan.FromMinutes(5.0), TickLeaderAsync, logger);
+			_settings = settings.Value;
+			_logger = logger;
+			_defaultScaleOutCooldown = TimeSpan.FromSeconds(settings.Value.AgentPoolScaleOutCooldownSeconds);
+			_defaultScaleInCooldown = TimeSpan.FromSeconds(settings.Value.AgentPoolScaleInCooldownSeconds);
+		}
+
+		/// <inheritdoc/>
+		public Task StartAsync(CancellationToken cancellationToken) => _ticker.StartAsync();
+
+		/// <inheritdoc/>
+		public Task StopAsync(CancellationToken cancellationToken) => _ticker.StopAsync();
+
+		/// <inheritdoc/>
+		public void Dispose() => _ticker.Dispose();
+
+		internal async ValueTask TickLeaderAsync(CancellationToken stoppingToken)
+		{
+			_logger.LogInformation("Autoscaling pools...");
+			Stopwatch stopwatch = Stopwatch.StartNew();
+			using IScope _ = GlobalTracer.Instance.BuildSpan("AutoscaleService").StartActive();
+			
+			// Group pools by strategy type to ensure they can be called once (more optimal)
+			Dictionary<PoolSizeStrategy, List<PoolSizeData>> poolSizeDataByStrategy = await GetPoolSizeDataByStrategyType();
+			foreach ((PoolSizeStrategy strategyType, List<PoolSizeData> currentData) in poolSizeDataByStrategy)
+			{
+				List<PoolSizeData> newData = await GetPoolSizeStrategy(strategyType, _settings.DefaultAgentPoolSizeStrategy).CalcDesiredPoolSizesAsync(currentData);
+				await ResizePools(newData);
+			}
+			
+			stopwatch.Stop();
+			_logger.LogInformation("Autoscaling pools took {ElapsedTime} ms", stopwatch.ElapsedMilliseconds);
+		}
+
+		internal async Task<Dictionary<PoolSizeStrategy, List<PoolSizeData>>> GetPoolSizeDataByStrategyType()
+		{
+			List<IAgent> agents = await _agentCollection.FindAsync(status: AgentStatus.Ok);
+			List<IAgent> GetAgentsInPool(PoolId poolId) => agents.FindAll(a => a.GetPools().Any(p => p == poolId));
+			List<IPool> pools = await _poolCollection.GetAsync();
+
+			Dictionary<PoolSizeStrategy, List<PoolSizeData>> result = new();
+			foreach (PoolSizeStrategy strategyType in Enum.GetValues<PoolSizeStrategy>())
+			{
+				result[strategyType] = pools
+					.Where(x => (x.SizeStrategy == PoolSizeStrategy.Default ? _settings.DefaultAgentPoolSizeStrategy : x.SizeStrategy) == strategyType)
+					.Select(x => new PoolSizeData(x, GetAgentsInPool(x.Id), null))
+					.ToList();
+			}
+
+			return result;
+		}
+
+		internal async Task ResizePools(List<PoolSizeData> poolSizeDatas)
+		{
+			foreach (PoolSizeData poolSizeData in poolSizeDatas.OrderByDescending(x => x.Agents.Count))
+			{
+				IPool pool = poolSizeData.Pool;
+
+				if (!pool.EnableAutoscaling || poolSizeData.DesiredAgentCount == null)
+				{
+					continue;
+				}
+
+				int currentAgentCount = poolSizeData.Agents.Count;
+				int desiredAgentCount = poolSizeData.DesiredAgentCount.Value;
+				int deltaAgentCount = desiredAgentCount - currentAgentCount;
+
+				_logger.LogInformation("{PoolName,-48} Current={Current,4} Target={Target,4} Delta={Delta,4} Status={Status}", pool.Name, currentAgentCount, desiredAgentCount, deltaAgentCount, poolSizeData.StatusMessage);
+				
+				try
+				{
+					using IScope scope = GlobalTracer.Instance.BuildSpan("ScalingPool").StartActive();
+					scope.Span.SetTag("poolName", pool.Name);
+					scope.Span.SetTag("current", currentAgentCount);
+					scope.Span.SetTag("desired", desiredAgentCount);
+					scope.Span.SetTag("delta", deltaAgentCount);
+
+					if (deltaAgentCount > 0)
+					{
+						TimeSpan scaleOutCooldown = pool.ScaleOutCooldown ?? _defaultScaleOutCooldown;
+						bool isCoolingDown = pool.LastScaleUpTime != null && pool.LastScaleUpTime + scaleOutCooldown > _clock.UtcNow;
+						scope.Span.SetTag("isCoolingDown", isCoolingDown);
+						if (!isCoolingDown)
+						{
+							await _fleetManager.ExpandPoolAsync(pool, poolSizeData.Agents, deltaAgentCount);
+							await _poolCollection.TryUpdateAsync(pool, lastScaleUpTime: _clock.UtcNow);
+						}
+						else
+						{
+							TimeSpan? cooldownTimeLeft = pool.LastScaleUpTime + _defaultScaleOutCooldown - _clock.UtcNow;
+							_logger.LogDebug("Cannot scale out {PoolName}, it's cooling down for another {TimeLeft} secs", pool.Name, cooldownTimeLeft?.TotalSeconds);
+						}
+					}
+
+					if (deltaAgentCount < 0)
+					{
+						TimeSpan scaleInCooldown = pool.ScaleInCooldown ?? _defaultScaleInCooldown;
+						bool isCoolingDown = pool.LastScaleDownTime != null && pool.LastScaleDownTime + scaleInCooldown > _clock.UtcNow;
+						scope.Span.SetTag("isCoolingDown", isCoolingDown);
+						if (!isCoolingDown)
+						{
+							await _fleetManager.ShrinkPoolAsync(pool, poolSizeData.Agents, -deltaAgentCount);
+							await _poolCollection.TryUpdateAsync(pool, lastScaleDownTime: _clock.UtcNow);
+						}
+						else
+						{
+							TimeSpan? cooldownTimeLeft = pool.LastScaleDownTime + _defaultScaleInCooldown - _clock.UtcNow;
+							_logger.LogDebug("Cannot scale in {PoolName}, it's cooling down for another {TimeLeft} secs", pool.Name, cooldownTimeLeft?.TotalSeconds);
+						}
+					}
+				}
+				catch (Exception ex)
+				{
+					_logger.LogError(ex, "Failed to scale {PoolName}:\n{Exception}", pool.Name, ex);
+					continue;
+				}
+
+				_dogStatsd.Gauge("agentpools.autoscale.target", desiredAgentCount, tags: new []{"pool:" + pool.Name});
+				_dogStatsd.Gauge("agentpools.autoscale.current", currentAgentCount, tags: new []{"pool:" + pool.Name});
+			}
+		}
+
+		/// <summary>
+		/// Backdoor for tests to override strategies with test doubles
+		/// These cannot supplied in the constructor since the DI requires concrete implementations.
+		/// </summary>
+		/// <param name="leaseUtilizationStrategy"></param>
+		/// <param name="jobQueueStrategy"></param>
+		/// <param name="noOpStrategy"></param>
+		internal void OverridePoolSizeStrategiesDuringTesting(IPoolSizeStrategy leaseUtilizationStrategy, IPoolSizeStrategy jobQueueStrategy, IPoolSizeStrategy noOpStrategy)
+		{
+			_leaseUtilizationStrategy = leaseUtilizationStrategy;
+			_jobQueueStrategy = jobQueueStrategy;
+			_noOpStrategy = noOpStrategy;
+		}
+		
+		private IPoolSizeStrategy GetPoolSizeStrategy(PoolSizeStrategy type, PoolSizeStrategy defaultStrategy)
+		{
+
+			if (type == PoolSizeStrategy.Default)
+			{
+				type = defaultStrategy;
+			}
+
+			return type switch
+			{
+				PoolSizeStrategy.LeaseUtilization => _leaseUtilizationStrategy,
+				PoolSizeStrategy.JobQueue => _jobQueueStrategy,
+				PoolSizeStrategy.NoOp => _noOpStrategy,
+				_ => throw new Exception($"Unknown strategy: {type}")
+			};
+		}
+	}
+}
diff --git a/Engine/Source/Programs/Horde/Horde.Build/Agents/Fleet/IPoolSizeStrategy.cs b/Engine/Source/Programs/Horde/Horde.Build/Agents/Fleet/IPoolSizeStrategy.cs
new file mode 100644
--- /dev/null
+++ b/Engine/Source/Programs/Horde/Horde.Build/Agents/Fleet/IPoolSizeStrategy.cs
@@ -0,0 +1,128 @@
+// Copyright Epic Games, Inc. All Rights Reserved.
+
+using System.Collections.Generic;
+using System.Linq;
+using System.Threading.Tasks;
+using Horde.Build.Models;
+
+namespace Horde.Build.Fleet.Autoscale
+{
+	/// <summary>
+	/// Available pool sizing strategies
+	/// </summary>
+	public enum PoolSizeStrategy
+	{
+		/// <summary>
+		/// Use the default strategy configured in server settings
+		/// </summary>
+		Default,
+
+		/// <summary>
+		/// Strategy based on lease utilization
+		/// <see cref="Horde.Build.Fleet.Autoscale.LeaseUtilizationStrategy"/>
+		/// </summary>
+		LeaseUtilization,
+		
+		/// <summary>
+		/// Strategy based on size of job build queue
+		/// <see cref="Horde.Build.Fleet.Autoscale.JobQueueStrategy"/> 
+		/// </summary>
+		JobQueue,
+		
+		/// <summary>
+		/// No-op strategy used as fallback/default behavior
+		/// <see cref="Horde.Build.Fleet.Autoscale.NoOpPoolSizeStrategy"/> 
+		/// </summary>
+		NoOp
+	}
+
+	/// <summary>
+	/// Class for specifying and grouping data together required for calculating pool size
+	/// </summary>
+	public class PoolSizeData
+	{
+		/// <summary>
+		/// Pool being resized
+		/// </summary>
+		public IPool Pool { get; }
+		
+		/// <summary>
+		/// All agents currently associated with the pool
+		/// </summary>
+		public List<IAgent> Agents { get; }
+
+		/// <summary>
+		/// The desired agent count (calculated and updated once the strategy has been run, null otherwise)
+		/// </summary>
+		public int? DesiredAgentCount  { get; }
+		
+		/// <summary>
+		/// Human-readable text describing the status of the pool (data the sizing is based on etc)
+		/// </summary>
+		public string StatusMessage { get; }
+
+		/// <summary>
+		/// Constructor
+		/// </summary>
+		/// <param name="pool"></param>
+		/// <param name="agents"></param>
+		/// <param name="desiredAgentCount"></param>
+		/// <param name="statusMessage"></param>
+		public PoolSizeData(IPool pool, List<IAgent> agents, int? desiredAgentCount, string statusMessage = "N/A")
+		{
+			Pool = pool;
+			Agents = agents;
+			DesiredAgentCount = desiredAgentCount;
+			StatusMessage = statusMessage;
+		}
+
+		/// <summary>
+		/// Copy the object, inheriting any unspecified values from current instance
+		/// Needed because the class is immutable 
+		/// </summary>
+		/// <param name="pool"></param>
+		/// <param name="agents"></param>
+		/// <param name="desiredAgentCount"></param>
+		/// <param name="statusMessage"></param>
+		/// <returns>A new copy</returns>
+		public PoolSizeData Copy(IPool? pool = null, List<IAgent>? agents = null, int? desiredAgentCount = null, string? statusMessage = null)
+		{
+			return new PoolSizeData(pool ?? Pool, agents ?? Agents, desiredAgentCount ?? DesiredAgentCount, statusMessage ?? StatusMessage);
+		}
+	}
+	
+	/// <summary>
+	/// Interface for different agent pool sizing strategies
+	/// </summary>
+	public interface IPoolSizeStrategy
+	{
+		/// <summary>
+		/// Calculate the adequate number of agents to be online for given pools
+		/// </summary>
+		/// <param name="pools">Pools including attached agents</param>
+		/// <returns></returns>
+		Task<List<PoolSizeData>> CalcDesiredPoolSizesAsync(List<PoolSizeData> pools);
+		
+		/// <summary>
+		/// Name of the strategy
+		/// </summary>
+		string Name { get; }
+	}
+	
+	/// <summary>
+	/// No-operation strategy that won't resize pools, just return the existing count.
+	/// Used to ensure there's always a strategy available for dependency injection, even if it does nothing.
+	/// </summary>
+	public class NoOpPoolSizeStrategy : IPoolSizeStrategy
+	{
+		/// <inheritdoc/>
+		public Task<List<PoolSizeData>> CalcDesiredPoolSizesAsync(List<PoolSizeData> pools)
+		{
+			List<PoolSizeData> result = pools.Select(x => new PoolSizeData(x.Pool, x.Agents, x.Agents.Count, "(no-op)")).ToList();
+			return Task.FromResult(result);
+		}
+
+		/// <inheritdoc/>
+		public string Name { get; } = "NoOp";
+	}
+}
diff --git a/Engine/Source/Programs/Horde/Horde.Build/Agents/Fleet/JobQueueStrategy.cs b/Engine/Source/Programs/Horde/Horde.Build/Agents/Fleet/JobQueueStrategy.cs
new file mode 100644
--- /dev/null
+++ b/Engine/Source/Programs/Horde/Horde.Build/Agents/Fleet/JobQueueStrategy.cs
@@ -0,0 +1,183 @@
+// Copyright Epic Games, Inc. All Rights Reserved.
+
+using System;
+using System.Collections.Generic;
+using System.Linq;
+using System.Threading.Tasks;
+using Horde.Build.Api;
+using Horde.Build.Collections;
+using Horde.Build.Models;
+using Horde.Build.Services;
+using Horde.Build.Utilities;
+using HordeCommon;
+
+namespace Horde.Build.Fleet.Autoscale
+{
+	using PoolId = StringId<IPool>;
+	using StreamId = StringId<IStream>;
+
+	/// <summary>
+	/// Job queue sizing settings for a pool
+	/// </summary>
+	public class JobQueueSettings
+	{
+		/// <summary>
+		/// Factor translating queue size to additional agents to grow the pool with
+		/// The result is always rounded up to nearest integer. 
+		/// Example: if there are 20 jobs in queue, a factor 0.25 will result in 5 new agents being added (20 * 0.25)
+		/// </summary>
+		public double ScaleOutFactor { get; set;  } = 0.25;
+		
+		/// <summary>
+		/// Factor by which to shrink the pool size with when queue is empty
+		/// The result is always rounded up to nearest integer.
+		/// Example: when the queue size is zero, a default value of 0.9 will shrink the pool by 10% (current agent count * 0.9)
+		/// </summary>
+		public double ScaleInFactor { get; set; } = 0.9;
+
+		/// <summary>
+		/// Constructor
+		/// </summary>
+		/// <param name="scaleOutFactor"></param>
+		/// <param name="scaleInFactor"></param>
+		public JobQueueSettings(double? scaleOutFactor = null, double? scaleInFactor = null)
+		{
+			ScaleOutFactor = scaleOutFactor.GetValueOrDefault(ScaleOutFactor);
+			ScaleInFactor = scaleInFactor.GetValueOrDefault(ScaleInFactor);
+		}
+	}
+	
+	/// <summary>
+	/// Calculate pool size by observing the number of jobs in waiting state
+	///
+	/// Allows for more proactive scaling compared to LeaseUtilizationStrategy.
+	/// <see cref="Horde.Build.Fleet.Autoscale.LeaseUtilizationStrategy"/> 
+	/// </summary>
+	public class JobQueueStrategy : IPoolSizeStrategy
+	{
+		private readonly IJobCollection _jobs;
+		private readonly IGraphCollection _graphs;
+		private readonly StreamService _streamService;
+		private readonly IClock _clock;
+		
+		/// <summary>
+		/// How far back in time to look for job batches (that potentially are in the queue)
+		/// </summary>
+		private readonly TimeSpan _samplePeriod = TimeSpan.FromDays(5);
+		
+		/// <summary>
+		/// Time spent in ready state before considered truly waiting for an agent
+		///
+		/// A job batch can be in ready state before getting picked up and executed.
+		/// This threshold will help ensure only batches that have been waiting longer than this value will be considered.
+		/// </summary>
+		internal readonly TimeSpan ReadyTimeThreshold = TimeSpan.FromSeconds(45.0);
+
+		private readonly JobQueueSettings _defaultPoolSettings = new();
+
+		/// <summary>
+		/// Constructor
+		/// </summary>
+		/// <param name="jobs"></param>
+		/// <param name="graphs"></param>
+		/// <param name="streamService"></param>
+		/// <param name="clock"></param>
+		/// <param name="samplePeriod">Time period for each sample</param>
+		public JobQueueStrategy(IJobCollection jobs, IGraphCollection graphs, StreamService streamService, IClock clock, TimeSpan? samplePeriod = null)
+		{
+			_jobs = jobs;
+			_graphs = graphs;
+			_streamService = streamService;
+			_clock = clock;
+			_samplePeriod = samplePeriod ?? _samplePeriod;
+		}
+
+		/// <inheritdoc/>
+		public string Name { get; } = "JobQueue";
+
+		/// <summary>
+		/// Extract all job step batches from a job, with their associated pool 
+		/// </summary>
+		/// <param name="job">Job to extract from</param>
+		/// <param name="streams">Cached lookup table of streams</param>
+		/// <returns></returns>
+		private async Task<List<(IJob Job, IJobStepBatch Batch, PoolId PoolId)>> GetJobBatchesWithPools(IJob job, Dictionary<StreamId, IStream> streams)
+		{
+			IGraph graph = await _graphs.GetAsync(job.GraphHash);
+
+			List<(IJob Job, IJobStepBatch Batch, PoolId PoolId)> jobBatches = new();
+			foreach (IJobStepBatch batch in job.Batches)
+			{
+				if (batch.State != JobStepBatchState.Ready)
+				{
+					continue;
+				}
+				
+				TimeSpan? waitTime = _clock.UtcNow - batch.ReadyTimeUtc;
+				if (waitTime == null)
+				{
+					continue;
+				}
+				if (waitTime.Value < ReadyTimeThreshold)
+				{
+					continue;
+				}
+
+				if (!streams.TryGetValue(job.StreamId, out IStream? stream))
+				{
+					continue;
+				}
+
+				string batchAgentType = graph.Groups[batch.GroupIdx].AgentType;
+				if (!stream.AgentTypes.TryGetValue(batchAgentType, out AgentType? agentType))
+				{
+					continue;
+				}
+				
+				jobBatches.Add((job, batch, agentType.Pool));
+			}
+
+			return jobBatches;
+		}
+
+		internal async Task<Dictionary<PoolId, int>> GetPoolQueueSizesAsync(DateTimeOffset jobsCreatedAfter)
+		{
+			List<IStream> streamsList = await _streamService.GetStreamsAsync();
+			Dictionary<StreamId, IStream> streams = streamsList.ToDictionary(x => x.Id, x => x);
+			List<IJob> recentJobs = await _jobs.FindAsync(minCreateTime: jobsCreatedAfter);
+
+			List<(IJob Job, IJobStepBatch Batch, PoolId PoolId)> jobBatches = new();
+			foreach (IJob job in recentJobs)
+			{
+				jobBatches.AddRange(await GetJobBatchesWithPools(job, streams));
+			}
+
+			List<(PoolId PoolId, int QueueSize)> poolsWithQueueSize = jobBatches.GroupBy(t => t.PoolId).Select(t => (t.Key, t.Count())).ToList();
+			return poolsWithQueueSize.ToDictionary(x => x.PoolId, x => x.QueueSize);
+		}
+
+		/// <inheritdoc/>
+		public async Task<List<PoolSizeData>> CalcDesiredPoolSizesAsync(List<PoolSizeData> pools)
+		{
+			DateTimeOffset minCreateTime = _clock.UtcNow - _samplePeriod;
+			Dictionary<PoolId, int> poolQueueSizes = await GetPoolQueueSizesAsync(minCreateTime);
+
+			return pools.Select(current =>
+			{
+				JobQueueSettings settings = current.Pool.JobQueueSettings ?? _defaultPoolSettings;
+				poolQueueSizes.TryGetValue(current.Pool.Id, out int queueSize);
+				if (queueSize > 0)
+				{
+					int additionalAgentCount = (int)Math.Ceiling(queueSize * settings.ScaleOutFactor);
+					int desiredAgentCount = current.Agents.Count + additionalAgentCount;
+					return new PoolSizeData(current.Pool, current.Agents, desiredAgentCount, $"QueueSize={queueSize}");
+				}
+				else
+				{
+					int desiredAgentCount = (int)(current.Agents.Count * settings.ScaleInFactor);
+					return new PoolSizeData(current.Pool, current.Agents, desiredAgentCount, "Empty job queue");
+				}
+			}).ToList();
+		}
+	}
+}
diff --git a/Engine/Source/Programs/Horde/Horde.Build/Agents/Fleet/LeaseUtilizationStrategy.cs b/Engine/Source/Programs/Horde/Horde.Build/Agents/Fleet/LeaseUtilizationStrategy.cs
new file mode 100644
--- /dev/null
+++ b/Engine/Source/Programs/Horde/Horde.Build/Agents/Fleet/LeaseUtilizationStrategy.cs
@@ -0,0 +1,244 @@
+// Copyright Epic Games, Inc. All Rights Reserved.
+
+using System;
+using System.Collections.Generic;
+using System.Linq;
+using System.Text;
+using System.Threading.Tasks;
+using Google.Protobuf.WellKnownTypes;
+using Horde.Build.Collections;
+using Horde.Build.Models;
+using Horde.Build.Utilities;
+using HordeCommon;
+using HordeCommon.Rpc.Tasks;
+using OpenTracing;
+using OpenTracing.Util;
+
+namespace Horde.Build.Fleet.Autoscale
+{
+	using PoolId = StringId<IPool>;
+
+	/// <summary>
+	/// Lease utilization sizing settings for a pool
+	/// </summary>
+	public class LeaseUtilizationSettings
+	{
+	}
+	
+	/// <summary>
+	/// Calculate pool size by looking at previously finished leases
+	/// </summary>
+	public class LeaseUtilizationStrategy : IPoolSizeStrategy
+	{
+		struct UtilizationSample
+		{
+			public double _jobWork;
+			public double _otherWork;
+		}
+
+		class AgentData
+		{
+			public IAgent Agent { get; }
+			public UtilizationSample[] Samples { get; }
+			public List<ILease> Leases { get; } = new ();
+			private readonly int _numSamples;
+
+			public AgentData(IAgent agent, int numSamples)
+			{
+				Agent = agent;
+				_numSamples = numSamples;
+				Samples = new UtilizationSample[numSamples];
+			}
+
+			public void Add(double minT, double maxT, double jobWork, double otherWork)
+			{
+				int minIdx = Math.Clamp((int)minT, 0, _numSamples - 1);
+				int maxIdx = Math.Clamp((int)maxT, minIdx, _numSamples - 1);
+				for (int idx = minIdx; idx <= maxIdx; idx++)
+				{
+					double fraction = Math.Clamp(maxT - idx, 0.0, 1.0) - Math.Clamp(minT - idx, 0.0, 1.0);
+					Samples[idx]._jobWork += jobWork * fraction;
+					Samples[idx]._otherWork += otherWork * fraction;
+				}
+			}
+		}
+
+		class PoolData
+		{
+			public IPool Pool { get; }
+			public List<IAgent> Agents { get; } = new ();
+			public UtilizationSample[] Samples { get; }
+
+			public PoolData(IPool pool, int numSamples)
+			{
+				Pool = pool;
+				Samples = new UtilizationSample[numSamples];
+			}
+
+			public void Add(AgentData agentData)
+			{
+				Agents.Add(agentData.Agent);
+				for (int idx = 0; idx < Samples.Length; idx++)
+				{
+					Samples[idx]._jobWork += agentData.Samples[idx]._jobWork;
+					Samples[idx]._otherWork += agentData.Samples[idx]._otherWork;
+				}
+			}
+		}
+		
+		private readonly IAgentCollection _agentCollection;
+		private readonly IPoolCollection _poolCollection;
+		private readonly ILeaseCollection _leaseCollection;
+		private readonly IClock _clock;
+		private readonly TimeSpan _sampleTime = TimeSpan.FromMinutes(6.0);
+		private readonly int _numSamples = 10;
+		private readonly int _numSamplesForResult = 9;
+
+		/// <summary>
+		/// Constructor
+		/// </summary>
+		/// <param name="agentCollection"></param>
+		/// <param name="poolCollection"></param>
+		/// <param name="leaseCollection"></param>
+		/// <param name="clock"></param>
+		/// <param name="numSamples">Number of samples to collect for calculating lease utilization</param>
+		/// <param name="numSamplesForResult">Min number of samples for a valid result</param>
+		/// <param name="sampleTime">Time period for each sample</param>
+		public LeaseUtilizationStrategy(IAgentCollection agentCollection, IPoolCollection poolCollection, ILeaseCollection leaseCollection, IClock clock, int numSamples = 10, int numSamplesForResult = 9, TimeSpan? sampleTime = null)
+		{
+			_agentCollection = agentCollection;
+			_poolCollection = poolCollection;
+			_leaseCollection = leaseCollection;
+			_clock = clock;
+			_numSamples = numSamples;
+			_numSamplesForResult = numSamplesForResult;
+			_sampleTime = sampleTime ?? _sampleTime;
+		}
+
+		private async Task<Dictionary<AgentId, AgentData>> GetAgentDataAsync()
+		{
+			using IScope _ = GlobalTracer.Instance.BuildSpan("GetAgentDataAsync").StartActive();
+			
+			// Find all the current agents
+			List<IAgent> agents = await _agentCollection.FindAsync(status: AgentStatus.Ok);
+
+			// Query leases in last interval
+			DateTime maxTime = _clock.UtcNow;
+			DateTime minTime = maxTime - (_sampleTime * _numSamples);
+			List<ILease> leases = await _leaseCollection.FindLeasesAsync(minTime, maxTime);
+
+			// Add all the leases to a data object for each agent
+			Dictionary<AgentId, AgentData> agentIdToData = agents.ToDictionary(x => x.Id, x => new AgentData(x, _numSamples));
+			foreach (ILease lease in leases)
+			{
+				AgentData? agentData;
+				if (agentIdToData.TryGetValue(lease.AgentId, out agentData) && agentData.Agent.SessionId == lease.SessionId)
+				{
+					agentData.Leases.Add(lease);
+				}
+			}
+			
+			// Compute utilization for each agent
+			foreach (AgentData agentData in agentIdToData.Values)
+			{
+				foreach (ILease lease in agentData.Leases.OrderBy(x => x.StartTime))
+				{
+					double minT = (lease.StartTime - minTime).TotalSeconds / _sampleTime.TotalSeconds;
+					double maxT = (lease.FinishTime == null)
+						? _numSamples
+						: ((lease.FinishTime.Value - minTime).TotalSeconds / _sampleTime.TotalSeconds);
+
+					Any payload = Any.Parser.ParseFrom(lease.Payload.ToArray());
+					if (payload.Is(ExecuteJobTask.Descriptor))
+					{
+						agentData.Add(minT, maxT, 1.0, 0.0);
+					}
+					else
+					{
+						agentData.Add(minT, maxT, 0.0, 1.0);
+					}
+				}
+			}
+
+			return agentIdToData;
+		}
+		
+		private async Task<Dictionary<PoolId, PoolData>> GetPoolDataAsync()
+		{
+			using IScope _ = GlobalTracer.Instance.BuildSpan("GetPoolDataAsync").StartActive();
+
+			Dictionary<AgentId, AgentData> agentIdToData = await GetAgentDataAsync();
+			
+			// Get all the pools
+			List<IPool> pools = await _poolCollection.GetAsync();
+			Dictionary<PoolId, PoolData> poolToData = pools.ToDictionary(x => x.Id, x => new PoolData(x, _numSamples));
+
+			// Find pool utilization over the query period
+			foreach (AgentData agentData in agentIdToData.Values)
+			{
+				foreach (PoolId poolId in agentData.Agent.GetPools())
+				{
+					PoolData? poolData;
+					if (poolToData.TryGetValue(poolId, out poolData))
+					{
+						poolData.Add(agentData);
+					}
+				}
+			}
+
+			return poolToData;
+		}
+
+		/// <inheritdoc/>
+		public string Name { get; } = "LeaseUtilization";
+
+		/// <inheritdoc/>
+		public async Task<List<PoolSizeData>> CalcDesiredPoolSizesAsync(List<PoolSizeData> pools)
+		{
+			Dictionary<PoolId, PoolData> poolToData = await GetPoolDataAsync();
+			List<PoolSizeData> result = new();
+
+			foreach (PoolData poolData in poolToData.Values.OrderByDescending(x => x.Agents.Count))
+			{
+				IPool pool = poolData.Pool;
+
+				PoolSizeData? poolSize = pools.Find(x => x.Pool.Id == pool.Id);
+				if (poolSize != null)
+				{
+					int minAgents = pool.MinAgents ?? 1;
+					int numReserveAgents = pool.NumReserveAgents ?? 5;
+					double utilization = poolData.Samples.Select(x => x._jobWork).OrderByDescending(x => x).Skip(_numSamples - _numSamplesForResult).First();
+				
+					// Number of agents in use over the sampling period. Can never be greater than number of agents available in pool.
+					int numAgentsUtilized = (int)utilization;
+					
+					// Include reserve agent count to ensure pool always can grow
+					int desiredAgentCount = Math.Max(numAgentsUtilized + numReserveAgents, minAgents);
+					
+					StringBuilder sb = new();
+					sb.AppendFormat("Jobs=[{0}] ", GetDensityMap(poolData.Samples.Select(x => x._jobWork / Math.Max(1, poolData.Agents.Count))));
+					sb.AppendFormat("Total=[{0}] ", GetDensityMap(poolData.Samples.Select(x => x._otherWork / Math.Max(1, poolData.Agents.Count))));
+					sb.AppendFormat("Min=[{0,5:0.0}] ", poolData.Samples.Min(x => x._jobWork));
+					sb.AppendFormat("Max=[{0,5:0.0}] ", poolData.Samples.Max(x => x._jobWork));
+					sb.AppendFormat("Avg=[{0,5:0.0}] ", utilization);
+					sb.AppendFormat("Pct=[{0,5:0.0}] ", poolData.Samples.Sum(x => x._jobWork) / _numSamples);
+
+					result.Add(new(pool, poolSize.Agents, desiredAgentCount, sb.ToString()));
+				}
+			}
+
+			return result;
+		}
+		
+		/// <summary>
+		/// Creates a string of characters indicating a sequence of 0-1 values over time
+		/// </summary>
+		/// <param name="values">Sequence of values</param>
+		/// <returns>Density map</returns>
+		private static string GetDensityMap(IEnumerable<double> values)
+		{
+			const string Greyscale = " 123456789";
+			return new string(values.Select(x => Greyscale[Math.Clamp((int)(x * Greyscale.Length), 0, Greyscale.Length - 1)]).ToArray());
+		}
+	}
+}
diff --git a/Engine/Source/Programs/Horde/Horde.Build/Agents/Fleet/Providers/AwsFleetManager.cs b/Engine/Source/Programs/Horde/Horde.Build/Agents/Fleet/Providers/AwsFleetManager.cs
new file mode 100644
--- /dev/null
+++ b/Engine/Source/Programs/Horde/Horde.Build/Agents/Fleet/Providers/AwsFleetManager.cs
@@ -0,0 +1,148 @@
+// Copyright Epic Games, Inc. All Rights Reserved.
+
+using System;
+using System.Collections.Generic;
+using System.Linq;
+using System.Threading.Tasks;
+using Amazon;
+using Amazon.EC2;
+using Amazon.EC2.Model;
+using Horde.Build.Collections;
+using Horde.Build.Models;
+using Horde.Build.Utilities;
+using Microsoft.Extensions.Logging;
+using OpenTracing;
+using OpenTracing.Util;
+
+namespace Horde.Build.Services.Impl
+{
+	/// <summary>
+	/// Fleet manager for handling AWS EC2 instances
+	/// </summary>
+	public sealed class AwsFleetManager : IFleetManager, IDisposable
+	{
+		const string AwsTagPropertyName = "aws-tag";
+		const string PoolTagName = "Horde_Autoscale_Pool";
+		readonly AmazonEC2Client _client;
+		readonly IAgentCollection _agentCollection;
+		readonly ILogger _logger;
+
+		/// <summary>
+		/// Constructor
+		/// </summary>
+		public AwsFleetManager(IAgentCollection agentCollection, ILogger<AwsFleetManager> logger)
+		{
+			_agentCollection = agentCollection;
+			_logger = logger;
+
+			AmazonEC2Config config = new AmazonEC2Config();
+			config.RegionEndpoint = RegionEndpoint.USEast1;
+
+			logger.LogInformation("Initializing AWS fleet manager for region {Region}", config.RegionEndpoint);
+
+			_client = new AmazonEC2Client(config);
+		}
+
+		/// <inheritdoc/>
+		public void Dispose()
+		{
+			_client.Dispose();
+		}
+
+		/// <inheritdoc/>
+		public async Task ExpandPoolAsync(IPool pool, IReadOnlyList<IAgent> agents, int count)
+		{
+			using IScope scope = GlobalTracer.Instance.BuildSpan("ExpandPool").StartActive();
+			scope.Span.SetTag("poolName", pool.Name);
+			scope.Span.SetTag("numAgents", agents.Count);
+			scope.Span.SetTag("count", count);
+
+			DescribeInstancesResponse describeResponse;
+			using (IScope describeScope = GlobalTracer.Instance.BuildSpan("DescribeInstances").StartActive())
+			{
+				// Find stopped instances in the correct pool
+				DescribeInstancesRequest describeRequest = new DescribeInstancesRequest();
+				describeRequest.Filters = new List<Filter>();
+				describeRequest.Filters.Add(new Filter("instance-state-name", new List<string> { InstanceStateName.Stopped.Value }));
+				describeRequest.Filters.Add(new Filter("tag:" + PoolTagName, new List<string> { pool.Name }));
+				describeResponse = await _client.DescribeInstancesAsync(describeRequest);
+				describeScope.Span.SetTag("res.statusCode", (int)describeResponse.HttpStatusCode);
+				describeScope.Span.SetTag("res.numReservations", describeResponse.Reservations.Count);
+			}
+
+			using (IScope startScope = GlobalTracer.Instance.BuildSpan("StartInstances").StartActive())
+			{
+				// Try to start the given instances
+				StartInstancesRequest startRequest = new StartInstancesRequest();
+				startRequest.InstanceIds.AddRange(describeResponse.Reservations.SelectMany(x => x.Instances).Select(x => x.InstanceId).Take(count));
+				
+				startScope.Span.SetTag("req.instanceIds", String.Join(",", startRequest.InstanceIds));
+				if (startRequest.InstanceIds.Count > 0)
+				{
+					StartInstancesResponse startResponse = await _client.StartInstancesAsync(startRequest);
+					startScope.Span.SetTag("res.statusCode", (int)startResponse.HttpStatusCode);
+					startScope.Span.SetTag("res.numInstances", startResponse.StartingInstances.Count);
+					if ((int)startResponse.HttpStatusCode >= 200 && (int)startResponse.HttpStatusCode <= 299)
+					{
+						foreach (InstanceStateChange instanceChange in startResponse.StartingInstances)
+						{
+							_logger.LogInformation("Starting instance {InstanceId} for pool {PoolId} (prev state {PrevState}, current state {CurrentState}", instanceChange.InstanceId, pool.Id, instanceChange.PreviousState, instanceChange.CurrentState);
+						}
+					}
+				}
+
+				if (startRequest.InstanceIds.Count < count)
+				{
+					_logger.LogInformation("Unable to expand pool {PoolName} with the requested number of instances. " +
+					                      "Num requested instances to add {RequestedCount}. Actual instances started {InstancesStarted}", pool.Name, count, startRequest.InstanceIds.Count);
+				}
+			}
+		}
+
+		/// <inheritdoc/>
+		public async Task ShrinkPoolAsync(IPool pool, IReadOnlyList<IAgent> agents, int count)
+		{
+			using IScope scope = GlobalTracer.Instance.BuildSpan("ShrinkPool").StartActive();
+			scope.Span.SetTag("poolName", pool.Name);
+			scope.Span.SetTag("count", count);
+			
+			string awsTagProperty = $"{AwsTagPropertyName}={PoolTagName}:{pool.Name}";
+			
+			// Sort the agents by number of active leases. It's better to shutdown agents currently doing nothing.
+			List<IAgent> filteredAgents = agents.OrderBy(x => x.Leases.Count).ToList();
+			List<IAgent> agentsWithAwsTags = filteredAgents.Where(x => x.HasProperty(awsTagProperty)).ToList(); 
+			List<IAgent> agentsLimitedByCount = agentsWithAwsTags.Take(count).ToList();
+			
+			scope.Span.SetTag("agents.num", agents.Count);
+			scope.Span.SetTag("agents.filtered.num", filteredAgents.Count);
+			scope.Span.SetTag("agents.withAwsTags.num", agentsWithAwsTags.Count);
+			scope.Span.SetTag("agents.limitedByCount.num", agentsLimitedByCount.Count);
+
+			foreach (IAgent agent in agentsLimitedByCount)
+			{
+				IAuditLogChannel<AgentId> agentLogger = _agentCollection.GetLogger(agent.Id);
+				if (await _agentCollection.TryUpdateSettingsAsync(agent, bRequestShutdown: true, shutdownReason: "Autoscaler") != null)
+				{
+					agentLogger.LogInformation("Marked {AgentId} in pool {PoolName} for shutdown due to autoscaling (currently {NumLeases} leases outstanding)", agent.Id, pool.Name, agent.Leases.Count);
+				}
+				else
+				{
+					agentLogger.LogError("Unable to mark agent {AgentId} in pool {PoolName} for shutdown due to autoscaling", agent.Id, pool.Name);
+				}
+			}
+		}
+
+		/// <inheritdoc/>
+		public async Task<int> GetNumStoppedInstancesAsync(IPool pool)
+		{
+			// Find all instances in the pool
+			DescribeInstancesRequest describeRequest = new DescribeInstancesRequest();
+			describeRequest.Filters = new List<Filter>();
+			describeRequest.Filters.Add(new Filter("instance-state-name", new List<string> { InstanceStateName.Stopped.Value }));
+			describeRequest.Filters.Add(new Filter("tag:" + PoolTagName, new List<string> { pool.Name }));
+
+			DescribeInstancesResponse describeResponse = await _client.DescribeInstancesAsync(describeRequest);
+			return describeResponse.Reservations.SelectMany(x => x.Instances).Select(x => x.InstanceId).Distinct().Count();
+		}
+	}
+}
diff --git a/Engine/Source/Programs/Horde/Horde.Build/Agents/Fleet/Providers/DefaultFleetManager.cs b/Engine/Source/Programs/Horde/Horde.Build/Agents/Fleet/Providers/DefaultFleetManager.cs
new file mode 100644
--- /dev/null
+++ b/Engine/Source/Programs/Horde/Horde.Build/Agents/Fleet/Providers/DefaultFleetManager.cs
@@ -0,0 +1,43 @@
+// Copyright Epic Games, Inc. All Rights Reserved.
+
+using System.Collections.Generic;
+using System.Threading.Tasks;
+using Horde.Build.Models;
+using Microsoft.Extensions.Logging;
+
+namespace Horde.Build.Services.Impl
+{
+	/// <summary>
+	/// Default implementation of <see cref="IFleetManager"/>
+	/// </summary>
+	public class DefaultFleetManager : IFleetManager
+	{
+		readonly ILogger _logger;
+
+		/// <summary>
+		/// Constructor
+		/// </summary>
+		/// <param name="logger">Logging device</param>
+		public DefaultFleetManager(ILogger<DefaultFleetManager> logger)
+		{
+			_logger = logger;
+		}
+
+		/// <inheritdoc/>
+		public Task ExpandPoolAsync(IPool pool, IReadOnlyList<IAgent> agents, int count)
+		{
+			_logger.LogInformation("Expand pool {PoolId} by {Count} agents", pool.Id, count);
+			return Task.CompletedTask;
+		}
+
+		/// <inheritdoc/>
+		public Task ShrinkPoolAsync(IPool pool, IReadOnlyList<IAgent> agents, int count)
+		{
+			_logger.LogInformation("Shrink pool {PoolId} by {Count} agents", pool.Id, count);
+			return Task.CompletedTask;
+		}
+
+		/// <inheritdoc/>
+		public Task<int> GetNumStoppedInstancesAsync(IPool pool) => Task.FromResult(0);
+	}
+}
diff --git a/Engine/Source/Programs/Horde/Horde.Build/Agents/Leases/LeaseCollection.cs b/Engine/Source/Programs/Horde/Horde.Build/Agents/Leases/LeaseCollection.cs
new file mode 100644
--- /dev/null
+++ b/Engine/Source/Programs/Horde/Horde.Build/Agents/Leases/LeaseCollection.cs
@@ -0,0 +1,195 @@
+// Copyright Epic Games, Inc. All Rights Reserved.
+
+using System;
+using System.Collections.Generic;
+using System.Linq;
+using System.Threading.Tasks;
+using Horde.Build.Models;
+using Horde.Build.Server;
+using Horde.Build.Utilities;
+using HordeCommon;
+using MongoDB.Bson;
+using MongoDB.Bson.Serialization.Attributes;
+using MongoDB.Driver;
+
+namespace Horde.Build.Collections.Impl
+{
+	using LeaseId = ObjectId<ILease>;
+	using LogId = ObjectId<ILogFile>;
+	using PoolId = StringId<IPool>;
+	using SessionId = ObjectId<ISession>;
+	using StreamId = StringId<IStream>;
+
+	/// <summary>
+	/// Collection of lease documents
+	/// </summary>
+	public class LeaseCollection : ILeaseCollection
+	{
+		/// <summary>
+		/// Concrete implementation of a lease document
+		/// </summary>
+		class LeaseDocument : ILease
+		{
+			public LeaseId Id { get; set; }
+			public string Name { get; set; }
+			public AgentId AgentId { get; set; }
+			public SessionId SessionId { get; set; }
+			public StreamId? StreamId { get; set; }
+			public PoolId? PoolId { get; set; }
+			public LogId? LogId { get; set; }
+			public DateTime StartTime { get; set; }
+			public DateTime? FinishTime { get; set; }
+			public byte[] Payload { get; set; }
+			public LeaseOutcome Outcome { get; set; }
+
+			[BsonIgnoreIfNull]
+			public byte[]? Output { get; set; }
+
+			ReadOnlyMemory<byte> ILease.Payload => Payload;
+			ReadOnlyMemory<byte> ILease.Output => Output ?? Array.Empty<byte>();
+
+			[BsonConstructor]
+			private LeaseDocument()
+			{
+				Name = String.Empty;
+				Payload = null!;
+			}
+
+			public LeaseDocument(LeaseId id, string name, AgentId agentId, SessionId sessionId, StreamId? streamId, PoolId? poolId, LogId? logId, DateTime startTime, byte[] payload)
+			{
+				Id = id;
+				Name = name;
+				AgentId = agentId;
+				SessionId = sessionId;
+				StreamId = streamId;
+				PoolId = poolId;
+				LogId = logId;
+				StartTime = startTime;
+				Payload = payload;
+			}
+		}
+		
+		readonly IMongoCollection<LeaseDocument> _leases;
+		readonly MongoIndex<LeaseDocument> _finishTimeStartTimeCompoundIndex;
+
+		/// <summary>
+		/// Constructor
+		/// </summary>
+		/// <param name="mongoService">The database service instance</param>
+		public LeaseCollection(MongoService mongoService)
+		{
+			List<MongoIndex<LeaseDocument>> indexes = new List<MongoIndex<LeaseDocument>>();
+			indexes.Add(keys => keys.Ascending(x => x.AgentId));
+			indexes.Add(keys => keys.Ascending(x => x.SessionId));
+			indexes.Add(keys => keys.Ascending(x => x.StartTime));
+			indexes.Add(keys => keys.Ascending(x => x.FinishTime));
+			indexes.Add(_finishTimeStartTimeCompoundIndex = MongoIndex.Create<LeaseDocument>(keys => keys.Ascending(x => x.FinishTime).Descending(x => x.StartTime)));
+
+			_leases = mongoService.GetCollection<LeaseDocument>("Leases", indexes);
+		}
+
+		/// <inheritdoc/>
+		public async Task<ILease> AddAsync(LeaseId id, string name, AgentId agentId, SessionId sessionId, StreamId? streamId, PoolId? poolId, LogId? logId, DateTime startTime, byte[] payload)
+		{
+			LeaseDocument lease = new LeaseDocument(id, name, agentId, sessionId, streamId, poolId, logId, startTime, payload);
+			await _leases.ReplaceOneAsync(x => x.Id == id, lease, new ReplaceOptions { IsUpsert = true });
+			return lease;
+		}
+
+		/// <inheritdoc/>
+		public async Task DeleteAsync(LeaseId leaseId)
+		{
+			await _leases.DeleteOneAsync(x => x.Id == leaseId);
+		}
+
+		/// <inheritdoc/>
+		public async Task<ILease?> GetAsync(LeaseId leaseId)
+		{
+			return await _leases.Find(x => x.Id == leaseId).FirstOrDefaultAsync();
+		}
+
+		/// <inheritdoc/>
+		public async Task<List<ILease>> FindLeasesAsync(AgentId? agentId, SessionId? sessionId, DateTime? minTime, DateTime? maxTime, int? index, int? count, string? indexHint = null, bool consistentRead = true)
+		{
+			IMongoCollection<LeaseDocument> collection = consistentRead ? _leases : _leases.WithReadPreference(ReadPreference.SecondaryPreferred);
+			
+			FilterDefinitionBuilder<LeaseDocument> filterBuilder = Builders<LeaseDocument>.Filter;
+			FilterDefinition<LeaseDocument> filter = FilterDefinition<LeaseDocument>.Empty;
+			if (agentId != null)
+			{
+				filter &= filterBuilder.Eq(x => x.AgentId, agentId.Value);
+			}
+			if (sessionId != null)
+			{
+				filter &= filterBuilder.Eq(x => x.SessionId, sessionId.Value);
+			}
+			if (minTime != null)
+			{
+				filter &= filterBuilder.Or(filterBuilder.Eq(x => x.FinishTime!, null), filterBuilder.Gt(x => x.FinishTime!, minTime.Value));
+			}
+			if (maxTime != null)
+			{
+				filter &= filterBuilder.Lt(x => x.StartTime, maxTime.Value);
+			}
+
+			List<LeaseDocument> results = await collection.FindWithHint(filter, indexHint, x => x.SortByDescending(x => x.StartTime).Range(index, count).ToListAsync());
+			return results.ConvertAll<ILease>(x => x);
+		}
+		
+		/// <inheritdoc/>
+		public async Task<List<ILease>> FindLeasesByFinishTimeAsync(DateTime? minFinishTime, DateTime? maxFinishTime, int? index, int? count, string? indexHint, bool consistentRead)
+		{
+			IMongoCollection<LeaseDocument> collection = consistentRead ? _leases : _leases.WithReadPreference(ReadPreference.SecondaryPreferred);
+			FilterDefinitionBuilder<LeaseDocument> filterBuilder = Builders<LeaseDocument>.Filter;
+			FilterDefinition<LeaseDocument> filter = FilterDefinition<LeaseDocument>.Empty;
+
+			if (minFinishTime == null && maxFinishTime == null)
+			{
+				throw new ArgumentException($"Both {nameof(minFinishTime)} and {nameof(maxFinishTime)} cannot be null");
+			}
+
+			if (minFinishTime != null)
+			{
+				filter &= filterBuilder.Gt(x => x.FinishTime, minFinishTime.Value);
+			}
+			if (maxFinishTime != null)
+			{
+				filter &= filterBuilder.Lt(x => x.FinishTime, maxFinishTime.Value);
+			}
+
+			List<LeaseDocument> results = await collection.FindWithHint(filter, indexHint, x => x.SortByDescending(x => x.FinishTime).Range(index, count).ToListAsync());
+			return results.ConvertAll<ILease>(x => x);
+		}
+
+		/// <inheritdoc/>
+		public async Task<List<ILease>> FindLeasesAsync(DateTime? minTime, DateTime? maxTime)
+		{
+			return await FindLeasesAsync(null, null, minTime, maxTime, null, null, _finishTimeStartTimeCompoundIndex.Name, false);
+		}
+
+		/// <inheritdoc/>
+		public async Task<List<ILease>> FindActiveLeasesAsync(int? index = null, int? count = null)
+		{
+			List<LeaseDocument> results = await _leases.Find(x => x.FinishTime == null).Range(index, count).ToListAsync();
+			return results.ConvertAll<ILease>(x => x);
+		}
+
+		/// <inheritdoc/>
+		public async Task<bool> TrySetOutcomeAsync(LeaseId leaseId, DateTime finishTime, LeaseOutcome outcome, byte[]? output)
+		{
+			FilterDefinitionBuilder<LeaseDocument> filterBuilder = Builders<LeaseDocument>.Filter;
+			FilterDefinition<LeaseDocument> filter = filterBuilder.Eq(x => x.Id, leaseId) & filterBuilder.Eq(x => x.FinishTime, null);
+
+			UpdateDefinitionBuilder<LeaseDocument> updateBuilder = Builders<LeaseDocument>.Update;
+			UpdateDefinition<LeaseDocument> update = updateBuilder.Set(x => x.FinishTime, finishTime).Set(x => x.Outcome, outcome);
+
+			if(output != null && output.Length > 0)
+			{
+				update = update.Set(x => x.Output, output);
+			}
+
+			UpdateResult result = await _leases.UpdateOneAsync(filter, update);
+			return result.ModifiedCount > 0;
+		}
+	}
+}
diff --git a/Engine/Source/Programs/Horde/Horde.Build/Agents/Leases/LeasesController.Messages.cs b/Engine/Source/Programs/Horde/Horde.Build/Agents/Leases/LeasesController.Messages.cs
new file mode 100644
--- /dev/null
+++ b/Engine/Source/Programs/Horde/Horde.Build/Agents/Leases/LeasesController.Messages.cs
@@ -0,0 +1,16 @@
+// Copyright Epic Games, Inc. All Rights Reserved.
+
+namespace Horde.Build.Api
+{
+
+	/// <summary>
+	/// Updates an existing lease
+	/// </summary>
+	public class UpdateLeaseRequest
+	{
+		/// <summary>
+		/// Mark this lease as aborted
+		/// </summary>
+		public bool? Aborted { get; set; }
+	}
+}
diff --git a/Engine/Source/Programs/Horde/Horde.Build/Agents/Leases/LeasesController.cs b/Engine/Source/Programs/Horde/Horde.Build/Agents/Leases/LeasesController.cs
new file mode 100644
--- /dev/null
+++ b/Engine/Source/Programs/Horde/Horde.Build/Agents/Leases/LeasesController.cs
@@ -0,0 +1,220 @@
+// Copyright Epic Games, Inc. All Rights Reserved.
+
+using System;
+using System.Collections.Generic;
+using System.Linq;
+using System.Threading.Tasks;
+using Horde.Build.Acls;
+using Horde.Build.Api;
+using Horde.Build.Models;
+using Horde.Build.Services;
+using Horde.Build.Utilities;
+using Microsoft.AspNetCore.Authorization;
+using Microsoft.AspNetCore.Mvc;
+using OpenTracing;
+using OpenTracing.Util;
+
+namespace Horde.Build.Controllers
+{
+	using LeaseId = ObjectId<ILease>;
+	using SessionId = ObjectId<ISession>;
+
+	/// <summary>
+	/// Controller for the /api/v1/leases endpoint
+	/// </summary>
+	[ApiController]
+	[Authorize]
+	[Route("[controller]")]
+	public class LeasesController : HordeControllerBase
+	{
+		/// <summary>
+		/// Singleton instance of the ACL service
+		/// </summary>
+		readonly AclService _aclService;
+
+		/// <summary>
+		/// Singleton instance of the agent service
+		/// </summary>
+		readonly AgentService _agentService;
+
+		/// <summary>
+		/// Constructor
+		/// </summary>
+		/// <param name="aclService">The ACL service singleton</param>
+		/// <param name="agentService">The agent service</param>
+		public LeasesController(AclService aclService, AgentService agentService)
+		{
+			_aclService = aclService;
+			_agentService = agentService;
+		}
+
+		/// <summary>
+		/// Find all the leases for a particular agent
+		/// </summary>
+		/// <param name="agentId">Unique id of the agent to find</param>
+		/// <param name="sessionId">The session to query</param>
+		/// <param name="startTime">Start of the time window to consider</param>
+		/// <param name="finishTime">End of the time window to consider</param>
+		/// <param name="minFinishTime">Start of the time window to consider when querying by finish time (if set, other criteria are ignored)</param>
+		/// <param name="maxFinishTime">End of the time window to consider when querying by finish time (if set, other criteria are ignored)</param>
+		/// <param name="index">Index of the first result to return</param>
+		/// <param name="count">Number of results to return</param>
+		/// <param name="filter">Filter to apply to the results</param>
+		/// <returns>Sessions </returns>
+		[HttpGet]
+		[Route("/api/v1/leases")]
+		[ProducesResponseType(200, Type = typeof(List<GetAgentLeaseResponse>))]
+		public async Task<ActionResult<List<object>>> FindLeasesAsync(
+			[FromQuery] AgentId? agentId,
+			[FromQuery] SessionId? sessionId,
+			[FromQuery] DateTimeOffset? startTime,
+			[FromQuery] DateTimeOffset? finishTime,
+			[FromQuery] DateTimeOffset? minFinishTime,
+			[FromQuery] DateTimeOffset? maxFinishTime,
+			[FromQuery] int index = 0,
+			[FromQuery] int count = 1000,
+			[FromQuery] PropertyFilter? filter = null)
+		{
+			GlobalPermissionsCache permissionsCache = new GlobalPermissionsCache();
+			if (agentId == null)
+			{
+				if (!await _aclService.AuthorizeAsync(AclAction.ViewLeases, User, permissionsCache))
+				{
+					return Forbid(AclAction.ViewLeases);
+				}
+			}
+			else
+			{
+				IAgent? agent = await _agentService.GetAgentAsync(agentId.Value);
+				if (agent == null)
+				{
+					return NotFound(agentId.Value);
+				}
+				if (!await _agentService.AuthorizeAsync(agent, AclAction.ViewLeases, User, permissionsCache))
+				{
+					return Forbid(AclAction.ViewLeases, agentId.Value);
+				}
+			}
+
+			bool includeCosts = await _aclService.AuthorizeAsync(AclAction.ViewCosts, User, permissionsCache);
+
+			List<ILease> leases;
+			if (minFinishTime == null && maxFinishTime == null)
+			{
+				leases = await _agentService.FindLeasesAsync(agentId, sessionId, startTime?.UtcDateTime, finishTime?.UtcDateTime, index, count);
+			}
+			else
+			{
+				// Optimized path for queries made by finish time
+				leases = await _agentService.FindLeasesByFinishTimeAsync(minFinishTime?.UtcDateTime, maxFinishTime?.UtcDateTime, index, count);
+			}
+
+			List<object> responses = new List<object>();
+
+			using (IScope _ = GlobalTracer.Instance.BuildSpan($"GenerateResponses").StartActive())
+			{
+				Dictionary<AgentId, double?> cachedAgentRates = new Dictionary<AgentId, double?>();
+				foreach (ILease lease in leases)
+				{
+					double? agentRate = null;
+					if (includeCosts && !cachedAgentRates.TryGetValue(lease.AgentId, out agentRate))
+					{
+						agentRate = await _agentService.GetRateAsync(lease.AgentId);
+						cachedAgentRates.Add(lease.AgentId, agentRate);
+					}
+
+					Dictionary<string, string>? details = _agentService.GetPayloadDetails(lease.Payload);
+					responses.Add(PropertyFilter.Apply(new GetAgentLeaseResponse(lease, details, agentRate), filter));
+				}
+			}
+
+			return responses;
+		}
+
+		/// <summary>
+		/// Get info about a particular lease
+		/// </summary>
+		/// <param name="leaseId">Unique id of the particular lease</param>
+		/// <returns>Lease matching the given id</returns>
+		[HttpGet]
+		[Route("/api/v1/leases/{leaseId}")]
+		public async Task<ActionResult<GetAgentLeaseResponse>> GetLeaseAsync(LeaseId leaseId)
+		{
+			ILease? lease = await _agentService.GetLeaseAsync(leaseId);
+			if (lease == null)
+			{
+				return NotFound(leaseId);
+			}
+
+			IAgent? agent = await _agentService.GetAgentAsync(lease.AgentId);
+			if (agent == null)
+			{
+				return NotFound(lease.AgentId);
+			}
+
+			GlobalPermissionsCache permissionsCache = new GlobalPermissionsCache();
+			if (!await _agentService.AuthorizeAsync(agent, AclAction.ViewLeases, User, permissionsCache))
+			{
+				return Forbid(AclAction.ViewLeases);
+			}
+
+			double? agentRate = null;
+			if (await _aclService.AuthorizeAsync(AclAction.ViewCosts, User, permissionsCache))
+			{
+				agentRate = await _agentService.GetRateAsync(agent.Id);
+			}
+
+			Dictionary<string, string>? details = _agentService.GetPayloadDetails(lease.Payload);
+			return new GetAgentLeaseResponse(lease, details, agentRate);
+		}
+
+		/// <summary>
+		/// Update a particular lease
+		/// </summary>
+		/// <param name="leaseId">Unique id of the particular lease</param>
+		/// <param name="request"></param>
+		/// <returns>Lease matching the given id</returns>
+		[HttpPut]
+		[Route("/api/v1/leases/{leaseId}")]
+		public async Task<ActionResult> UpdateLeaseAsync(LeaseId leaseId, [FromBody] UpdateLeaseRequest request)
+		{
+			// only update supported right now is abort
+			if (!request.Aborted.HasValue || !request.Aborted.Value)
+			{
+				return Ok();
+			}
+
+			ILease? lease = await _agentService.GetLeaseAsync(leaseId);
+			if (lease == null)
+			{
+				return NotFound(leaseId);
+			}
+
+			IAgent? agent = await _agentService.GetAgentAsync(lease.AgentId);
+			if (agent == null)
+			{
+				return NotFound(lease.AgentId);
+			}
+
+			if (!await _agentService.AuthorizeAsync(agent, AclAction.AdminWrite, User, null))
+			{
+				return Forbid(AclAction.AdminWrite, lease.AgentId);
+			}
+
+			AgentLease? agentLease = agent.Leases.FirstOrDefault(x => x.Id == leaseId);
+
+			if (agentLease == null)
+			{
+				return NotFound(agent.Id, leaseId);
+			}
+
+			if (!agentLease.IsConformLease())
+			{
+				return BadRequest("Lease abort only supported on conform leases for now, {LeaseId}", leaseId);
+			}
+
+			await _agentService.CancelLeaseAsync(agent, leaseId);
+			return Ok();
+		}
+	}
+}
diff --git a/Engine/Source/Programs/Horde/Horde.Build/Agents/Pools/IPoolCollection.cs b/Engine/Source/Programs/Horde/Horde.Build/Agents/Pools/IPoolCollection.cs
new file mode 100644
--- /dev/null
+++ b/Engine/Source/Programs/Horde/Horde.Build/Agents/Pools/IPoolCollection.cs
@@ -0,0 +1,114 @@
+// Copyright Epic Games, Inc. All Rights Reserved.
+
+using System;
+using System.Collections.Generic;
+using System.Threading.Tasks;
+using EpicGames.Horde.Common;
+using Horde.Build.Fleet.Autoscale;
+using Horde.Build.Models;
+using Horde.Build.Utilities;
+
+namespace Horde.Build.Collections
+{
+	using PoolId = StringId<IPool>;
+
+	/// <summary>
+	/// Collection of pool documents
+	/// </summary>
+	public interface IPoolCollection
+	{
+		/// <summary>
+		/// Creates a new pool
+		/// </summary>
+		/// <param name="id">Unique id for the new pool</param>
+		/// <param name="name">Name of the new pool</param>
+		/// <param name="condition">Condition for agents to be included in this pool</param>
+		/// <param name="enableAutoscaling">Whether to enable autoscaling for this pool</param>
+		/// <param name="minAgents">Minimum number of agents in the pool</param>
+		/// <param name="numReserveAgents">Minimum number of idle agents to maintain</param>
+		/// <param name="scaleOutCooldown">Cooldown time between scale-out events</param>
+		/// <param name="scaleInCooldown">Cooldown time between scale-in events</param>
+		/// <param name="sizeStrategy">Pool sizing strategy</param>
+		/// <param name="leaseUtilizationSettings">Settings for lease utilization strategy</param>
+		/// <param name="jobQueueSettings">Settings for job queue strategy</param>
+		/// <param name="properties">Properties for the pool</param>
+		/// <returns>The new pool document</returns>
+		Task<IPool> AddAsync(
+			PoolId id,
+			string name,
+			Condition? condition = null,
+			bool? enableAutoscaling = null,
+			int? minAgents = null,
+			int? numReserveAgents = null,
+			TimeSpan? scaleOutCooldown = null,
+			TimeSpan? scaleInCooldown = null,
+			PoolSizeStrategy? sizeStrategy = null,
+			LeaseUtilizationSettings? leaseUtilizationSettings = null,
+			JobQueueSettings? jobQueueSettings = null,
+			IEnumerable<KeyValuePair<string, string>>? properties = null);
+
+		/// <summary>
+		/// Enumerates all the pools
+		/// </summary>
+		/// <returns>The pool documents</returns>
+		Task<List<IPool>> GetAsync();
+
+		/// <summary>
+		/// Gets a pool by ID
+		/// </summary>
+		/// <param name="id">Unique id of the pool</param>
+		/// <returns>The pool document</returns>
+		Task<IPool?> GetAsync(PoolId id);
+
+		/// <summary>
+		/// Gets a list of all valid pool ids
+		/// </summary>
+		/// <returns>List of pool ids</returns>
+		Task<List<PoolId>> GetPoolIdsAsync();
+
+		/// <summary>
+		/// Gets a pool by ID
+		/// </summary>
+		/// <param name="id">Unique id of the pool</param>
+		/// <returns>The pool document</returns>
+		Task<bool> DeleteAsync(PoolId id);
+
+		/// <summary>
+		/// Updates a pool
+		/// </summary>
+		/// <param name="pool">The pool to update</param>
+		/// <param name="newName">New name for the pool</param>
+		/// <param name="newCondition">New condition for the pool</param>
+		/// <param name="newEnableAutoscaling">New setting for whether to enable autoscaling</param>
+		/// <param name="newMinAgents">Minimum number of agents in the pool</param>
+		/// <param name="newNumReserveAgents">Minimum number of idle agents to maintain</param>
+		/// <param name="newWorkspaces">New workspaces for the pool</param>
+		/// <param name="newUseAutoSdk">New setting for whether to use autosdk</param>
+		/// <param name="newProperties">New properties for the pool</param>
+		/// <param name="lastScaleUpTime">New time for last (auto) scale up</param>
+		/// <param name="lastScaleDownTime">New time for last (auto) scale down</param>
+		/// <param name="scaleOutCooldown">Cooldown time between scale-out events</param>
+		/// <param name="scaleInCooldown">Cooldown time between scale-in events</param>
+		/// <param name="sizeStrategy">Pool sizing strategy</param>
+		/// <param name="leaseUtilizationSettings">Settings for lease utilization strategy</param>
+		/// <param name="jobQueueSettings">Settings for job queue strategy</param>
+		/// <returns>Async task</returns>
+		Task<IPool?> TryUpdateAsync(
+			IPool pool,
+			string? newName = null,
+			Condition? newCondition = null,
+			bool? newEnableAutoscaling = null,
+			int? newMinAgents = null,
+			int? newNumReserveAgents = null,
+			List<AgentWorkspace>? newWorkspaces = null,
+			bool? newUseAutoSdk = null,
+			Dictionary<string, string?>? newProperties = null,
+			DateTime? lastScaleUpTime = null,
+			DateTime? lastScaleDownTime = null,
+			TimeSpan? scaleOutCooldown = null,
+			TimeSpan? scaleInCooldown = null,
+			PoolSizeStrategy? sizeStrategy = null,
+			LeaseUtilizationSettings? leaseUtilizationSettings = null,
+			JobQueueSettings? jobQueueSettings = null);
+	}
+}
diff --git a/Engine/Source/Programs/Horde/Horde.Build/Agents/Pools/PoolCollection.cs b/Engine/Source/Programs/Horde/Horde.Build/Agents/Pools/PoolCollection.cs
new file mode 100644
--- /dev/null
+++ b/Engine/Source/Programs/Horde/Horde.Build/Agents/Pools/PoolCollection.cs
@@ -0,0 +1,321 @@
+// Copyright Epic Games, Inc. All Rights Reserved.
+
+using System;
+using System.Collections.Generic;
+using System.Linq;
+using System.Threading.Tasks;
+using EpicGames.Horde.Common;
+using Horde.Build.Fleet.Autoscale;
+using Horde.Build.Models;
+using Horde.Build.Server;
+using Horde.Build.Utilities;
+using MongoDB.Bson;
+using MongoDB.Bson.Serialization.Attributes;
+using MongoDB.Driver;
+
+namespace Horde.Build.Collections.Impl
+{
+	using PoolId = StringId<IPool>;
+
+	/// <summary>
+	/// Collection of pool documents
+	/// </summary>
+	public class PoolCollection : IPoolCollection
+	{
+		class PoolDocument : IPool
+		{
+			// Properties
+			[BsonId]
+			public PoolId Id { get; set; }
+			[BsonRequired]
+			public string Name { get; set; } = null!;
+			[BsonIgnoreIfNull]
+			public Condition? Condition { get; set; }
+			public List<AgentWorkspace> Workspaces { get; set; } = new List<AgentWorkspace>();
+			public bool UseAutoSdk { get; set; } = true;
+			public Dictionary<string, string> Properties { get; set; } = new Dictionary<string, string>();
+			public bool EnableAutoscaling { get; set; } = true;
+			public int? MinAgents { get; set; }
+			public int? NumReserveAgents { get; set; }
+			public DateTime? LastScaleUpTime { get; set; }
+			public DateTime? LastScaleDownTime { get; set; }
+			public TimeSpan? ScaleOutCooldown { get; set; }
+			public TimeSpan? ScaleInCooldown { get; set; }
+			public PoolSizeStrategy? SizeStrategy { get; set; }
+			PoolSizeStrategy IPool.SizeStrategy => SizeStrategy ?? PoolSizeStrategy.Default;
+			public LeaseUtilizationSettings? LeaseUtilizationSettings { get; set; }
+			public JobQueueSettings? JobQueueSettings { get; set; }
+			public int UpdateIndex { get; set; }
+
+			// Read-only wrappers
+			IReadOnlyList<AgentWorkspace> IPool.Workspaces => Workspaces;
+			IReadOnlyDictionary<string, string> IPool.Properties => Properties;
+
+			public PoolDocument()
+			{
+			}
+
+			public PoolDocument(IPool other)
+			{
+				Id = other.Id;
+				Name = other.Name;
+				Condition = other.Condition;
+				Workspaces.AddRange(other.Workspaces);
+				UseAutoSdk = other.UseAutoSdk;
+				Properties = new Dictionary<string, string>(other.Properties);
+				EnableAutoscaling = other.EnableAutoscaling;
+				MinAgents = other.MinAgents;
+				NumReserveAgents = other.NumReserveAgents;
+				LastScaleUpTime = other.LastScaleUpTime;
+				LastScaleDownTime = other.LastScaleDownTime;
+				ScaleOutCooldown = other.ScaleOutCooldown;
+				ScaleInCooldown = other.ScaleInCooldown;
+				SizeStrategy = other.SizeStrategy;
+				LeaseUtilizationSettings = other.LeaseUtilizationSettings;
+				JobQueueSettings = other.JobQueueSettings;
+				UpdateIndex = other.UpdateIndex;
+			}
+		}
+
+		/// <summary>
+		/// Collection of pool documents
+		/// </summary>
+		readonly IMongoCollection<PoolDocument> _pools;
+
+		/// <summary>
+		/// Constructor
+		/// </summary>
+		/// <param name="mongoService">The database service instance</param>
+		public PoolCollection(MongoService mongoService)
+		{
+			_pools = mongoService.GetCollection<PoolDocument>("Pools");
+		}
+
+		/// <inheritdoc/>
+		public async Task<IPool> AddAsync(
+			PoolId id,
+			string name,
+			Condition? condition,
+			bool? enableAutoscaling,
+			int? minAgents,
+			int? numReserveAgents,
+			TimeSpan? scaleOutCooldown,
+			TimeSpan? scaleInCooldown,
+			PoolSizeStrategy? sizeStrategy,
+			LeaseUtilizationSettings? leaseUtilizationSettings,
+			JobQueueSettings? jobQueueSettings,
+			IEnumerable<KeyValuePair<string, string>>? properties)
+		{
+			PoolDocument pool = new PoolDocument();
+			pool.Id = id;
+			pool.Name = name;
+			pool.Condition = condition;
+			if (enableAutoscaling != null)
+			{
+				pool.EnableAutoscaling = enableAutoscaling.Value;
+			}
+			pool.MinAgents = minAgents;
+			pool.NumReserveAgents = numReserveAgents;
+			if (properties != null)
+			{
+				pool.Properties = new Dictionary<string, string>(properties);
+			}
+
+			pool.ScaleOutCooldown = scaleOutCooldown;
+			pool.ScaleInCooldown = scaleInCooldown;
+			pool.SizeStrategy = sizeStrategy;
+			pool.LeaseUtilizationSettings = leaseUtilizationSettings;
+			pool.JobQueueSettings = jobQueueSettings;
+			await _pools.InsertOneAsync(pool);
+			return pool;
+		}
+
+		/// <inheritdoc/>
+		public async Task<List<IPool>> GetAsync()
+		{
+			List<PoolDocument> results = await _pools.Find(FilterDefinition<PoolDocument>.Empty).ToListAsync();
+			return results.ConvertAll<IPool>(x => x);
+		}
+
+		/// <inheritdoc/>
+		public async Task<IPool?> GetAsync(PoolId id)
+		{
+			return await _pools.Find(x => x.Id == id).FirstOrDefaultAsync();
+		}
+
+		/// <inheritdoc/>
+		public async Task<List<PoolId>> GetPoolIdsAsync()
+		{
+			ProjectionDefinition<PoolDocument, BsonDocument> projection = Builders<PoolDocument>.Projection.Include(x => x.Id);
+			List<BsonDocument> results = await _pools.Find(FilterDefinition<PoolDocument>.Empty).Project(projection).ToListAsync();
+			return results.ConvertAll(x => new PoolId(x.GetElement("_id").Value.AsString));
+		}
+
+		/// <inheritdoc/>
+		public async Task<bool> DeleteAsync(PoolId id)
+		{
+			FilterDefinition<PoolDocument> filter = Builders<PoolDocument>.Filter.Eq(x => x.Id, id);
+			DeleteResult result = await _pools.DeleteOneAsync(filter);
+			return result.DeletedCount > 0;
+		}
+
+		/// <summary>
+		/// Attempts to update a pool, upgrading it to the latest document schema if necessary
+		/// </summary>
+		/// <param name="poolToUpdate">Interface for the document to update</param>
+		/// <param name="transaction">The transaction to execute</param>
+		/// <returns>New pool definition, or null on failure.</returns>
+		async Task<IPool?> TryUpdateAsync(IPool poolToUpdate, TransactionBuilder<PoolDocument> transaction)
+		{
+			if (transaction.IsEmpty)
+			{
+				return poolToUpdate;
+			}
+
+			transaction.Set(x => x.UpdateIndex, poolToUpdate.UpdateIndex + 1);
+
+			PoolDocument? mutablePool = poolToUpdate as PoolDocument;
+			if (mutablePool != null)
+			{
+				UpdateResult result = await _pools.UpdateOneAsync(x => x.Id == poolToUpdate.Id && x.UpdateIndex == poolToUpdate.UpdateIndex, transaction.ToUpdateDefinition());
+				if (result.ModifiedCount == 0)
+				{
+					return null;
+				}
+
+				transaction.ApplyTo(mutablePool);
+			}
+			else
+			{
+				mutablePool = new PoolDocument(poolToUpdate);
+				transaction.ApplyTo(mutablePool);
+
+				ReplaceOneResult result = await _pools.ReplaceOneAsync<PoolDocument>(x => x.Id == poolToUpdate.Id && x.UpdateIndex == poolToUpdate.UpdateIndex, mutablePool);
+				if (result.ModifiedCount == 0)
+				{
+					return null;
+				}
+			}
+			return mutablePool;
+		}
+
+		/// <inheritdoc/>
+		public Task<IPool?> TryUpdateAsync(
+			IPool pool,
+			string? newName,
+			Condition? newCondition,
+			bool? newEnableAutoscaling,
+			int? newMinAgents,
+			int? newNumReserveAgents,
+			List<AgentWorkspace>? newWorkspaces,
+			bool? newUseAutoSdk,
+			Dictionary<string, string?>? newProperties,
+			DateTime? lastScaleUpTime,
+			DateTime? lastScaleDownTime,
+			TimeSpan? scaleOutCooldown,
+			TimeSpan? scaleInCooldown, 
+			PoolSizeStrategy? sizeStrategy,
+			LeaseUtilizationSettings? leaseUtilizationSettings,
+			JobQueueSettings? jobQueueSettings)
+		{
+			TransactionBuilder<PoolDocument> transaction = new TransactionBuilder<PoolDocument>();
+			if (newName != null)
+			{
+				transaction.Set(x => x.Name, newName);
+			}
+			if (newCondition != null)
+			{
+				if (newCondition.IsEmpty())
+				{
+					transaction.Unset(x => x.Condition!);
+				}
+				else
+				{
+					transaction.Set(x => x.Condition, newCondition);
+				}
+			}
+			if (newEnableAutoscaling != null)
+			{
+				if (newEnableAutoscaling.Value)
+				{
+					transaction.Unset(x => x.EnableAutoscaling);
+				}
+				else
+				{
+					transaction.Set(x => x.EnableAutoscaling, newEnableAutoscaling.Value);
+				}
+			}
+			if (newMinAgents != null)
+			{
+				if (newMinAgents.Value < 0)
+				{
+					transaction.Unset(x => x.MinAgents!);
+				}
+				else
+				{
+					transaction.Set(x => x.MinAgents, newMinAgents.Value);
+				}
+			}
+			if (newNumReserveAgents != null)
+			{
+				if (newNumReserveAgents.Value < 0)
+				{
+					transaction.Unset(x => x.NumReserveAgents!);
+				}
+				else
+				{
+					transaction.Set(x => x.NumReserveAgents, newNumReserveAgents.Value);
+				}
+			}
+			if (newWorkspaces != null)
+			{
+				transaction.Set(x => x.Workspaces, newWorkspaces);
+			}
+			if (newUseAutoSdk != null)
+			{
+				transaction.Set(x => x.UseAutoSdk, newUseAutoSdk.Value);
+			}
+			if (newProperties != null)
+			{
+				transaction.UpdateDictionary(x => x.Properties, newProperties);
+			}
+			if (lastScaleUpTime != null)
+			{
+				transaction.Set(x => x.LastScaleUpTime, lastScaleUpTime);
+			}
+			if (lastScaleDownTime != null)
+			{
+				transaction.Set(x => x.LastScaleDownTime, lastScaleDownTime);
+			}
+			if (scaleOutCooldown != null)
+			{
+				transaction.Set(x => x.ScaleOutCooldown, scaleOutCooldown);
+			}
+			if (scaleInCooldown != null)
+			{
+				transaction.Set(x => x.ScaleInCooldown, scaleInCooldown);
+			}
+			if (sizeStrategy != null)
+			{
+				if (sizeStrategy == PoolSizeStrategy.Default)
+				{
+					transaction.Set(x => x.SizeStrategy, null);
+				}
+				else
+				{
+					transaction.Set(x => x.SizeStrategy, sizeStrategy);
+				}
+				
+			}
+			if (leaseUtilizationSettings != null)
+			{
+				transaction.Set(x => x.LeaseUtilizationSettings, leaseUtilizationSettings);
+			}
+			if (jobQueueSettings != null)
+			{
+				transaction.Set(x => x.JobQueueSettings, jobQueueSettings);
+			}
+			return TryUpdateAsync(pool, transaction);
+		}
+	}
+}
diff --git a/Engine/Source/Programs/Horde/Horde.Build/Agents/Pools/PoolService.cs b/Engine/Source/Programs/Horde/Horde.Build/Agents/Pools/PoolService.cs
new file mode 100644
--- /dev/null
+++ b/Engine/Source/Programs/Horde/Horde.Build/Agents/Pools/PoolService.cs
@@ -0,0 +1,239 @@
+// Copyright Epic Games, Inc. All Rights Reserved.
+
+using System;
+using System.Collections.Generic;
+using System.Linq;
+using System.Threading;
+using System.Threading.Tasks;
+using EpicGames.Horde.Common;
+using Horde.Build.Collections;
+using Horde.Build.Fleet.Autoscale;
+using Horde.Build.Models;
+using Horde.Build.Server;
+using Horde.Build.Utilities;
+using HordeCommon;
+
+namespace Horde.Build.Services
+{
+	using PoolId = StringId<IPool>;
+
+	/// <summary>
+	/// Wraps functionality for manipulating pools
+	/// </summary>
+	public class PoolService
+	{
+		/// <summary>
+		/// The database service instance
+		/// </summary>
+		readonly MongoService _mongoService;
+
+		/// <summary>
+		/// Collection of pool documents
+		/// </summary>
+		readonly IPoolCollection _pools;
+
+		/// <summary>
+		/// Returns the current time
+		/// </summary>
+		readonly IClock _clock;
+
+		/// <summary>
+		/// Cached set of pools, along with the timestamp that it was obtained
+		/// </summary>
+		Tuple<DateTime, Dictionary<PoolId, IPool>>? _cachedPoolLookup;
+
+		/// <summary>
+		/// Constructor
+		/// </summary>
+		/// <param name="mongoService"></param>
+		/// <param name="pools">Collection of pool documents</param>
+		/// <param name="clock"></param>
+		public PoolService(MongoService mongoService, IPoolCollection pools, IClock clock)
+		{
+			_mongoService = mongoService;
+			_pools = pools;
+			_clock = clock;
+		}
+
+		/// <summary>
+		/// Creates a new pool
+		/// </summary>
+		/// <param name="name">Name of the new pool</param>
+		/// <param name="condition">Condition for agents to be automatically included in this pool</param>
+		/// <param name="enableAutoscaling">Whether to enable autoscaling for this pool</param>
+		/// <param name="minAgents">Minimum number of agents in the pool</param>
+		/// <param name="numReserveAgents">Minimum number of idle agents to maintain</param>
+		/// <param name="scaleOutCooldown">Cooldown time between scale-out events</param>
+		/// <param name="scaleInCooldown">Cooldown time between scale-in events</param>
+		/// <param name="sizeStrategy">Pool sizing strategy</param>
+		/// <param name="leaseUtilizationSettings">Settings for lease utilization strategy</param>
+		/// <param name="jobQueueSettings">Settings for job queue strategy</param>
+		/// <param name="properties">Properties for the new pool</param>
+		/// <returns>The new pool document</returns>
+		public Task<IPool> CreatePoolAsync(
+			string name,
+			Condition? condition = null,
+			bool? enableAutoscaling = null,
+			int? minAgents = null,
+			int? numReserveAgents = null,
+			TimeSpan? scaleOutCooldown = null,
+			TimeSpan? scaleInCooldown = null,
+			PoolSizeStrategy? sizeStrategy = null,
+			LeaseUtilizationSettings? leaseUtilizationSettings = null,
+			JobQueueSettings? jobQueueSettings = null,
+			Dictionary<string, string>? properties = null)
+		{
+			return _pools.AddAsync(
+				PoolId.Sanitize(name),
+				name,
+				condition,
+				enableAutoscaling,
+				minAgents,
+				numReserveAgents,
+				scaleOutCooldown,
+				scaleInCooldown,
+				sizeStrategy,
+				leaseUtilizationSettings,
+				jobQueueSettings,
+				properties);
+		}
+
+		/// <summary>
+		/// Deletes a pool
+		/// </summary>
+		/// <param name="poolId">Unique id of the pool</param>
+		/// <returns>Async task object</returns>
+		public Task<bool> DeletePoolAsync(PoolId poolId)
+		{
+			return _pools.DeleteAsync(poolId);
+		}
+
+		/// <summary>
+		/// Updates an existing pool
+		/// </summary>
+		/// <param name="pool">The pool to update</param>
+		/// <param name="newName">The new name for the pool</param>
+		/// <param name="newCondition">New requirements for the pool</param>
+		/// <param name="newEnableAutoscaling">Whether to enable autoscaling</param>
+		/// <param name="newMinAgents">Minimum number of agents in the pool</param>
+		/// <param name="newNumReserveAgents">Minimum number of idle agents to maintain</param>
+		/// <param name="newProperties">Properties on the pool to update. Any properties with a value of null will be removed.</param>
+		/// <param name="scaleOutCooldown">Cooldown time between scale-out events</param>
+		/// <param name="scaleInCooldown">Cooldown time between scale-in events</param>
+		/// <param name="sizeStrategy">New pool sizing strategy for the pool</param>
+		/// <returns>Async task object</returns>
+		public async Task<IPool?> UpdatePoolAsync(
+			IPool? pool,
+			string? newName = null,
+			Condition? newCondition = null,
+			bool? newEnableAutoscaling = null,
+			int? newMinAgents = null,
+			int? newNumReserveAgents = null,
+			Dictionary<string, string?>? newProperties = null,
+			TimeSpan? scaleOutCooldown = null,
+			TimeSpan? scaleInCooldown = null,
+			PoolSizeStrategy? sizeStrategy = null)
+		{
+			for (; pool != null; pool = await _pools.GetAsync(pool.Id))
+			{
+				IPool? newPool = await _pools.TryUpdateAsync(
+					pool,
+					newName,
+					newCondition,
+					newEnableAutoscaling,
+					newMinAgents,
+					newNumReserveAgents,
+					newProperties: newProperties,
+					scaleOutCooldown: scaleOutCooldown,
+					scaleInCooldown: scaleInCooldown,
+					sizeStrategy: sizeStrategy);
+				
+				if (newPool != null)
+				{
+					return newPool;
+				}
+			}
+			return pool;
+		}
+
+		/// <summary>
+		/// Gets all the available pools
+		/// </summary>
+		/// <returns>List of pool documents</returns>
+		public Task<List<IPool>> GetPoolsAsync()
+		{
+			return _pools.GetAsync();
+		}
+
+		/// <summary>
+		/// Gets a pool by ID
+		/// </summary>
+		/// <param name="poolId">Unique id of the pool</param>
+		/// <returns>The pool document</returns>
+		public Task<IPool?> GetPoolAsync(PoolId poolId)
+		{
+			return _pools.GetAsync(poolId);
+		}
+
+		/// <summary>
+		/// Get a list of workspaces for the given agent
+		/// </summary>
+		/// <param name="agent">The agent to return workspaces for</param>
+		/// <param name="validAtTime">Absolute time at which we expect the results to be valid. Values may be cached as long as they are after this time.</param>
+		/// <returns>List of workspaces</returns>
+		public async Task<HashSet<AgentWorkspace>> GetWorkspacesAsync(IAgent agent, DateTime validAtTime)
+		{
+			bool bAddAutoSdkWorkspace = false;
+			HashSet<AgentWorkspace> workspaces = new HashSet<AgentWorkspace>();
+
+			Dictionary<PoolId, IPool> poolMapping = await GetPoolLookupAsync(validAtTime);
+			foreach (PoolId poolId in agent.GetPools())
+			{
+				IPool? pool;
+				if (poolMapping.TryGetValue(poolId, out pool))
+				{
+					workspaces.UnionWith(pool.Workspaces);
+					bAddAutoSdkWorkspace |= pool.UseAutoSdk;
+				}
+			}
+
+			if (bAddAutoSdkWorkspace)
+			{
+				Globals globals = await _mongoService.GetGlobalsAsync();
+				workspaces.UnionWith(agent.GetAutoSdkWorkspaces(globals, workspaces.ToList()));
+			}
+
+			return workspaces;
+		}
+
+		/// <summary>
+		/// Gets a mapping from pool identifiers to definitions
+		/// </summary>
+		/// <param name="validAtTime">Absolute time at which we expect the results to be valid. Values may be cached as long as they are after this time.</param>
+		/// <returns>Map of pool ids to pool documents</returns>
+		private async Task<Dictionary<PoolId, IPool>> GetPoolLookupAsync(DateTime validAtTime)
+		{
+			Tuple<DateTime, Dictionary<PoolId, IPool>>? cachedPoolLookupCopy = _cachedPoolLookup;
+			if (cachedPoolLookupCopy == null || cachedPoolLookupCopy.Item1 < validAtTime)
+			{
+				// Get a new list of cached pools
+				DateTime newCacheTime = _clock.UtcNow;
+				List<IPool> newPools = await _pools.GetAsync();
+				Tuple<DateTime, Dictionary<PoolId, IPool>> newCachedPoolLookup = Tuple.Create(newCacheTime, newPools.ToDictionary(x => x.Id, x => x));
+
+				// Try to swap it with the current version
+				while (cachedPoolLookupCopy == null || cachedPoolLookupCopy.Item1 < newCacheTime)
+				{
+					Tuple<DateTime, Dictionary<PoolId, IPool>>? originalValue = Interlocked.CompareExchange(ref _cachedPoolLookup, newCachedPoolLookup, cachedPoolLookupCopy);
+					if (originalValue == cachedPoolLookupCopy)
+					{
+						cachedPoolLookupCopy = newCachedPoolLookup;
+						break;
+					}
+					cachedPoolLookupCopy = originalValue;
+				}
+			}
+			return cachedPoolLookupCopy.Item2;
+		}
+	}
+}
diff --git a/Engine/Source/Programs/Horde/Horde.Build/Agents/Pools/PoolUpdateService.cs b/Engine/Source/Programs/Horde/Horde.Build/Agents/Pools/PoolUpdateService.cs
new file mode 100644
--- /dev/null
+++ b/Engine/Source/Programs/Horde/Horde.Build/Agents/Pools/PoolUpdateService.cs
@@ -0,0 +1,136 @@
+// Copyright Epic Games, Inc. All Rights Reserved.
+
+using System;
+using System.Collections.Generic;
+using System.Linq;
+using System.Threading;
+using System.Threading.Tasks;
+using Horde.Build.Collections;
+using Horde.Build.Models;
+using Horde.Build.Utilities;
+using HordeCommon;
+using Microsoft.Extensions.Hosting;
+using Microsoft.Extensions.Logging;
+
+namespace Horde.Build.Services
+{
+	using PoolId = StringId<IPool>;
+
+	/// <summary>
+	/// Periodically updates pool documents to contain the correct workspaces
+	/// </summary>
+	public sealed class PoolUpdateService : IHostedService, IDisposable
+	{
+		readonly IPoolCollection _pools;
+		readonly IStreamCollection _streams;
+		readonly ILogger<PoolService> _logger;
+		readonly ITicker _ticker;
+
+		/// <summary>
+		/// Constructor
+		/// </summary>
+		public PoolUpdateService(IPoolCollection pools, IStreamCollection streams, IClock clock, ILogger<PoolService> logger)
+		{
+			_pools = pools;
+			_streams = streams;
+			_logger = logger;
+			_ticker = clock.AddTicker($"{nameof(PoolUpdateService)}.{nameof(UpdatePoolsAsync)}", TimeSpan.FromSeconds(30.0), UpdatePoolsAsync, logger);
+		}
+
+		/// <inheritdoc/>
+		public Task StartAsync(CancellationToken cancellationToken) => _ticker.StartAsync();
+
+		/// <inheritdoc/>
+		public Task StopAsync(CancellationToken cancellationToken) => _ticker.StopAsync();
+
+		/// <inheritdoc/>
+		public void Dispose() => _ticker.Dispose();
+
+		/// <summary>
+		/// Execute the background task
+		/// </summary>
+		/// <param name="stoppingToken">Cancellation token for the async task</param>
+		/// <returns>Async task</returns>
+		async ValueTask UpdatePoolsAsync(CancellationToken stoppingToken)
+		{
+			// Capture the start time for this operation. We use this to attempt to sequence updates to agents, and prevent overriding another server's updates.
+			DateTime startTime = DateTime.UtcNow;
+
+			// Update the list
+			bool bRetryUpdate = true;
+			while (bRetryUpdate && !stoppingToken.IsCancellationRequested)
+			{
+				_logger.LogDebug("Updating pool->workspace map");
+
+				// Assume this will be the last iteration
+				bRetryUpdate = false;
+
+				// Capture the list of pools at the start of this update
+				List<IPool> currentPools = await _pools.GetAsync();
+
+				// Lookup table of pool id to workspaces
+				HashSet<PoolId> poolsWithAutoSdk = new HashSet<PoolId>();
+				Dictionary<PoolId, List<AgentWorkspace>> poolToAgentWorkspaces = new Dictionary<PoolId, List<AgentWorkspace>>();
+
+				// Populate the workspace list from the current stream
+				List<IStream> activeStreams = await _streams.FindAllAsync();
+				foreach (IStream activeStream in activeStreams)
+				{
+					foreach (KeyValuePair<string, AgentType> agentTypePair in activeStream.AgentTypes)
+					{
+						// Create the new agent workspace
+						(AgentWorkspace, bool)? result;
+						if (activeStream.TryGetAgentWorkspace(agentTypePair.Value, out result))
+						{
+							(AgentWorkspace agentWorkspace, bool useAutoSdk) = result.Value;
+							AgentType agentType = agentTypePair.Value;
+
+							// Find or add a list of workspaces for this pool
+							List<AgentWorkspace>? agentWorkspaces;
+							if (!poolToAgentWorkspaces.TryGetValue(agentType.Pool, out agentWorkspaces))
+							{
+								agentWorkspaces = new List<AgentWorkspace>();
+								poolToAgentWorkspaces.Add(agentType.Pool, agentWorkspaces);
+							}
+
+							// Add it to the list
+							if (!agentWorkspaces.Contains(agentWorkspace))
+							{
+								agentWorkspaces.Add(agentWorkspace);
+							}
+							if (useAutoSdk)
+							{
+								poolsWithAutoSdk.Add(agentType.Pool);
+							}
+						}
+					}
+				}
+
+				// Update the list of workspaces for each pool
+				foreach (IPool currentPool in currentPools)
+				{
+					// Get the new list of workspaces for this pool
+					List<AgentWorkspace>? newWorkspaces;
+					if (!poolToAgentWorkspaces.TryGetValue(currentPool.Id, out newWorkspaces))
+					{
+						newWorkspaces = new List<AgentWorkspace>();
+					}
+
+					// Update the pools document
+					bool useAutoSdk = poolsWithAutoSdk.Contains(currentPool.Id);
+					if (!AgentWorkspace.SetEquals(currentPool.Workspaces, newWorkspaces) || currentPool.Workspaces.Count != newWorkspaces.Count || currentPool.UseAutoSdk != useAutoSdk)
+					{
+						_logger.LogInformation("New workspaces for pool {Pool}:{Workspaces}", currentPool.Id, String.Join("", newWorkspaces.Select(x => $"\n  Identifier=\"{x.Identifier}\", Stream={x.Stream}")));
+
+						IPool? result = await _pools.TryUpdateAsync(currentPool, newWorkspaces: newWorkspaces, newUseAutoSdk: useAutoSdk);
+						if (result == null)
+						{
+							_logger.LogInformation("Pool modified; will retry");
+							bRetryUpdate = true;
+						}
+					}
+				}
+			}
+		}
+	}
+}
