Horde: Finish re-org of Horde solution started a while ago.

diff --git a/Engine/Source/Programs/Horde/Horde.Build/Agents/Pools/PoolsController.Messages.cs b/Engine/Source/Programs/Horde/Horde.Build/Agents/Pools/PoolsController.Messages.cs
new file mode 100644
--- /dev/null
+++ b/Engine/Source/Programs/Horde/Horde.Build/Agents/Pools/PoolsController.Messages.cs
@@ -0,0 +1,328 @@
+// Copyright Epic Games, Inc. All Rights Reserved.
+
+using System.Collections.Generic;
+using System.ComponentModel.DataAnnotations;
+using System.Linq;
+using EpicGames.Horde.Common;
+using Horde.Build.Fleet.Autoscale;
+using Horde.Build.Models;
+
+namespace Horde.Build.Api
+{
+	/// <see cref="Horde.Build.Fleet.Autoscale.LeaseUtilizationSettings" />
+	public class LeaseUtilizationSettings
+	{
+		/// <summary>
+		/// Construct a public REST API representation from the internal one
+		/// </summary>
+		/// <param name="settings"></param>
+#pragma warning disable IDE0060 // Remove unused parameter
+		public LeaseUtilizationSettings(Horde.Build.Fleet.Autoscale.LeaseUtilizationSettings settings)
+#pragma warning restore IDE0060 // Remove unused parameter
+		{
+		}
+		
+		/// <summary>
+		/// Convert this public REST API object to an internal representation
+		/// </summary>
+		/// <returns></returns>
+		public Horde.Build.Fleet.Autoscale.LeaseUtilizationSettings Convert()
+		{
+			return new Horde.Build.Fleet.Autoscale.LeaseUtilizationSettings();
+		}
+	}
+	
+	/// <see cref="Horde.Build.Fleet.Autoscale.JobQueueSettings" />
+	public class JobQueueSettings
+	{
+		/// <summary>
+		/// Factor translating queue size to additional agents to grow the pool with
+		/// The result is always rounded up to nearest integer. 
+		/// Example: if there are 20 jobs in queue, a factor 0.25 will result in 5 new agents being added (20 * 0.25)
+		/// </summary>
+		public double ScaleOutFactor { get; set; }
+
+		/// <summary>
+		/// Factor by which to shrink the pool size with when queue is empty
+		/// The result is always rounded up to nearest integer.
+		/// Example: when the queue size is zero, a default value of 0.9 will shrink the pool by 10% (current agent count * 0.9)
+		/// </summary>
+		public double ScaleInFactor { get; set; }
+
+		/// <summary>
+		/// Construct a public REST API representation from the internal one
+		/// </summary>
+		/// <param name="settings"></param>
+		public JobQueueSettings(Horde.Build.Fleet.Autoscale.JobQueueSettings settings)
+		{
+			ScaleInFactor = settings.ScaleInFactor;
+			ScaleOutFactor = settings.ScaleOutFactor;
+		}
+		
+		/// <summary>
+		/// Convert this public REST API object to an internal representation
+		/// </summary>
+		/// <returns></returns>
+		public Horde.Build.Fleet.Autoscale.JobQueueSettings Convert()
+		{
+			return new Horde.Build.Fleet.Autoscale.JobQueueSettings(ScaleOutFactor, ScaleInFactor);
+		}
+	}
+	
+	/// <summary>
+	/// Parameters to create a new pool
+	/// </summary>
+	public class CreatePoolRequest
+	{
+		/// <summary>
+		/// Name for the new pool
+		/// </summary>
+		[Required]
+		public string Name { get; set; } = null!;
+
+		/// <summary>
+		/// Condition to satisfy for agents to be included in this pool
+		/// </summary>
+		public Condition? Condition { get; set; }
+
+		/// <summary>
+		/// Whether to enable autoscaling for this pool
+		/// </summary>
+		public bool? EnableAutoscaling { get; set; }
+		
+		/// <summary>
+		/// Cooldown time between scale-out events in seconds
+		/// </summary>
+		public int? ScaleOutCooldown { get; set; }
+		
+		/// <summary>
+		/// Cooldown time between scale-in events in seconds
+		/// </summary>
+		public int? ScaleInCooldown { get; set; }
+		
+		/// <summary>
+		/// Pool sizing strategy
+		/// </summary>
+		public PoolSizeStrategy? SizeStrategy { get; set; }
+		
+		/// <summary>
+		/// Settings for lease utilization pool sizing strategy (if used)
+		/// </summary>
+		public LeaseUtilizationSettings? LeaseUtilizationSettings { get; set; }
+		
+		/// <summary>
+		/// Settings for job queue pool sizing strategy (if used) 
+		/// </summary>
+		public JobQueueSettings? JobQueueSettings { get; set; }
+
+		/// <summary>
+		/// The minimum nunmber of agents to retain in this pool
+		/// </summary>
+		public int? MinAgents { get; set; }
+
+		/// <summary>
+		/// The minimum number of idle agents in this pool, if autoscaling is enabled
+		/// </summary>
+		public int? NumReserveAgents { get; set; }
+
+		/// <summary>
+		/// Properties for the new pool
+		/// </summary>
+		public Dictionary<string, string>? Properties { get; set; }
+	}
+
+	/// <summary>
+	/// Response from creating a new pool
+	/// </summary>
+	public class CreatePoolResponse
+	{
+		/// <summary>
+		/// Unique id for the new pool
+		/// </summary>
+		public string Id { get; set; }
+
+		/// <summary>
+		/// Constructor
+		/// </summary>
+		/// <param name="id">Unique id for the new pool</param>
+		public CreatePoolResponse(string id)
+		{
+			Id = id;
+		}
+	}
+
+	/// <summary>
+	/// Parameters to update a pool
+	/// </summary>
+	public class UpdatePoolRequest
+	{
+		/// <summary>
+		/// Optional new name for the pool
+		/// </summary>
+		public string? Name { get; set; }
+
+		/// <summary>
+		/// Requirements for this pool
+		/// </summary>
+		public Condition? Condition { get; set; }
+
+		/// <summary>
+		/// Whether to enable autoscaling for this pool
+		/// </summary>
+		public bool? EnableAutoscaling { get; set; }
+		
+		/// <summary>
+		/// Cooldown time between scale-out events in seconds
+		/// </summary>
+		public int? ScaleOutCooldown { get; set; }
+		
+		/// <summary>
+		/// Cooldown time between scale-in events in seconds
+		/// </summary>
+		public int? ScaleInCooldown { get; set; }
+		
+		/// <summary>
+		/// Pool sizing strategy
+		/// </summary>
+		public PoolSizeStrategy? SizeStrategy { get; set; }
+		
+		/// <summary>
+		/// Settings for lease utilization pool sizing strategy (if used)
+		/// </summary>
+		public LeaseUtilizationSettings? LeaseUtilizationSettings { get; set; }
+		
+		/// <summary>
+		/// Settings for job queue pool sizing strategy (if used) 
+		/// </summary>
+		public JobQueueSettings? JobQueueSettings { get; set; }
+
+		/// <summary>
+		/// The minimum nunmber of agents to retain in this pool
+		/// </summary>
+		public int? MinAgents { get; set; }
+
+		/// <summary>
+		/// The minimum number of idle agents in this pool, if autoscaling is enabled
+		/// </summary>
+		public int? NumReserveAgents { get; set; }
+
+		/// <summary>
+		/// Properties to update for the pool. Properties set to null will be removed.
+		/// </summary>
+		public Dictionary<string, string?>? Properties { get; set; }
+	}
+
+	/// <summary>
+	/// Parameters to update a pool
+	/// </summary>
+	public class BatchUpdatePoolRequest
+	{
+		/// <summary>
+		///  ID of the pool to update
+		/// </summary>
+		[Required]
+		public string Id { get; set; } = null!;
+
+		/// <summary>
+		/// Optional new name for the pool
+		/// </summary>
+		public string? Name { get; set; }
+
+		/// <summary>
+		/// Properties to update for the pool. Properties set to null will be removed.
+		/// </summary>
+		public Dictionary<string, string?>? Properties { get; set; }
+	}
+
+	/// <summary>
+	/// Response describing a pool
+	/// </summary>
+	public class GetPoolResponse
+	{
+		/// <summary>
+		/// Unique id of the pool
+		/// </summary>
+		public string Id { get; set; }
+
+		/// <summary>
+		/// Name of the pool
+		/// </summary>
+		public string Name { get; set; }
+
+		/// <summary>
+		/// Condition for agents to be auto-added to the pool
+		/// </summary>
+		public Condition? Condition { get; set; }
+
+		/// <summary>
+		/// Whether to enable autoscaling for this pool
+		/// </summary>
+		public bool EnableAutoscaling { get; set; }
+		
+		/// <summary>
+		/// Cooldown time between scale-out events in seconds
+		/// </summary>
+		public int? ScaleOutCooldown { get; set; }
+		
+		/// <summary>
+		/// Cooldown time between scale-in events in seconds
+		/// </summary>
+		public int? ScaleInCooldown { get; set; }
+		
+		/// <summary>
+		/// Pool sizing strategy to be used for this pool
+		/// </summary>
+		public PoolSizeStrategy? SizeStrategy { get; set; }
+		
+		/// <summary>
+		/// Settings for lease utilization pool sizing strategy (if used)
+		/// </summary>
+		public LeaseUtilizationSettings? LeaseUtilizationSettings { get; set; }
+		
+		/// <summary>
+		/// Settings for job queue pool sizing strategy (if used) 
+		/// </summary>
+		public JobQueueSettings? JobQueueSettings { get; set; }
+
+		/// <summary>
+		/// The minimum nunmber of agents to retain in this pool
+		/// </summary>
+		public int? MinAgents { get; set; }
+
+		/// <summary>
+		/// The minimum number of idle agents in this pool, if autoscaling is enabled
+		/// </summary>
+		public int? NumReserveAgents { get; set; }
+
+		/// <summary>
+		/// List of workspaces that this agent contains
+		/// </summary>
+		public List<GetAgentWorkspaceResponse> Workspaces { get; set; }
+
+		/// <summary>
+		/// Arbitrary properties for this pool.
+		/// </summary>
+		public IReadOnlyDictionary<string, string> Properties { get; set; }
+
+		/// <summary>
+		/// Constructor
+		/// </summary>
+		/// <param name="pool">The pool to construct from</param>
+		public GetPoolResponse(IPool pool)
+		{
+			Id = pool.Id.ToString();
+			Name = pool.Name;
+			Condition = pool.Condition;
+			EnableAutoscaling = pool.EnableAutoscaling;
+			ScaleOutCooldown = pool.ScaleOutCooldown == null ? null : (int)pool.ScaleOutCooldown.Value.TotalSeconds;
+			ScaleInCooldown = pool.ScaleInCooldown == null ? null : (int)pool.ScaleInCooldown.Value.TotalSeconds;
+			SizeStrategy = pool.SizeStrategy;
+			LeaseUtilizationSettings = pool.LeaseUtilizationSettings == null ? null : new LeaseUtilizationSettings(pool.LeaseUtilizationSettings);
+			JobQueueSettings = pool.JobQueueSettings == null ? null : new JobQueueSettings(pool.JobQueueSettings);
+			MinAgents = pool.MinAgents;
+			NumReserveAgents = pool.NumReserveAgents;
+			Workspaces = pool.Workspaces.Select(x => new GetAgentWorkspaceResponse(x)).ToList();
+			Properties = pool.Properties;
+		}
+	}
+}
diff --git a/Engine/Source/Programs/Horde/Horde.Build/Agents/Pools/PoolsController.cs b/Engine/Source/Programs/Horde/Horde.Build/Agents/Pools/PoolsController.cs
new file mode 100644
--- /dev/null
+++ b/Engine/Source/Programs/Horde/Horde.Build/Agents/Pools/PoolsController.cs
@@ -0,0 +1,207 @@
+﻿// Copyright Epic Games, Inc. All Rights Reserved.
+
+using System;
+using System.Collections.Generic;
+using System.Threading.Tasks;
+using Horde.Build.Acls;
+using Horde.Build.Api;
+using Horde.Build.Models;
+using Horde.Build.Services;
+using Horde.Build.Utilities;
+using Microsoft.AspNetCore.Authorization;
+using Microsoft.AspNetCore.Mvc;
+
+namespace Horde.Build.Controllers
+{
+	using PoolId = StringId<IPool>;
+
+	/// <summary>
+	/// Controller for the /api/v1/pools endpoint
+	/// </summary>
+	[ApiController]
+	[Authorize]
+	[Route("[controller]")]
+	public class PoolsController : HordeControllerBase
+	{
+		/// <summary>
+		/// Singleton instance of the ACL service
+		/// </summary>
+		private readonly AclService _aclService;
+
+		/// <summary>
+		/// Singleton instance of the pool service
+		/// </summary>
+		private readonly PoolService _poolService;
+
+		/// <summary>
+		/// Constructor
+		/// </summary>
+		/// <param name="aclService">The ACL service</param>
+		/// <param name="poolService">The pool service</param>
+		public PoolsController(AclService aclService, PoolService poolService)
+		{
+			_aclService = aclService;
+			_poolService = poolService;
+		}
+
+		/// <summary>
+		/// Creates a new pool
+		/// </summary>
+		/// <param name="create">Parameters for the new pool.</param>
+		/// <returns>Http result code</returns>
+		[HttpPost]
+		[Route("/api/v1/pools")]
+		public async Task<ActionResult<CreatePoolResponse>> CreatePoolAsync([FromBody] CreatePoolRequest create)
+		{
+			if(!await _aclService.AuthorizeAsync(AclAction.CreatePool, User))
+			{
+				return Forbid(AclAction.CreatePool);
+			}
+
+			Horde.Build.Fleet.Autoscale.LeaseUtilizationSettings? luSettings = create.LeaseUtilizationSettings?.Convert();
+			Horde.Build.Fleet.Autoscale.JobQueueSettings? jqSettings = create.JobQueueSettings?.Convert();
+			
+			TimeSpan? scaleOutCooldown = create.ScaleOutCooldown == null ? null : TimeSpan.FromSeconds(create.ScaleOutCooldown.Value);
+			TimeSpan? scaleInCooldown = create.ScaleInCooldown == null ? null : TimeSpan.FromSeconds(create.ScaleInCooldown.Value);
+
+			IPool newPool = await _poolService.CreatePoolAsync(
+				create.Name, create.Condition, create.EnableAutoscaling, create.MinAgents, create.NumReserveAgents,
+				scaleOutCooldown, scaleInCooldown, create.SizeStrategy, luSettings, jqSettings, create.Properties);
+			return new CreatePoolResponse(newPool.Id.ToString());
+		}
+
+		/// <summary>
+		/// Query all the pools
+		/// </summary>
+		/// <param name="filter">Filter for the properties to return</param>
+		/// <returns>Information about all the pools</returns>
+		[HttpGet]
+		[Route("/api/v1/pools")]
+		[ProducesResponseType(typeof(List<GetPoolResponse>), 200)]
+		public async Task<ActionResult<List<object>>> GetPoolsAsync([FromQuery] PropertyFilter? filter = null)
+		{
+			if (!await _aclService.AuthorizeAsync(AclAction.ListPools, User))
+			{
+				return Forbid(AclAction.ListPools);
+			}
+
+			List<IPool> pools = await _poolService.GetPoolsAsync();
+
+			List<object> responses = new List<object>();
+			foreach (IPool pool in pools)
+			{
+				responses.Add(new GetPoolResponse(pool).ApplyFilter(filter));
+			}
+			return responses;
+		}
+
+		/// <summary>
+		/// Retrieve information about a specific pool
+		/// </summary>
+		/// <param name="poolId">Id of the pool to get information about</param>
+		/// <param name="filter">Filter to apply to the returned properties</param>
+		/// <returns>Information about the requested pool</returns>
+		[HttpGet]
+		[Route("/api/v1/pools/{poolId}")]
+		[ProducesResponseType(typeof(GetPoolResponse), 200)]
+		public async Task<ActionResult<object>> GetPoolAsync(string poolId, [FromQuery] PropertyFilter? filter = null)
+		{
+			if (!await _aclService.AuthorizeAsync(AclAction.ViewPool, User))
+			{
+				return Forbid(AclAction.ViewPool);
+			}
+
+			PoolId poolIdValue = new PoolId(poolId);
+
+			IPool? pool = await _poolService.GetPoolAsync(poolIdValue);
+			if (pool == null)
+			{
+				return NotFound(poolIdValue);
+			}
+
+			return new GetPoolResponse(pool).ApplyFilter(filter);
+		}
+
+		/// <summary>
+		/// Update a pool's properties.
+		/// </summary>
+		/// <param name="poolId">Id of the pool to update</param>
+		/// <param name="update">Items on the pool to update</param>
+		/// <returns>Http result code</returns>
+		[HttpPut]
+		[Route("/api/v1/pools/{poolId}")]
+		public async Task<ActionResult> UpdatePoolAsync(string poolId, [FromBody] UpdatePoolRequest update)
+		{
+			if (!await _aclService.AuthorizeAsync(AclAction.UpdatePool, User))
+			{
+				return Forbid(AclAction.UpdatePool);
+			}
+
+			PoolId poolIdValue = new PoolId(poolId);
+			
+			IPool? pool = await _poolService.GetPoolAsync(poolIdValue);
+			if(pool == null)
+			{
+				return NotFound(poolIdValue);
+			}
+			
+			TimeSpan? scaleOutCooldown = update.ScaleOutCooldown == null ? null : TimeSpan.FromSeconds(update.ScaleOutCooldown.Value);
+			TimeSpan? scaleInCooldown = update.ScaleInCooldown == null ? null : TimeSpan.FromSeconds(update.ScaleInCooldown.Value);
+
+			await _poolService.UpdatePoolAsync(pool, update.Name, update.Condition, update.EnableAutoscaling,
+				update.MinAgents, update.NumReserveAgents, update.Properties, scaleOutCooldown, scaleInCooldown, update.SizeStrategy);
+			return new OkResult();
+		}
+
+		/// <summary>
+		/// Delete a pool
+		/// </summary>
+		/// <param name="poolId">Id of the pool to delete</param>
+		/// <returns>Http result code</returns>
+		[HttpDelete]
+		[Route("/api/v1/pools/{poolId}")]
+		public async Task<ActionResult> DeletePoolAsync(string poolId)
+		{
+			if (!await _aclService.AuthorizeAsync(AclAction.DeletePool, User))
+			{
+				return Forbid(AclAction.DeletePool);
+			}
+
+			PoolId poolIdValue = new PoolId(poolId);
+			if(!await _poolService.DeletePoolAsync(poolIdValue))
+			{
+				return NotFound(poolIdValue);
+			}
+			return Ok();
+		}
+
+		/// <summary>
+		/// Batch update pool properties
+		/// </summary>
+		/// <param name="batchUpdates">List of pools to update</param>
+		/// <returns>Http result code</returns>
+		[HttpPut]
+		[Route("/api/v1/pools")]
+		public async Task<ActionResult> UpdatePoolAsync([FromBody] List<BatchUpdatePoolRequest> batchUpdates)
+		{
+			if (!await _aclService.AuthorizeAsync(AclAction.UpdatePool, User))
+			{
+				return Forbid(AclAction.UpdatePool);
+			}
+
+			foreach (BatchUpdatePoolRequest update in batchUpdates)
+			{
+				PoolId poolIdValue = new PoolId(update.Id);
+
+				IPool? pool = await _poolService.GetPoolAsync(poolIdValue);
+				if (pool == null)
+				{
+					return NotFound(poolIdValue);
+				}
+
+				await _poolService.UpdatePoolAsync(pool, update.Name, newProperties: update.Properties);
+			}
+			return Ok();
+		}
+	}
+}
diff --git a/Engine/Source/Programs/Horde/Horde.Build/Agents/Sessions/ISessionCollection.cs b/Engine/Source/Programs/Horde/Horde.Build/Agents/Sessions/ISessionCollection.cs
new file mode 100644
--- /dev/null
+++ b/Engine/Source/Programs/Horde/Horde.Build/Agents/Sessions/ISessionCollection.cs
@@ -0,0 +1,72 @@
+// Copyright Epic Games, Inc. All Rights Reserved.
+
+using System;
+using System.Collections.Generic;
+using System.Threading.Tasks;
+using Horde.Build.Models;
+using Horde.Build.Utilities;
+
+namespace Horde.Build.Collections
+{
+	using SessionId = ObjectId<ISession>;
+
+	/// <summary>
+	/// Interface for a collection of session documents
+	/// </summary>
+	public interface ISessionCollection
+	{
+		/// <summary>
+		/// Adds a new session to the collection
+		/// </summary>
+		/// <param name="id">The session id</param>
+		/// <param name="agentId">The agent this session is for</param>
+		/// <param name="startTime">Start time of this session</param>
+		/// <param name="properties">Properties of this agent at the time the session started</param>
+		/// <param name="resources">Resources which the agent has</param>
+		/// <param name="version">Version of the agent software</param>
+		Task<ISession> AddAsync(SessionId id, AgentId agentId, DateTime startTime, IReadOnlyList<string>? properties, IReadOnlyDictionary<string, int>? resources, string? version);
+
+		/// <summary>
+		/// Gets information about a particular session
+		/// </summary>
+		/// <param name="sessionId">The unique session id</param>
+		/// <returns>The session information</returns>
+		Task<ISession?> GetAsync(SessionId sessionId);
+
+		/// <summary>
+		/// Find sessions for the given agent
+		/// </summary>
+		/// <param name="agentId">The unique agent id</param>
+		/// <param name="startTime">Start time to include in the search</param>
+		/// <param name="finishTime">Finish time to include in the search</param>
+		/// <param name="index">Index of the first result to return</param>
+		/// <param name="count">Number of results to return</param>
+		/// <returns>List of sessions matching the given criteria</returns>
+		Task<List<ISession>> FindAsync(AgentId agentId, DateTime? startTime, DateTime? finishTime, int index, int count);
+
+		/// <summary>
+		/// Finds any active sessions
+		/// </summary>
+		/// <param name="index">Index of the first result to return</param>
+		/// <param name="count">Number of results to return</param>
+		/// <returns>List of sessions</returns>
+		Task<List<ISession>> FindActiveSessionsAsync(int? index = null, int? count = null);
+
+		/// <summary>
+		/// Update a session from the collection
+		/// </summary>
+		/// <param name="sessionId">The session to update</param>
+		/// <param name="finishTime">Time at which the session finished</param>
+		/// <param name="properties">The agent properties</param>
+		/// <param name="resources">Resources which the agent has</param>
+		/// <returns>Async task</returns>
+		Task UpdateAsync(SessionId sessionId, DateTime finishTime, IReadOnlyList<string> properties, IReadOnlyDictionary<string, int> resources);
+
+		/// <summary>
+		/// Delete a session from the collection
+		/// </summary>
+		/// <param name="sessionId">The session id</param>
+		/// <returns>Async task</returns>
+		Task DeleteAsync(SessionId sessionId);
+	}
+}
diff --git a/Engine/Source/Programs/Horde/Horde.Build/Agents/Sessions/SessionCollection.cs b/Engine/Source/Programs/Horde/Horde.Build/Agents/Sessions/SessionCollection.cs
new file mode 100644
--- /dev/null
+++ b/Engine/Source/Programs/Horde/Horde.Build/Agents/Sessions/SessionCollection.cs
@@ -0,0 +1,139 @@
+// Copyright Epic Games, Inc. All Rights Reserved.
+
+using System;
+using System.Collections.Generic;
+using System.Linq;
+using System.Threading.Tasks;
+using Horde.Build.Models;
+using Horde.Build.Server;
+using Horde.Build.Utilities;
+using MongoDB.Bson.Serialization.Attributes;
+using MongoDB.Driver;
+
+namespace Horde.Build.Collections.Impl
+{
+	using SessionId = ObjectId<ISession>;
+
+	/// <summary>
+	/// Collection of session documents
+	/// </summary>
+	public class SessionCollection : ISessionCollection
+	{
+		/// <summary>
+		/// Concrete implementation of ISession
+		/// </summary>
+		class SessionDocument : ISession
+		{
+			[BsonRequired, BsonId]
+			public SessionId Id { get; set; }
+
+			[BsonRequired]
+			public AgentId AgentId { get; set; }
+
+			public DateTime StartTime { get; set; }
+			public DateTime? FinishTime { get; set; }
+			public List<string>? Properties { get; set; }
+			public Dictionary<string, int>? Resources { get; set; }
+			public string? Version { get; set; }
+
+			IReadOnlyList<string>? ISession.Properties => Properties;
+
+			[BsonConstructor]
+			private SessionDocument()
+			{
+			}
+
+			public SessionDocument(SessionId id, AgentId agentId, DateTime startTime, IReadOnlyList<string>? properties, IReadOnlyDictionary<string, int>? resources, string? version)
+			{
+				Id = id;
+				AgentId = agentId;
+				StartTime = startTime;
+				if (properties != null)
+				{
+					Properties = new List<string>(properties);
+				}
+				if (resources != null)
+				{
+					Resources = new Dictionary<string, int>(resources);
+				}
+				Version = version;
+			}
+		}
+
+		/// <summary>
+		/// Collection of session documents
+		/// </summary>
+		readonly IMongoCollection<SessionDocument> _sessions;
+
+		/// <summary>
+		/// Constructor
+		/// </summary>
+		/// <param name="mongoService">The database service</param>
+		public SessionCollection(MongoService mongoService)
+		{
+			List<MongoIndex<SessionDocument>> indexes = new List<MongoIndex<SessionDocument>>();
+			indexes.Add(keys => keys.Ascending(x => x.AgentId));
+			indexes.Add(keys => keys.Ascending(x => x.StartTime));
+			indexes.Add(keys => keys.Ascending(x => x.FinishTime));
+
+			_sessions = mongoService.GetCollection<SessionDocument>("Sessions", indexes);
+		}
+
+		/// <inheritdoc/>
+		public async Task<ISession> AddAsync(SessionId id, AgentId agentId, DateTime startTime, IReadOnlyList<string>? properties, IReadOnlyDictionary<string, int>? resources, string? version)
+		{
+			SessionDocument newSession = new SessionDocument(id, agentId, startTime, properties, resources, version);
+			await _sessions.InsertOneAsync(newSession);
+			return newSession;
+		}
+
+		/// <inheritdoc/>
+		public async Task<ISession?> GetAsync(SessionId sessionId)
+		{
+			return await _sessions.Find(x => x.Id == sessionId).FirstOrDefaultAsync();
+		}
+
+		/// <inheritdoc/>
+		public async Task<List<ISession>> FindAsync(AgentId agentId, DateTime? startTime, DateTime? finishTime, int index, int count)
+		{
+			FilterDefinitionBuilder<SessionDocument> filterBuilder = Builders<SessionDocument>.Filter;
+
+			FilterDefinition<SessionDocument> filter = filterBuilder.Eq(x => x.AgentId, agentId);
+			if (startTime != null)
+			{
+				filter &= filterBuilder.Gte(x => x.StartTime, startTime.Value);
+			}
+			if (finishTime != null)
+			{
+				filter &= filterBuilder.Or(filterBuilder.Eq(x => x.FinishTime, null), filterBuilder.Lte(x => x.FinishTime, finishTime.Value));
+			}
+
+			List<SessionDocument> results = await _sessions.Find(filter).SortByDescending(x => x.StartTime).Skip(index).Limit(count).ToListAsync();
+			return results.ConvertAll<ISession>(x => x);
+		}
+
+		/// <inheritdoc/>
+		public async Task<List<ISession>> FindActiveSessionsAsync(int? index, int? count)
+		{
+			List<SessionDocument> results = await _sessions.Find(x => x.FinishTime == null).Range(index, count).ToListAsync();
+			return results.ConvertAll<ISession>(x => x);
+		}
+
+		/// <inheritdoc/>
+		public Task UpdateAsync(SessionId sessionId, DateTime finishTime, IReadOnlyList<string> properties, IReadOnlyDictionary<string, int> resources)
+		{
+			UpdateDefinition<SessionDocument> update = Builders<SessionDocument>.Update
+				.Set(x => x.FinishTime, finishTime)
+				.Set(x => x.Properties, new List<string>(properties))
+				.Set(x => x.Resources, new Dictionary<string, int>(resources));
+
+			return _sessions.FindOneAndUpdateAsync(x => x.Id == sessionId, update);
+		}
+
+		/// <inheritdoc/>
+		public Task DeleteAsync(SessionId sessionId)
+		{
+			return _sessions.DeleteOneAsync(x => x.Id == sessionId);
+		}
+	}
+}
diff --git a/Engine/Source/Programs/Horde/Horde.Build/Agents/Software/AgentSoftwareCollection.cs b/Engine/Source/Programs/Horde/Horde.Build/Agents/Software/AgentSoftwareCollection.cs
new file mode 100644
--- /dev/null
+++ b/Engine/Source/Programs/Horde/Horde.Build/Agents/Software/AgentSoftwareCollection.cs
@@ -0,0 +1,197 @@
+// Copyright Epic Games, Inc. All Rights Reserved.
+
+using System.Collections.Generic;
+using System.IO;
+using System.Linq;
+using System.Threading.Tasks;
+using Horde.Build.Server;
+using Horde.Build.Utilities;
+using MongoDB.Bson;
+using MongoDB.Bson.Serialization.Attributes;
+using MongoDB.Driver;
+
+namespace Horde.Build.Collections.Impl
+{
+	/// <summary>
+	/// Collection of agent software documents
+	/// </summary>
+	public class AgentSoftwareCollection : IAgentSoftwareCollection
+	{
+		/// <summary>
+		/// Information about agent software
+		/// </summary>
+		class AgentSoftwareDocument
+		{
+			[BsonId]
+			public string Id { get; set; }
+
+			[BsonIgnoreIfNull]
+			public byte[]? Data { get; set; }
+
+			[BsonIgnoreIfNull]
+            public List<string>? ChunkIds { get; set; }
+
+            [BsonConstructor]
+			private AgentSoftwareDocument()
+			{
+				Id = null!;
+				Data = null;
+                ChunkIds = null;			
+            }
+
+			public AgentSoftwareDocument(string id, byte[]? data, List<string>? chunkIds = null)
+			{
+				Id = id;
+				Data = data;
+                ChunkIds = chunkIds;
+            }
+		}
+
+		class ChunkDocument
+		{
+			[BsonId]
+			public string Id { get; set; }
+
+			public byte[] Data { get; set; } = null!;
+
+			public string Version { get; set; }
+
+            [BsonConstructor]
+			private ChunkDocument()
+			{
+				Id = null!;
+				Data = null!;
+                Version = null!;
+            }
+
+			public ChunkDocument(string id, string version, byte[] data)
+			{
+				Id = id;
+				Data = data;
+                Version = version;
+            }
+		}
+
+		// MongoDB document size limit is 16 megs, we chunk at this with a 1k buffer for anything else in document
+		readonly int _chunkSize = 16 * 1024 * 1024 - 1024;
+
+		/// <summary>
+		/// Template documents
+		/// </summary>
+		readonly IMongoCollection<AgentSoftwareDocument> _collection;
+
+		/// <summary>
+		/// Agent template chunks
+		/// </summary>
+		readonly IMongoCollection<ChunkDocument> _agentChunks;
+
+		/// <summary>
+		/// Constructor
+		/// </summary>
+		/// <param name="mongoService">The database service singleton</param>
+		public AgentSoftwareCollection(MongoService mongoService)
+		{
+			_collection = mongoService.GetCollection<AgentSoftwareDocument>("AgentSoftware");
+			_agentChunks = mongoService.GetCollection<ChunkDocument>("AgentSoftwareChunks");
+		}
+
+		/// <inheritdoc/>
+		public async Task<bool> AddAsync(string version, byte[] data)
+		{
+
+			if (await ExistsAsync(version))
+			{
+                return false;
+            }
+
+            AgentSoftwareDocument document;
+
+            // Check if agent will fit in document, otherwise we need to chunk
+            if (data.Length <= _chunkSize)
+			{
+				document = new AgentSoftwareDocument(version, data);
+				await _collection.InsertOneAsync(document);
+                return true;
+            }
+
+			int totalRead = 0;            
+            byte[] buffer = new byte[_chunkSize];            
+			List<ChunkDocument> chunks = new List<ChunkDocument>();
+            using (MemoryStream chunkStream = new MemoryStream(data))
+            {
+                for (; ; )
+                {
+					int bytesRead = await chunkStream.ReadAsync(buffer, 0, _chunkSize);
+					chunks.Add(new ChunkDocument(ObjectId.GenerateNewId().ToString(), version, buffer.Take(bytesRead).ToArray()));
+					totalRead += bytesRead;
+					if (totalRead == data.Length)
+					{
+						break;
+					}
+                }
+            }
+
+            // insert the chunks
+            await _agentChunks.InsertManyAsync(chunks);
+
+            document = new AgentSoftwareDocument(version, null, chunks.Select( chunk => chunk.Id).ToList());
+			await _collection.InsertOneAsync(document);
+            return true;
+        }
+
+		/// <inheritdoc/>
+		public async Task<bool> ExistsAsync(string version)
+		{
+			return await _collection.Find(x => x.Id == version).AnyAsync();
+		}
+
+		/// <inheritdoc/>
+		public async Task<bool> RemoveAsync(string version)
+		{
+			// Delete any chunks
+			await _agentChunks.DeleteManyAsync(x => x.Version == version);
+
+			DeleteResult result = await _collection.DeleteOneAsync(x => x.Id == version);			
+			return result.DeletedCount > 0;
+		}
+
+		/// <inheritdoc/>
+		public async Task<byte[]?> GetAsync(string version)
+		{			
+			IFindFluent<AgentSoftwareDocument, AgentSoftwareDocument> query = _collection.Find(x => x.Id == version);
+			AgentSoftwareDocument? software = await query.FirstOrDefaultAsync();
+
+			if (software == null)
+			{
+                return null;
+            }
+
+			if (software.Data != null && software.Data.Length > 0)
+			{
+                return software.Data;
+            }
+
+			if (software.ChunkIds != null)
+			{
+				FilterDefinitionBuilder<ChunkDocument> filterBuilder = Builders<ChunkDocument>.Filter;
+				FilterDefinition<ChunkDocument> filter = filterBuilder.Empty;
+				filter &= filterBuilder.In(x => x.Id, software.ChunkIds);
+
+				IFindFluent<ChunkDocument, ChunkDocument> chunkQuery = _agentChunks.Find(x => software.ChunkIds.Contains(x.Id));				
+				List<ChunkDocument> chunks = await chunkQuery.ToListAsync();
+
+				using (MemoryStream dataStream = new MemoryStream())
+				{
+					foreach (string chunkId in software.ChunkIds)
+					{
+                        dataStream.Write(chunks.First(x => x.Id == chunkId).Data);
+                    }
+
+                    return dataStream.ToArray();
+                }
+			}
+
+            return null;
+        }
+	}
+}
diff --git a/Engine/Source/Programs/Horde/Horde.Build/Agents/Software/AgentSoftwareController.Messages.cs b/Engine/Source/Programs/Horde/Horde.Build/Agents/Software/AgentSoftwareController.Messages.cs
new file mode 100644
--- /dev/null
+++ b/Engine/Source/Programs/Horde/Horde.Build/Agents/Software/AgentSoftwareController.Messages.cs
@@ -0,0 +1,45 @@
+// Copyright Epic Games, Inc. All Rights Reserved.
+
+using System;
+using Horde.Build.Models;
+
+namespace Horde.Build.Api
+{
+	/// <summary>
+	/// Information about an agent software channel
+	/// </summary>
+	public class GetAgentSoftwareChannelResponse
+	{
+		/// <summary>
+		/// Name of the channel
+		/// </summary>
+		public string? Name { get; set; }
+
+		/// <summary>
+		/// Name of the user that last modified this channel
+		/// </summary>
+		public string? ModifiedBy { get; set; }
+
+		/// <summary>
+		/// The modified timestamp
+		/// </summary>
+		public DateTime ModifiedTime { get; set; }
+
+		/// <summary>
+		/// Version number of this software
+		/// </summary>
+		public string? Version { get; set; }
+
+		/// <summary>
+		/// Constructor
+		/// </summary>
+		/// <param name="channel">The channel information</param>
+		public GetAgentSoftwareChannelResponse(IAgentSoftwareChannel channel)
+		{
+			Name = channel.Name.ToString();
+			ModifiedBy = channel.ModifiedBy;
+			ModifiedTime = channel.ModifiedTime;
+			Version = channel.Version.ToString();
+		}
+	}
+}
diff --git a/Engine/Source/Programs/Horde/Horde.Build/Agents/Software/AgentSoftwareController.cs b/Engine/Source/Programs/Horde/Horde.Build/Agents/Software/AgentSoftwareController.cs
new file mode 100644
--- /dev/null
+++ b/Engine/Source/Programs/Horde/Horde.Build/Agents/Software/AgentSoftwareController.cs
@@ -0,0 +1,152 @@
+﻿// Copyright Epic Games, Inc. All Rights Reserved.
+
+using System.Collections.Generic;
+using System.IO;
+using System.Threading.Tasks;
+using Horde.Build.Acls;
+using Horde.Build.Api;
+using Horde.Build.Models;
+using Horde.Build.Services;
+using Horde.Build.Utilities;
+using Microsoft.AspNetCore.Authorization;
+using Microsoft.AspNetCore.Http;
+using Microsoft.AspNetCore.Mvc;
+using Microsoft.Net.Http.Headers;
+
+namespace Horde.Build.Controllers
+{
+	using AgentSoftwareChannelName = StringId<AgentSoftwareChannels>;
+
+	/// <summary>
+	/// Controller for the /api/v1/software endpoint
+	/// </summary>
+	[ApiController]
+	[Authorize]
+	[Route("[controller]")]
+	public class AgentSoftwareController : ControllerBase
+	{
+		/// <summary>
+		/// Singleton instance of the ACL service
+		/// </summary>
+		private readonly AclService _aclService;
+
+		/// <summary>
+		/// Singleton instance of the client service
+		/// </summary>
+		private readonly AgentSoftwareService _agentSoftwareService;
+
+		/// <summary>
+		/// Constructor
+		/// </summary>
+		/// <param name="aclService">The ACL service</param>
+		/// <param name="agentSoftwareService">The client service</param>
+		public AgentSoftwareController(AclService aclService, AgentSoftwareService agentSoftwareService)
+		{
+			_aclService = aclService;
+			_agentSoftwareService = agentSoftwareService;
+		}
+
+		/// <summary>
+		/// Finds all uploaded software matching the given criteria
+		/// </summary>
+		/// <param name="filter">Filter for the properties to return</param>
+		/// <returns>Http response</returns>
+		[HttpGet]
+		[Route("/api/v1/agentsoftware")]
+		[ProducesResponseType(typeof(List<GetAgentSoftwareChannelResponse>), 200)]
+		public async Task<ActionResult<List<object>>> FindSoftwareAsync([FromQuery] PropertyFilter? filter = null)
+		{
+			if (!await _aclService.AuthorizeAsync(AclAction.DownloadSoftware, User))
+			{
+				return Forbid();
+			}
+
+			List<IAgentSoftwareChannel> results = await _agentSoftwareService.FindChannelsAsync();
+
+			List<object> responses = new List<object>();
+			foreach (IAgentSoftwareChannel result in results)
+			{
+				responses.Add(new GetAgentSoftwareChannelResponse(result).ApplyFilter(filter));
+			}
+			return responses;
+		}
+
+		/// <summary>
+		/// Finds all uploaded software matching the given criteria
+		/// </summary>
+		/// <param name="name">Name of the channel to get</param>
+		/// <param name="filter">Filter for the properties to return</param>
+		/// <returns>Http response</returns>
+		[HttpGet]
+		[Route("/api/v1/agentsoftware/{name}")]
+		[ProducesResponseType(typeof(GetAgentSoftwareChannelResponse), 200)]
+		public async Task<ActionResult<object>> FindSoftwareAsync(string name, [FromQuery] PropertyFilter? filter = null)
+		{
+			if (!await _aclService.AuthorizeAsync(AclAction.DownloadSoftware, User))
+			{
+				return Forbid();
+			}
+
+			IAgentSoftwareChannel? channel = await _agentSoftwareService.GetChannelAsync(new AgentSoftwareChannelName(name));
+			if(channel == null)
+			{
+				return NotFound();
+			}
+
+			return new GetAgentSoftwareChannelResponse(channel).ApplyFilter(filter);
+		}
+
+		/// <summary>
+		/// Uploads a new agent zip file
+		/// </summary>
+		/// <param name="name">Name of the channel to post to</param>
+		/// <param name="file">Zip archive containing the new client software</param>
+		/// <returns>Http result code</returns>
+		[HttpPost]
+		[Route("/api/v1/agentsoftware/{name}/zip")]
+		public async Task<ActionResult> SetArchiveAsync(string name, [FromForm] IFormFile file)
+		{
+			if (!await _aclService.AuthorizeAsync(AclAction.UploadSoftware, User))
+			{
+				return Forbid();
+			}
+
+			byte[] data;
+			using (MemoryStream memoryStream = new MemoryStream())
+			{
+				using (System.IO.Stream stream = file.OpenReadStream())
+				{
+					await stream.CopyToAsync(memoryStream);
+				}
+				data = memoryStream.ToArray();
+			}
+
+			await _agentSoftwareService.SetArchiveAsync(new AgentSoftwareChannelName(name), User.Identity?.Name, data);
+			return Ok();
+		}
+
+		/// <summary>
+		/// Gets the zip file for a specific channel
+		/// </summary>
+		/// <param name="name">Name of the channel</param>
+		/// <returns>Http response</returns>
+		[HttpGet]
+		[Route("/api/v1/agentsoftware/{name}/zip")]
+		public async Task<ActionResult> GetArchiveAsync(string name)
+		{
+			if (!await _aclService.AuthorizeAsync(AclAction.DownloadSoftware, User))
+			{
+				return Forbid();
+			}
+
+			byte[]? data = await _agentSoftwareService.GetArchiveAsync(new AgentSoftwareChannelName(name));
+			if (data == null)
+			{
+				return NotFound();
+			}
+
+			return new FileStreamResult(new MemoryStream(data), new MediaTypeHeaderValue("application/octet-stream")) { FileDownloadName = $"HordeAgent.zip" };
+			
+		}
+	}
+}
diff --git a/Engine/Source/Programs/Horde/Horde.Build/Agents/Software/AgentSoftwareService.cs b/Engine/Source/Programs/Horde/Horde.Build/Agents/Software/AgentSoftwareService.cs
new file mode 100644
--- /dev/null
+++ b/Engine/Source/Programs/Horde/Horde.Build/Agents/Software/AgentSoftwareService.cs
@@ -0,0 +1,341 @@
+// Copyright Epic Games, Inc. All Rights Reserved.
+
+using System;
+using System.Collections.Generic;
+using System.IO;
+using System.Linq;
+using System.Security.Cryptography.X509Certificates;
+using System.Threading.Tasks;
+using EpicGames.Core;
+using Horde.Build.Collections;
+using Horde.Build.Commands;
+using Horde.Build.Models;
+using Horde.Build.Utilities;
+using Microsoft.Extensions.Logging;
+using Microsoft.Extensions.Options;
+using MongoDB.Bson.Serialization.Attributes;
+using MongoDB.Driver;
+
+namespace Horde.Build.Services
+{
+	using AgentSoftwareChannelName = StringId<AgentSoftwareChannels>;
+
+	/// <summary>
+	/// Information about a channel
+	/// </summary>
+	public class AgentSoftwareChannel : IAgentSoftwareChannel
+	{
+		/// <summary>
+		/// The channel id
+		/// </summary>
+		public AgentSoftwareChannelName Name { get; set; }
+
+		/// <summary>
+		/// Name of the user that made the last modification
+		/// </summary>
+		public string? ModifiedBy { get; set; }
+
+		/// <summary>
+		/// Last modification time
+		/// </summary>
+		public DateTime ModifiedTime { get; set; }
+
+		/// <summary>
+		/// The software revision number
+		/// </summary>
+		public string Version { get; set; } = String.Empty;
+
+		/// <summary>
+		/// Private constructor for serialization
+		/// </summary>
+		[BsonConstructor]
+		private AgentSoftwareChannel()
+		{
+		}
+
+		/// <summary>
+		/// Constructor
+		/// </summary>
+		/// <param name="name">Name of the channel</param>
+		public AgentSoftwareChannel(AgentSoftwareChannelName name)
+		{
+			Name = name;
+		}
+	}
+
+	/// <summary>
+	/// Singleton document used to track different versions of the agent software
+	/// </summary>
+	[SingletonDocument("5f455039d97900f2b6c735a9")]
+	public class AgentSoftwareChannels : SingletonBase
+	{
+		/// <summary>
+		/// The next version number
+		/// </summary>
+		public int NextVersion { get; set; } = 1;
+
+		/// <summary>
+		/// List of channels
+		/// </summary>
+		public List<AgentSoftwareChannel> Channels { get; set; } = new List<AgentSoftwareChannel>();
+
+		/// <summary>
+		/// Finds an existing channel by the given name, or adds a new one
+		/// </summary>
+		/// <param name="name"></param>
+		/// <returns></returns>
+		public AgentSoftwareChannel FindOrAddChannel(AgentSoftwareChannelName name)
+		{
+			AgentSoftwareChannel? channel = Channels.FirstOrDefault(x => x.Name == name);
+			if (channel == null)
+			{
+				channel = new AgentSoftwareChannel(name);
+				Channels.Add(channel);
+			}
+			return channel;
+		}
+	}
+
+	/// <summary>
+	/// Wrapper for a collection of software revisions. 
+	/// </summary>
+	public class AgentSoftwareService
+	{
+		/// <summary>
+		/// Name of the default channel
+		/// </summary>
+		public static AgentSoftwareChannelName DefaultChannelName { get; } = new AgentSoftwareChannelName("default");
+
+		/// <summary>
+		/// Collection of software documents
+		/// </summary>
+		readonly IAgentSoftwareCollection _collection;
+
+		/// <summary>
+		/// Channels singleton
+		/// </summary>
+		readonly ISingletonDocument<AgentSoftwareChannels> _singleton;
+
+		/// <summary>
+		/// Cached copy of the channels singleton
+		/// </summary>
+		readonly LazyCachedValue<Task<AgentSoftwareChannels>> _channelsDocument;
+
+		/// <summary>
+		/// The server settings
+		/// </summary>
+		readonly IOptionsMonitor<ServerSettings> _settings;
+
+		/// <summary>
+		///  Logger for controller
+		/// </summary>
+		private readonly ILogger<AgentSoftwareService> _logger;
+
+		/// <summary>
+		/// Constructor
+		/// </summary>
+		/// <param name="collection">The software collection</param>
+		/// <param name="singleton">The channels singleton</param>		
+		/// <param name="settings">The settings monitor</param>		
+		/// <param name="logger">The logger instance</param>
+		public AgentSoftwareService(IAgentSoftwareCollection collection, ISingletonDocument<AgentSoftwareChannels> singleton, IOptionsMonitor<ServerSettings> settings, ILogger<AgentSoftwareService> logger)
+		{
+			_collection = collection;
+			_singleton = singleton;
+			_settings = settings;
+			_logger = logger;
+
+			_channelsDocument = new LazyCachedValue<Task<AgentSoftwareChannels>>(() => singleton.GetAsync(), TimeSpan.FromSeconds(20.0));
+
+			Task.Run(() => RegisterDefaultAgent(5000));
+		}
+
+		/// <summary>
+		/// Finds all software archives matching a set of criteria
+		/// </summary>
+		public async Task<List<IAgentSoftwareChannel>> FindChannelsAsync()
+		{
+			AgentSoftwareChannels current = await _channelsDocument.GetLatest();
+			return current.Channels.ConvertAll<IAgentSoftwareChannel>(x => x);
+		}
+
+		/// <summary>
+		/// Gets a single software archive
+		/// </summary>
+		/// <param name="name">Unique id of the channel to find software for</param>
+		public async Task<IAgentSoftwareChannel?> GetChannelAsync(AgentSoftwareChannelName name)
+		{
+			AgentSoftwareChannels current = await _channelsDocument.GetLatest();
+			return current.Channels.FirstOrDefault(x => x.Name == name);
+		}
+
+		/// <summary>
+		/// Gets a cached channel name
+		/// </summary>
+		/// <param name="name">The cached channel name</param>
+		/// <returns></returns>
+		public async Task<IAgentSoftwareChannel?> GetCachedChannelAsync(AgentSoftwareChannelName name)
+		{
+			AgentSoftwareChannels document = await _channelsDocument.GetCached();
+			return document.Channels.FirstOrDefault(x => x.Name == name);
+		}
+
+		/// <summary>
+		/// Removes a channel
+		/// </summary>
+		/// <param name="name">The channel id</param>
+		/// <returns>Async task</returns>
+		public async Task DeleteChannelAsync(AgentSoftwareChannelName name)
+		{
+			for (; ; )
+			{
+				AgentSoftwareChannels current = await _singleton.GetAsync();
+
+				int channelIdx = current.Channels.FindIndex(x => x.Name == name);
+				if (channelIdx == -1)
+				{
+					break;
+				}
+
+				string version = current.Channels[channelIdx].Version;
+				current.Channels.RemoveAt(channelIdx);
+
+				if (await _singleton.TryUpdateAsync(current))
+				{
+					if (!current.Channels.Any(x => x.Version == version))
+					{
+						await _collection.RemoveAsync(version);
+					}
+					break;
+				}
+			}
+		}
+
+		/// <summary>
+		/// Updates a new software revision
+		/// </summary>
+		/// <param name="name">Name of the channel</param>
+		/// <param name="author">Name of the user uploading this file</param>
+		/// <param name="data">The input data stream. This should be a zip archive containing the HordeAgent executable.</param>
+		/// <returns>Unique id for the file</returns>
+		public async Task<string> SetArchiveAsync(AgentSoftwareChannelName name, string? author, byte[] data)
+		{
+			// Upload the software
+			string version = AgentUtilities.ReadVersion(data);
+			await _collection.AddAsync(version, data);
+
+			// Update the channel
+			for(; ;)
+			{
+				AgentSoftwareChannels instance = await _singleton.GetAsync();
+
+				AgentSoftwareChannel channel = instance.FindOrAddChannel(name);
+				channel.ModifiedBy = author;
+				channel.ModifiedTime = DateTime.UtcNow;
+				channel.Version = version;
+
+				if (await _singleton.TryUpdateAsync(instance))
+				{
+					break;
+				}
+			}
+			return version;
+		}
+
+		/// <summary>
+		/// Gets the zip file for a given channel
+		/// </summary>
+		/// <param name="name">The channel name</param>
+		/// <returns>Data for the given archive</returns>
+		public async Task<byte[]?> GetArchiveAsync(AgentSoftwareChannelName name)
+		{
+			AgentSoftwareChannels instance = await _singleton.GetAsync();
+
+			AgentSoftwareChannel? channel = instance.Channels.FirstOrDefault(x => x.Name == name);
+			if (channel == null)
+			{
+				return null;
+			}
+			else
+			{
+				return await _collection.GetAsync(channel.Version);
+			}
+		}
+
+		/// <summary>
+		/// Gets the zip file with a given version number
+		/// </summary>
+		/// <param name="version">The version</param>
+		/// <returns>Data for the given archive</returns>
+		public Task<byte[]?> GetArchiveAsync(string version)
+		{
+			return _collection.GetAsync(version);
+		}
+		async Task RegisterDefaultAgent(int delayMs)
+		{
+			await Task.Delay(delayMs);
+
+			// Check whether we have an installed agent zip
+			FileReference agentZip = FileReference.Combine(Program.AppDir, "DefaultAgent/Agent.zip");
+			if (!_settings.CurrentValue.SingleInstance || !FileReference.Exists(agentZip))
+			{
+				return;
+			}
+
+			_logger.LogInformation("Checking for default agent software update");
+
+			string agentHash = ContentHash.MD5(agentZip).ToString();
+
+			FileReference agentHashFile = FileReference.Combine(Program.DataDir, "Agent/DefaultAgentHash");
+
+			try
+			{
+
+				if (FileReference.Exists(agentHashFile) && FileReference.ReadAllText(agentHashFile).Trim() == agentHash)
+				{
+					_logger.LogInformation("Default agent software is up to date");
+					return;
+				}
+			}
+			catch (Exception ex)
+			{
+				_logger.LogError(ex, "Error checking default agent software, {Message}", ex.Message);
+			}
+
+            byte[] bytes = await File.ReadAllBytesAsync(agentZip.ToString());
+
+			using X509Certificate2? grpcCertificate = ServerCommand.ReadGrpcCertificate(_settings.CurrentValue);
+
+			if (grpcCertificate == null)
+			{
+				throw new Exception("Unable to register default agent without valid grpc certicate");
+			}
+
+			// construct agent app settings
+			
+			string profileName = "Default";
+
+			string serverUrl = $"https://{System.Net.Dns.GetHostName()}:{_settings.CurrentValue.HttpsPort}";			
+
+			Dictionary<string, object> defaultProfile = new Dictionary<string, object>() { { "Name", profileName }, { "Environment", "prod" }, { "Thumbprint", grpcCertificate.Thumbprint }, { "Url", serverUrl } };
+
+			List<object> serverProfiles = new List<object>() { defaultProfile };
+
+			Dictionary<string, object> agentSettings = new Dictionary<string, object>() { { "ServerProfiles", serverProfiles }, { "Server", profileName }};
+
+			bytes = AgentUtilities.UpdateAppSettings(bytes, agentSettings);
+
+			string version = await SetArchiveAsync(new AgentSoftwareChannelName("default"), null, bytes);
+
+			if (!DirectoryReference.Exists(agentHashFile.Directory))
+			{
+				DirectoryReference.CreateDirectory(agentHashFile.Directory);
+			}
+
+			await FileReference.WriteAllTextAsync(agentHashFile, agentHash);
+
+			_logger.LogInformation("Updated default agent software to {AgentHash}", agentHash);
+
+		}
+	}
+}
diff --git a/Engine/Source/Programs/Horde/Horde.Build/Agents/Software/IAgentSoftwareCollection.cs b/Engine/Source/Programs/Horde/Horde.Build/Agents/Software/IAgentSoftwareCollection.cs
new file mode 100644
--- /dev/null
+++ b/Engine/Source/Programs/Horde/Horde.Build/Agents/Software/IAgentSoftwareCollection.cs
@@ -0,0 +1,40 @@
+// Copyright Epic Games, Inc. All Rights Reserved.
+
+using System.Threading.Tasks;
+
+namespace Horde.Build.Collections
+{
+	/// <summary>
+	/// Interface for a collection of template documents
+	/// </summary>
+	public interface IAgentSoftwareCollection
+	{
+		/// <summary>
+		/// Adds a new software revision
+		/// </summary>
+		/// <param name="version">The version number</param>
+		/// <param name="data">Zip file containing the new software</param>
+		/// <returns>New software instance</returns>
+		Task<bool> AddAsync(string version, byte[] data);
+
+		/// <summary>
+		/// Tests whether a given version exists
+		/// </summary>
+		/// <param name="version">Version of the software</param>
+		/// <returns>True if it exists, false otherwise</returns>
+		Task<bool> ExistsAsync(string version);
+
+		/// <summary>
+		/// Removes a software archive 
+		/// </summary>
+		/// <param name="version">Version of the software to delete</param>
+		Task<bool> RemoveAsync(string version);
+
+		/// <summary>
+		/// Downloads software of a given revision
+		/// </summary>
+		/// <param name="version">Version of the software</param>
+		/// <returns>Data for the given software</returns>
+		Task<byte[]?> GetAsync(string version);
+	}
+}
diff --git a/Engine/Source/Programs/Horde/Horde.Build/Agents/Software/UpgradeService.cs b/Engine/Source/Programs/Horde/Horde.Build/Agents/Software/UpgradeService.cs
new file mode 100644
--- /dev/null
+++ b/Engine/Source/Programs/Horde/Horde.Build/Agents/Software/UpgradeService.cs
@@ -0,0 +1,107 @@
+// Copyright Epic Games, Inc. All Rights Reserved.
+
+using System;
+using System.Threading.Tasks;
+using Horde.Build.Collections;
+using Horde.Build.Collections.Impl;
+using Horde.Build.Models;
+using Horde.Build.Server;
+using Microsoft.Extensions.DependencyInjection;
+using Microsoft.Extensions.Logging;
+
+namespace Horde.Build.Services
+{
+	/// <summary>
+	/// Service to perform schema upgrades on the database
+	/// </summary>
+	public class UpgradeService
+	{
+		/// <summary>
+		/// The current schema version
+		/// </summary>
+		public const int LatestSchemaVersion = 5;
+
+		/// <summary>
+		/// The database service instance
+		/// </summary>
+		readonly MongoService _mongoService;
+
+		/// <summary>
+		/// The DI service provider
+		/// </summary>
+		readonly IServiceProvider _serviceProvider;
+
+		/// <summary>
+		/// Logger instance
+		/// </summary>
+		readonly ILogger<UpgradeService> _logger;
+
+		/// <summary>
+		/// Constructor
+		/// </summary>
+		/// <param name="mongoService">The database service singleton</param>
+		/// <param name="serviceProvider">The DI service provider</param>
+		/// <param name="logger">Logger instance</param>
+		public UpgradeService(MongoService mongoService, IServiceProvider serviceProvider, ILogger<UpgradeService> logger)
+		{
+			_mongoService = mongoService;
+			_serviceProvider = serviceProvider;
+			_logger = logger;
+		}
+
+		/// <summary>
+		/// Update the database schema
+		/// </summary>
+		/// <param name="fromVersion">The version to upgrade from. Automatically detected if this is null.</param>
+		/// <returns>Async task</returns>
+		public async Task UpgradeSchemaAsync(int? fromVersion)
+		{
+			Globals globals = await _mongoService.GetGlobalsAsync();
+			int schemaVersion = fromVersion ?? globals.SchemaVersion ?? 0;
+
+			while (schemaVersion < LatestSchemaVersion)
+			{
+				_logger.LogInformation("Upgrading from schema version {Version}", schemaVersion);
+
+				// Handle the current version
+				//				if (SchemaVersion == 0)
+				//				{
+				//					IssueCollection IssueCollection = new IssueCollection(DatabaseService);
+				//					IJobCollection JobCollection = ServiceProvider.GetService<IJobCollection>();
+				//					await IssueCollection.UpdateSchemaLogIdFields(JobCollection, Logger);
+				//				}
+				//				if(SchemaVersion == 2)
+				//				{
+				//					IssueCollection IssueCollection = new IssueCollection(DatabaseService);
+				//					await IssueCollection.FixResolvedIssuesAsync(Logger);
+				//				}
+
+				if (schemaVersion == 3)
+				{
+					IJobCollection jobCollection = _serviceProvider.GetRequiredService<IJobCollection>();
+					await jobCollection.UpgradeDocumentsAsync();
+				}
+				if (schemaVersion == 4)
+				{
+					UserCollectionV1 userCollectionV1 = new UserCollectionV1(_serviceProvider.GetRequiredService<MongoService>());
+					using UserCollectionV2 userCollectionV2 = new UserCollectionV2(_serviceProvider.GetRequiredService<MongoService>(), _serviceProvider.GetRequiredService<ILogger<UserCollectionV2>>());
+					await userCollectionV2.ResaveDocumentsAsync(userCollectionV1);
+				}
+
+				// Increment the version number
+				schemaVersion++;
+
+				// Try to update the current schema version number
+				while (globals.SchemaVersion == null || globals.SchemaVersion < schemaVersion)
+				{
+					globals.SchemaVersion = schemaVersion;
+					if (await _mongoService.TryUpdateSingletonAsync(globals))
+					{
+						break;
+					}
+					globals = await _mongoService.GetGlobalsAsync();
+				}
+			}
+		}
+	}
+}
diff --git a/Engine/Source/Programs/Horde/Horde.Build/Agents/Telemetry/ITelemetryCollection.cs b/Engine/Source/Programs/Horde/Horde.Build/Agents/Telemetry/ITelemetryCollection.cs
new file mode 100644
--- /dev/null
+++ b/Engine/Source/Programs/Horde/Horde.Build/Agents/Telemetry/ITelemetryCollection.cs
@@ -0,0 +1,36 @@
+// Copyright Epic Games, Inc. All Rights Reserved.
+
+using System;
+using System.Collections.Generic;
+using System.Threading.Tasks;
+using Horde.Build.Models;
+
+namespace Horde.Build.Collections
+{
+	/// <summary>
+	/// Collection of utilization collection
+	/// </summary>
+	public interface ITelemetryCollection
+	{
+		/// <summary>
+		/// Adds entries for the given utilization
+		/// </summary>
+		/// <param name="telemetry">Telemetry data to add</param>
+		/// <returns>Async task</returns>
+		Task AddUtilizationTelemetryAsync(IUtilizationTelemetry telemetry);
+
+		/// <summary>
+		/// Finds utilization data matching the given criteria
+		/// </summary>
+		/// <param name="startTimeUtc">Start time to query utilization for</param>
+		/// <param name="finishTimeUtc">Finish time to query utilization for</param>
+		/// <returns>The utilization data</returns>
+		Task<List<IUtilizationTelemetry>> GetUtilizationTelemetryAsync(DateTime startTimeUtc, DateTime finishTimeUtc);
+
+		/// <summary>
+		/// Finds the latest utilization data
+		/// </summary>
+		/// <returns>The utilization data</returns>
+		Task<IUtilizationTelemetry?> GetLatestUtilizationTelemetryAsync();
+	}
+}
diff --git a/Engine/Source/Programs/Horde/Horde.Build/Agents/Telemetry/Telemetry.cs b/Engine/Source/Programs/Horde/Horde.Build/Agents/Telemetry/Telemetry.cs
new file mode 100644
--- /dev/null
+++ b/Engine/Source/Programs/Horde/Horde.Build/Agents/Telemetry/Telemetry.cs
@@ -0,0 +1,227 @@
+// Copyright Epic Games, Inc. All Rights Reserved.
+
+using System;
+using System.Collections.Generic;
+using Horde.Build.Utilities;
+
+namespace Horde.Build.Models
+{
+	using PoolId = StringId<IPool>;
+	using StreamId = StringId<IStream>;
+
+	/// <summary>
+	/// Information about the utilization of a pool
+	/// </summary>
+	public interface IUtilizationTelemetry
+	{
+		/// <summary>
+		/// Start time for the bucket
+		/// </summary>
+		public DateTime StartTime { get; }
+
+		/// <summary>
+		/// Finish time for this bucket
+		/// </summary>
+		public DateTime FinishTime { get; }
+
+		/// <summary>
+		/// Number of agents in total
+		/// </summary>
+		public int NumAgents { get; }
+
+		/// <summary>
+		/// Breakdown of utilization by pool
+		/// </summary>
+		public IReadOnlyList<IPoolUtilizationTelemetry> Pools { get; }
+
+		/// <summary>
+		/// Amount of time that agents were hibernating
+		/// </summary>
+		public double HibernatingTime { get; }
+
+		/// <summary>
+		/// Total time spent running administrative tasks (conform, etc...)
+		/// </summary>
+		public double AdminTime { get; }
+	}
+
+	/// <summary>
+	/// Concrete implementation of <see cref="IUtilizationTelemetry"/>
+	/// </summary>
+	public sealed class NewUtilizationTelemetry : IUtilizationTelemetry
+	{
+		/// <inheritdoc/>
+		public DateTime StartTime { get; set; }
+
+		/// <inheritdoc/>
+		public DateTime FinishTime { get; set; }
+
+		/// <inheritdoc/>
+		public int NumAgents { get; set; }
+
+		/// <inheritdoc/>
+		List<NewPoolUtilizationTelemetry> Pools { get; set; } = new List<NewPoolUtilizationTelemetry>();
+		IReadOnlyList<IPoolUtilizationTelemetry> IUtilizationTelemetry.Pools => Pools;
+
+		Dictionary<PoolId, NewPoolUtilizationTelemetry> PoolsLookup { get; set; } = new Dictionary<PoolId, NewPoolUtilizationTelemetry>();
+
+		/// <inheritdoc/>
+		public double HibernatingTime { get; set; }
+
+		/// <inheritdoc/>
+		public double AdminTime { get; set; }
+
+		/// <summary>
+		/// Constructor
+		/// </summary>
+		/// <param name="startTime"></param>
+		/// <param name="finishTime"></param>
+		public NewUtilizationTelemetry(DateTime startTime, DateTime finishTime)
+		{
+			StartTime = startTime;
+			FinishTime = finishTime;
+		}
+
+		/// <summary>
+		/// Adds a new pool
+		/// </summary>
+		/// <param name="poolId">The pool id</param>
+		/// <returns>Telemetry for the given pool</returns>
+		public NewPoolUtilizationTelemetry FindOrAddPool(PoolId poolId)
+		{
+			NewPoolUtilizationTelemetry? pool;
+			if (!PoolsLookup.TryGetValue(poolId, out pool))
+			{
+				pool = new NewPoolUtilizationTelemetry(poolId);
+				Pools.Add(pool);
+				PoolsLookup.Add(poolId, pool);
+			}
+			return pool;
+		}
+	}
+
+	/// <summary>
+	/// Information about the utilization of a pool
+	/// </summary>
+	public interface IPoolUtilizationTelemetry
+	{
+		/// <summary>
+		/// Pool containing the work to execute
+		/// </summary>
+		public PoolId PoolId { get; }
+
+		/// <summary>
+		/// Number of agents in this pool
+		/// </summary>
+		public int NumAgents { get; }
+
+		/// <summary>
+		/// The stream executing work. If this is null, the time accounts for the machine executing work in another pool.
+		/// </summary>
+		public IReadOnlyList<IStreamUtilizationTelemetry> Streams { get; }
+
+		/// <summary>
+		/// Amount of time spent running 
+		/// </summary>
+		public double AdminTime { get; }
+
+		/// <summary>
+		/// Amount of time that agents were hibernating
+		/// </summary>
+		public double HibernatingTime { get; }
+
+		/// <summary>
+		/// Amount of time spent by agents in this pool servicing other pools
+		/// </summary>
+		public double OtherTime { get; }
+	}
+
+	/// <summary>
+	/// Concrete implementation of <see cref="IPoolUtilizationTelemetry"/>
+	/// </summary>
+	public sealed class NewPoolUtilizationTelemetry : IPoolUtilizationTelemetry
+	{
+		/// <inheritdoc/>
+		public PoolId PoolId { get; set; }
+
+		/// <inheritdoc/>
+		public int NumAgents { get; set; }
+
+		/// <inheritdoc/>
+		List<NewPoolUtilizationTelemetryStream> Streams { get; set; } = new List<NewPoolUtilizationTelemetryStream>();
+		IReadOnlyList<IStreamUtilizationTelemetry> IPoolUtilizationTelemetry.Streams => Streams;
+
+		Dictionary<StreamId, NewPoolUtilizationTelemetryStream> StreamLookup { get; set; } = new Dictionary<StreamId, NewPoolUtilizationTelemetryStream>();
+
+		/// <inheritdoc/>
+		public double AdminTime { get; set; }
+
+		/// <inheritdoc/>
+		public double HibernatingTime { get; set; }
+
+		/// <inheritdoc/>
+		public double OtherTime { get; set; }
+
+		/// <summary>
+		/// Constructor
+		/// </summary>
+		/// <param name="poolId"></param>
+		public NewPoolUtilizationTelemetry(PoolId poolId)
+		{
+			PoolId = poolId;
+		}
+
+		/// <summary>
+		/// Adds a stream to the object
+		/// </summary>
+		/// <param name="streamId"></param>
+		public NewPoolUtilizationTelemetryStream FindOrAddStream(StreamId streamId)
+		{
+			NewPoolUtilizationTelemetryStream? stream;
+			if (!StreamLookup.TryGetValue(streamId, out stream))
+			{
+				stream = new NewPoolUtilizationTelemetryStream(streamId);
+				Streams.Add(stream);
+				StreamLookup.Add(streamId, stream);
+			}
+			return stream;
+		}
+	}
+
+	/// <summary>
+	/// Utilization of a pool for a particular stream
+	/// </summary>
+	public interface IStreamUtilizationTelemetry
+	{
+		/// <summary>
+		/// The stream id
+		/// </summary>
+		public StreamId StreamId { get; }
+
+		/// <summary>
+		/// Number of machine hours spent executing work for this stream
+		/// </summary>
+		public double Time { get; }
+	}
+
+	/// <summary>
+	/// Concrete implementation of <see cref="IStreamUtilizationTelemetry"/>
+	/// </summary>
+	public sealed class NewPoolUtilizationTelemetryStream : IStreamUtilizationTelemetry
+	{
+		/// <inheritdoc/>
+		public StreamId StreamId { get; set; }
+
+		/// <inheritdoc/>
+		public double Time { get; set; }
+
+		/// <summary>
+		/// Constructor
+		/// </summary>
+		/// <param name="streamId"></param>
+		public NewPoolUtilizationTelemetryStream(StreamId streamId)
+		{
+			StreamId = streamId;
+		}
+	}
+}
diff --git a/Engine/Source/Programs/Horde/Horde.Build/Agents/Telemetry/TelemetryCollection.cs b/Engine/Source/Programs/Horde/Horde.Build/Agents/Telemetry/TelemetryCollection.cs
new file mode 100644
--- /dev/null
+++ b/Engine/Source/Programs/Horde/Horde.Build/Agents/Telemetry/TelemetryCollection.cs
@@ -0,0 +1,132 @@
+// Copyright Epic Games, Inc. All Rights Reserved.
+
+using System;
+using System.Collections.Generic;
+using System.Linq;
+using System.Threading.Tasks;
+using EpicGames.Core;
+using Horde.Build.Models;
+using Horde.Build.Server;
+using Horde.Build.Utilities;
+using MongoDB.Bson.Serialization.Attributes;
+using MongoDB.Driver;
+
+namespace Horde.Build.Collections.Impl
+{
+	using PoolId = StringId<IPool>;
+	using StreamId = StringId<IStream>;
+
+	/// <summary>
+	/// Collection of utilization data
+	/// </summary>
+	public class TelemetryCollection : ITelemetryCollection
+	{
+		class UtilizationDocument : IUtilizationTelemetry
+		{
+			public DateTime StartTime { get; set; }
+			public DateTime FinishTime { get; set; }
+
+			public int NumAgents { get; set; }
+			public List<PoolUtilizationDocument> Pools { get; set; }
+			IReadOnlyList<IPoolUtilizationTelemetry> IUtilizationTelemetry.Pools => Pools;
+
+			public double AdminTime { get; set; }
+			public double HibernatingTime { get; set; }
+			public int UpdateIndex { get; set; }
+
+			[BsonConstructor]
+			private UtilizationDocument()
+			{
+				Pools = new List<PoolUtilizationDocument>();
+			}
+
+			public UtilizationDocument(IUtilizationTelemetry other)
+			{
+				StartTime = other.StartTime;
+				FinishTime = other.FinishTime;
+				NumAgents = other.NumAgents;
+				Pools = other.Pools.ConvertAll(x => new PoolUtilizationDocument(x));
+				HibernatingTime = other.HibernatingTime;
+				AdminTime = other.AdminTime;
+			}
+		}
+
+		class PoolUtilizationDocument : IPoolUtilizationTelemetry
+		{
+			public PoolId PoolId { get; set; }
+			public int NumAgents { get; set; }
+
+			public List<StreamUtilizationDocument> Streams { get; set; } 
+			IReadOnlyList<IStreamUtilizationTelemetry> IPoolUtilizationTelemetry.Streams => Streams;
+
+			public double AdminTime { get; set; }
+			public double HibernatingTime { get; set; }
+			public double OtherTime { get; set; }
+
+			[BsonConstructor]
+			private PoolUtilizationDocument()
+			{
+				Streams = new List<StreamUtilizationDocument>();
+			}
+
+			public PoolUtilizationDocument(IPoolUtilizationTelemetry other)
+			{
+				PoolId = other.PoolId;
+				NumAgents = other.NumAgents;
+				Streams = other.Streams.ConvertAll(x => new StreamUtilizationDocument(x));
+				AdminTime = other.AdminTime;
+				HibernatingTime = other.HibernatingTime;
+				OtherTime = other.OtherTime;
+			}
+		}
+
+		class StreamUtilizationDocument : IStreamUtilizationTelemetry
+		{
+			public StreamId StreamId { get; set; }
+			public double Time { get; set; }
+
+			[BsonConstructor]
+			private StreamUtilizationDocument()
+			{
+			}
+
+			public StreamUtilizationDocument(IStreamUtilizationTelemetry other)
+			{
+				StreamId = other.StreamId;
+				Time = other.Time;
+			}
+		}
+
+		readonly IMongoCollection<UtilizationDocument> _utilization;
+
+		/// <summary>
+		/// Constructor
+		/// </summary>
+		/// <param name="database"></param>
+		public TelemetryCollection(MongoService database)
+		{
+			_utilization = database.GetCollection<UtilizationDocument>("Utilization", keys => keys.Ascending(x => x.FinishTime).Ascending(x => x.StartTime));
+		}
+
+		/// <inheritdoc/>
+		public async Task AddUtilizationTelemetryAsync(IUtilizationTelemetry newTelemetry)
+		{
+			await _utilization.InsertOneAsync(new UtilizationDocument(newTelemetry));
+		}
+
+		/// <inheritdoc/>
+		public async Task<List<IUtilizationTelemetry>> GetUtilizationTelemetryAsync(DateTime startTimeUtc, DateTime finishTimeUtc)
+		{
+			FilterDefinition<UtilizationDocument> filter = Builders<UtilizationDocument>.Filter.Gte(x => x.FinishTime, startTimeUtc) & Builders<UtilizationDocument>.Filter.Lte(x => x.StartTime, finishTimeUtc);
+
+			List<UtilizationDocument> documents = await _utilization.Find(filter).ToListAsync();
+			return documents.ConvertAll<IUtilizationTelemetry>(x => x);
+		}
+
+		/// <inheritdoc/>
+		public async Task<IUtilizationTelemetry?> GetLatestUtilizationTelemetryAsync()
+		{
+			return await _utilization.Find(FilterDefinition<UtilizationDocument>.Empty).SortByDescending(x => x.FinishTime).FirstOrDefaultAsync();
+		}
+	}
+}
diff --git a/Engine/Source/Programs/Horde/Horde.Build/Agents/Telemetry/TelemetryController.Messages.cs b/Engine/Source/Programs/Horde/Horde.Build/Agents/Telemetry/TelemetryController.Messages.cs
new file mode 100644
--- /dev/null
+++ b/Engine/Source/Programs/Horde/Horde.Build/Agents/Telemetry/TelemetryController.Messages.cs
@@ -0,0 +1,137 @@
+// Copyright Epic Games, Inc. All Rights Reserved.
+
+using System;
+using System.Collections.Generic;
+using EpicGames.Core;
+using Horde.Build.Models;
+
+namespace Horde.Build.Api
+{
+
+	/// <summary>
+	/// Represents one stream in one pool in one hour of telemetry
+	/// </summary>
+	public class UtilizationTelemetryStream
+	{
+		/// <summary>
+		/// Stream Id
+		/// </summary>
+		public string StreamId { get; set; }
+
+		/// <summary>
+		/// Total time
+		/// </summary>
+		public double Time { get; set; }
+
+		/// <summary>
+		/// Constructor
+		/// </summary>
+		/// <param name="stream">The stream</param>
+		public UtilizationTelemetryStream(IStreamUtilizationTelemetry stream)
+		{
+			StreamId = stream.StreamId.ToString();
+			Time = stream.Time;
+		}
+	}
+
+	/// <summary>
+	/// Representation of an hour of time
+	/// </summary>
+	public class UtilizationTelemetryPool
+	{
+		/// <summary>
+		/// Pool id
+		/// </summary>
+		public string PoolId { get; set; }
+
+		/// <summary>
+		/// Number of agents in this pool
+		/// </summary>
+		public int NumAgents { get; set; }
+
+		/// <summary>
+		/// Total time spent doing admin work
+		/// </summary>
+		public double AdminTime { get; set; }
+
+		/// <summary>
+		/// Time spent hibernating
+		/// </summary>
+		public double HibernatingTime { get; }
+
+		/// <summary>
+		/// Total time agents in this pool were doing work for other pools
+		/// </summary>
+		public double OtherTime { get; set; }
+
+		/// <summary>
+		/// List of streams
+		/// </summary>
+		public List<UtilizationTelemetryStream> Streams { get; set; }
+
+		/// <summary>
+		/// Constructor
+		/// </summary>
+		/// <param name="pool"></param>
+		public UtilizationTelemetryPool(IPoolUtilizationTelemetry pool)
+		{
+			PoolId = pool.PoolId.ToString();
+			NumAgents = pool.NumAgents;
+			AdminTime = pool.AdminTime;
+			HibernatingTime = pool.HibernatingTime;
+			OtherTime = pool.OtherTime;
+
+			Streams = pool.Streams.ConvertAll(stream => new UtilizationTelemetryStream(stream));
+		}
+	}
+
+	/// <summary>
+	/// Response data for a utilization request
+	/// </summary>
+	public class UtilizationTelemetryResponse
+	{
+		/// <summary>
+		/// Start hour
+		/// </summary>
+		public DateTimeOffset StartTime { get; set; }
+
+		/// <summary>
+		/// End hour
+		/// </summary>
+		public DateTimeOffset FinishTime { get; set; }
+
+		/// <summary>
+		/// List of pools
+		/// </summary>
+		public List<UtilizationTelemetryPool> Pools { get; set; }
+
+		/// <summary>
+		/// Total admin time
+		/// </summary>
+		public double AdminTime { get; set; }
+
+		/// <summary>
+		/// Total hibernating time
+		/// </summary>
+		public double HibernatingTime { get; set; }
+
+		/// <summary>
+		/// Total agents
+		/// </summary>
+		public int NumAgents { get; set; }
+
+		/// <summary>
+		/// Constructor
+		/// </summary>
+		public UtilizationTelemetryResponse(IUtilizationTelemetry telemetry)
+		{
+			StartTime = telemetry.StartTime;
+			FinishTime = telemetry.FinishTime;
+			AdminTime = telemetry.AdminTime;
+			HibernatingTime = telemetry.HibernatingTime;
+			NumAgents = telemetry.NumAgents;
+
+			Pools = telemetry.Pools.ConvertAll(pool => new UtilizationTelemetryPool(pool));
+		}
+	}
+}
diff --git a/Engine/Source/Programs/Horde/Horde.Build/Agents/Telemetry/TelemetryController.cs b/Engine/Source/Programs/Horde/Horde.Build/Agents/Telemetry/TelemetryController.cs
new file mode 100644
--- /dev/null
+++ b/Engine/Source/Programs/Horde/Horde.Build/Agents/Telemetry/TelemetryController.cs
@@ -0,0 +1,53 @@
+// Copyright Epic Games, Inc. All Rights Reserved.
+
+using System;
+using System.Collections.Generic;
+using System.Threading.Tasks;
+using Horde.Build.Api;
+using Horde.Build.Collections;
+using Horde.Build.Models;
+using Microsoft.AspNetCore.Mvc;
+
+namespace Horde.Build.Controllers
+{
+	/// <summary>
+	/// Controller for the /api/v1/reports endpoint, used for the reports pages
+	/// </summary>
+	[ApiController]
+	[Route("[controller]")]
+	public sealed class ReportsController : ControllerBase
+	{
+		/// <summary>
+		/// the Telemetry collection singleton
+		/// </summary>
+		readonly ITelemetryCollection _telemetryCollection;
+
+		/// <summary>
+		/// Constructor
+		/// </summary>
+		/// <param name="telemetryCollection">The telemetry collection</param>
+		public ReportsController(ITelemetryCollection telemetryCollection)
+		{
+			_telemetryCollection = telemetryCollection;
+		}
+
+		/// <summary>
+		/// Gets a collection of utilization data including an endpoint and a range
+		/// </summary>
+		[HttpGet]
+		[Route("/api/v1/reports/utilization/{endDate}")]
+		public async Task<ActionResult<List<UtilizationTelemetryResponse>>> GetStreamUtilizationData(DateTime endDate, [FromQuery(Name = "Range")] int range, [FromQuery(Name = "TzOffset")] int? tzOffset)
+		{
+			// Logic here is a bit messy. The client is always passing in a date at midnight
+			// If user passes in 12/1/2020 into the date with range of 1, the range should be 12/1/2020:00:00:00 to 12/1/2020:23:59:59
+			// Range 2 should be 11/30/2020:00:00:00 to 12/1/2020:23:59:59
+			int offset = tzOffset ?? 0;
+			DateTimeOffset endDateOffset = new DateTimeOffset(endDate, TimeSpan.FromHours(offset)).Add(new TimeSpan(23, 59, 59));
+			DateTimeOffset startDateOffset = endDate.Subtract(new TimeSpan(range - 1, 0, 0, 0));
+
+			List<IUtilizationTelemetry> telemetry = await _telemetryCollection.GetUtilizationTelemetryAsync(startDateOffset.UtcDateTime, endDateOffset.UtcDateTime);
+
+			return telemetry.ConvertAll(telemetry => new UtilizationTelemetryResponse(telemetry));
+		}
+	}
+}
diff --git a/Engine/Source/Programs/Horde/Horde.Build/Agents/Telemetry/TelemetryService.cs b/Engine/Source/Programs/Horde/Horde.Build/Agents/Telemetry/TelemetryService.cs
new file mode 100644
--- /dev/null
+++ b/Engine/Source/Programs/Horde/Horde.Build/Agents/Telemetry/TelemetryService.cs
@@ -0,0 +1,142 @@
+// Copyright Epic Games, Inc. All Rights Reserved.
+
+using System;
+using System.Collections.Generic;
+using System.Linq;
+using System.Threading;
+using System.Threading.Tasks;
+using Horde.Build.Collections;
+using Horde.Build.Models;
+using Horde.Build.Utilities;
+using HordeCommon;
+using Microsoft.Extensions.Hosting;
+using Microsoft.Extensions.Logging;
+
+namespace Horde.Build.Services
+{
+	using PoolId = StringId<IPool>;
+
+	/// <summary>
+	/// Service which updates telemetry periodically
+	/// </summary>
+	public sealed class TelemetryService : IHostedService, IDisposable
+	{
+		readonly ITelemetryCollection _telemetryCollection;
+		readonly IAgentCollection _agentCollection;
+		readonly ILeaseCollection _leaseCollection;
+		readonly IPoolCollection _poolCollection;
+		readonly IFleetManager _fleetManager;
+		readonly IClock _clock;
+		readonly ITicker _tick;
+		readonly ILogger _logger;
+
+		/// <summary>
+		/// Constructor
+		/// </summary>
+		public TelemetryService(ITelemetryCollection telemetryCollection, IAgentCollection agentCollection, ILeaseCollection leaseCollection, IPoolCollection poolCollection, IFleetManager fleetManager, IClock clock, ILogger<TelemetryService> logger)
+		{
+			_telemetryCollection = telemetryCollection;
+			_agentCollection = agentCollection;
+			_leaseCollection = leaseCollection;
+			_poolCollection = poolCollection;
+			_fleetManager = fleetManager;
+			_clock = clock;
+			_tick = clock.AddSharedTicker<TelemetryService>(TimeSpan.FromMinutes(10.0), TickLeaderAsync, logger);
+			_logger = logger;
+		}
+
+		/// <inheritdoc/>
+		public Task StartAsync(CancellationToken cancellationToken) => _tick.StartAsync();
+
+		/// <inheritdoc/>
+		public Task StopAsync(CancellationToken cancellationToken) => _tick.StopAsync();
+
+		/// <inheritdoc/>
+		public void Dispose() => _tick.Dispose();
+
+		/// <inheritdoc/>
+		async ValueTask TickLeaderAsync(CancellationToken stoppingToken)
+		{
+			DateTime currentTime = _clock.UtcNow;
+
+			// Find the last time that we need telemetry for
+			DateTime maxTime = currentTime.Date + TimeSpan.FromHours(currentTime.Hour);
+
+			// Get the latest telemetry data
+			IUtilizationTelemetry? latest = await _telemetryCollection.GetLatestUtilizationTelemetryAsync();
+			TimeSpan interval = TimeSpan.FromHours(1.0);
+			int count = (latest == null) ? (7 * 24) : (int)Math.Round((maxTime - latest.FinishTime) / interval);
+			DateTime minTime = maxTime - count * interval;
+
+			// Query all the current data
+			List<IAgent> agents = await _agentCollection.FindAsync();
+			List<IPool> pools = await _poolCollection.GetAsync();
+			List<ILease> leases = await _leaseCollection.FindLeasesAsync(minTime: minTime);
+
+			// Remove any agents which are offline
+			agents.RemoveAll(x => !x.Enabled || !x.IsSessionValid(currentTime));
+
+			// Find all the agents
+			Dictionary<AgentId, List<PoolId>> agentToPoolIds = agents.ToDictionary(x => x.Id, x => x.GetPools().ToList());
+
+			// Generate all the telemetry data
+			DateTime bucketMinTime = minTime;
+			for (int idx = 0; idx < count; idx++)
+			{
+				DateTime bucketMaxTime = bucketMinTime + interval;
+				_logger.LogInformation("Creating telemetry for {MinTime} to {MaxTime}", bucketMinTime, bucketMaxTime);
+
+				NewUtilizationTelemetry telemetry = new NewUtilizationTelemetry(bucketMinTime, bucketMaxTime);
+				telemetry.NumAgents = agents.Count;
+				foreach (IPool pool in pools)
+				{
+					if (pool.EnableAutoscaling)
+					{
+						int numStoppedInstances = await _fleetManager.GetNumStoppedInstancesAsync(pool);
+						telemetry.NumAgents += numStoppedInstances;
+
+						NewPoolUtilizationTelemetry poolTelemetry = telemetry.FindOrAddPool(pool.Id);
+						poolTelemetry.NumAgents += numStoppedInstances;
+						poolTelemetry.HibernatingTime += interval.TotalHours * numStoppedInstances;
+					}
+				}
+				foreach (PoolId poolId in agentToPoolIds.Values.SelectMany(x => x))
+				{
+					telemetry.FindOrAddPool(poolId).NumAgents++;
+				}
+				foreach (ILease lease in leases)
+				{
+					if (lease.StartTime < bucketMaxTime && (!lease.FinishTime.HasValue || lease.FinishTime >= bucketMinTime))
+					{
+						List<PoolId>? leasePools;
+						if (agentToPoolIds.TryGetValue(lease.AgentId, out leasePools))
+						{
+							DateTime finishTime = lease.FinishTime ?? bucketMaxTime;
+							double time = new TimeSpan(Math.Min(finishTime.Ticks, bucketMaxTime.Ticks) - Math.Max(lease.StartTime.Ticks, bucketMinTime.Ticks)).TotalHours;
+
+							foreach (PoolId poolId in leasePools)
+							{
+								NewPoolUtilizationTelemetry poolTelemetry = telemetry.FindOrAddPool(poolId);
+								if (lease.PoolId == null || lease.StreamId == null)
+								{
+									poolTelemetry.AdminTime += time;
+								}
+								else if (poolId == lease.PoolId)
+								{
+									poolTelemetry.FindOrAddStream(lease.StreamId.Value).Time += time;
+								}
+								else
+								{
+									poolTelemetry.OtherTime += time;
+								}
+							}
+						}
+					}
+				}
+				await _telemetryCollection.AddUtilizationTelemetryAsync(telemetry);
+
+				bucketMinTime = bucketMaxTime;
+			}
+		}
+	}
+}
diff --git a/Engine/Source/Programs/Horde/Horde.Build/Auditing/AuditLog.cs b/Engine/Source/Programs/Horde/Horde.Build/Auditing/AuditLog.cs
new file mode 100644
--- /dev/null
+++ b/Engine/Source/Programs/Horde/Horde.Build/Auditing/AuditLog.cs
@@ -0,0 +1,199 @@
+// Copyright Epic Games, Inc. All Rights Reserved.
+
+using System;
+using System.Collections.Generic;
+using System.Threading.Channels;
+using System.Threading.Tasks;
+using EpicGames.Core;
+using Horde.Build.Server;
+using Horde.Build.Utilities;
+using Microsoft.Extensions.Logging;
+using MongoDB.Bson;
+using MongoDB.Bson.Serialization.Attributes;
+using MongoDB.Driver;
+
+namespace Horde.Build.Collections.Impl
+{
+	class AuditLog<TSubject> : IAuditLog<TSubject>, IAsyncDisposable, IDisposable
+	{
+		class AuditLogMessage : IAuditLogMessage<TSubject>
+		{
+			public ObjectId Id { get; set; }
+
+			[BsonElement("s")]
+			public TSubject Subject { get; set; }
+
+			[BsonElement("t")]
+			public DateTime TimeUtc { get; set; }
+
+			[BsonElement("l")]
+			public LogLevel Level { get; set; }
+
+			[BsonElement("d")]
+			public string Data { get; set; }
+
+			public AuditLogMessage()
+			{
+				Subject = default!;
+				Data = String.Empty;
+			}
+
+			public AuditLogMessage(TSubject subject, DateTime timeUtc, LogLevel level, string data)
+			{
+				Id = ObjectId.GenerateNewId();
+				Subject = subject;
+				TimeUtc = timeUtc;
+				Level = level;
+				Data = data;
+			}
+		}
+
+		class AuditLogChannel : IAuditLogChannel<TSubject>
+		{
+			sealed class Scope : IDisposable
+			{
+				public void Dispose() { }
+			}
+
+			public readonly AuditLog<TSubject> Outer;
+			public TSubject Subject { get; }
+
+			public AuditLogChannel(AuditLog<TSubject> outer, TSubject subject)
+			{
+				Outer = outer;
+				Subject = subject;
+			}
+
+			public IDisposable BeginScope<TState>(TState state) => new Scope();
+
+			public bool IsEnabled(LogLevel logLevel) => true;
+
+			public void Log<TState>(LogLevel logLevel, EventId eventId, TState state, Exception? exception, Func<TState, Exception?, string> formatter)
+			{
+				DateTime time = DateTime.UtcNow;
+
+				LogEvent logEvent = LogEvent.FromState(logLevel, eventId, state, exception, formatter);
+				string data = logEvent.ToJson();
+				AuditLogMessage message = new AuditLogMessage(Subject, time, logLevel, data);
+				Outer._messageChannel.Writer.TryWrite(message);
+
+#pragma warning disable CA2254 // Template should be a static expression
+				using (IDisposable _ = Outer._logger.BeginScope($"Subject: {{{Outer._subjectProperty}}}", Subject))
+				{
+					Outer._logger.Log(logLevel, eventId, state, exception, formatter);
+				}
+#pragma warning restore CA2254 // Template should be a static expression
+			}
+
+			public IAsyncEnumerable<IAuditLogMessage> FindAsync(DateTime? minTime, DateTime? maxTime, int? index, int? count) => Outer.FindAsync(Subject, minTime, maxTime, index, count);
+
+			public Task<long> DeleteAsync(DateTime? minTime, DateTime? maxTime) => Outer.DeleteAsync(Subject, minTime, maxTime);
+		}
+
+		readonly IMongoCollection<AuditLogMessage> _messages;
+		readonly Channel<AuditLogMessage> _messageChannel;
+		readonly string _subjectProperty;
+		readonly ILogger _logger;
+		readonly Task _backgroundTask;
+
+		public IAuditLogChannel<TSubject> this[TSubject subject] => new AuditLogChannel(this, subject);
+
+		public AuditLog(MongoService mongoService, string collectionName, string subjectProperty, ILogger logger)
+		{
+			List<MongoIndex<AuditLogMessage>> indexes = new List<MongoIndex<AuditLogMessage>>();
+			indexes.Add(builder => builder.Ascending(x => x.Subject).Descending(x => x.TimeUtc));
+
+			_messages = mongoService.GetCollection<AuditLogMessage>(collectionName, indexes);
+
+			_messageChannel = Channel.CreateUnbounded<AuditLogMessage>();
+			_subjectProperty = subjectProperty;
+			_logger = logger;
+
+			_backgroundTask = Task.Run(() => WriteMessagesAsync());
+		}
+
+		public async ValueTask DisposeAsync()
+		{
+			_messageChannel.Writer.TryComplete();
+			await _backgroundTask;
+		}
+
+		public void Dispose()
+		{
+			DisposeAsync().AsTask().Wait();
+		}
+
+		async Task WriteMessagesAsync()
+		{
+			while (await _messageChannel.Reader.WaitToReadAsync())
+			{
+				List<AuditLogMessage> newMessages = new List<AuditLogMessage>();
+				while (_messageChannel.Reader.TryRead(out AuditLogMessage? newMessage))
+				{
+					newMessages.Add(newMessage);
+				}
+				if (newMessages.Count > 0)
+				{
+					await _messages.InsertManyAsync(newMessages);
+				}
+			}
+		}
+
+		async IAsyncEnumerable<IAuditLogMessage<TSubject>> FindAsync(TSubject subject, DateTime? minTime = null, DateTime? maxTime = null, int? index = null, int? count = null)
+		{
+			FilterDefinition<AuditLogMessage> filter = Builders<AuditLogMessage>.Filter.Eq(x => x.Subject, subject);
+			if (minTime != null)
+			{
+				filter &= Builders<AuditLogMessage>.Filter.Gte(x => x.TimeUtc, minTime.Value);
+			}
+			if (maxTime != null)
+			{
+				filter &= Builders<AuditLogMessage>.Filter.Lte(x => x.TimeUtc, maxTime.Value);
+			}
+
+			using (IAsyncCursor<AuditLogMessage> cursor = await _messages.Find(filter).SortByDescending(x => x.TimeUtc).Range(index, count).ToCursorAsync())
+			{
+				while (await cursor.MoveNextAsync())
+				{
+					foreach (AuditLogMessage message in cursor.Current)
+					{
+						yield return message;
+					}
+				}
+			}
+		}
+
+		async Task<long> DeleteAsync(TSubject subject, DateTime? minTime = null, DateTime? maxTime = null)
+		{
+			FilterDefinition<AuditLogMessage> filter = Builders<AuditLogMessage>.Filter.Eq(x => x.Subject, subject);
+			if (minTime != null)
+			{
+				filter &= Builders<AuditLogMessage>.Filter.Gte(x => x.TimeUtc, minTime.Value);
+			}
+			if (maxTime != null)
+			{
+				filter &= Builders<AuditLogMessage>.Filter.Lte(x => x.TimeUtc, maxTime.Value);
+			}
+
+			DeleteResult result = await _messages.DeleteManyAsync(filter);
+			return result.DeletedCount;
+		}
+	}
+
+	class AuditLogFactory<TSubject> : IAuditLogFactory<TSubject>
+	{
+		readonly MongoService _mongoService;
+		readonly ILogger<AuditLog<TSubject>> _logger;
+
+		public AuditLogFactory(MongoService mongoService, ILogger<AuditLog<TSubject>> logger)
+		{
+			_mongoService = mongoService;
+			_logger = logger;
+		}
+
+		public IAuditLog<TSubject> Create(string collectionName, string subjectProperty)
+		{
+			return new AuditLog<TSubject>(_mongoService, collectionName, subjectProperty, _logger);
+		}
+	}
+}
