Horde: Finish re-org of Horde solution started a while ago.

diff --git a/Engine/Source/Programs/Horde/Horde.Build/Devices/DeviceCollection.cs b/Engine/Source/Programs/Horde/Horde.Build/Devices/DeviceCollection.cs
new file mode 100644
--- /dev/null
+++ b/Engine/Source/Programs/Horde/Horde.Build/Devices/DeviceCollection.cs
@@ -0,0 +1,778 @@
+// Copyright Epic Games, Inc. All Rights Reserved.
+
+using System;
+using System.Collections.Generic;
+using System.Linq;
+using System.Threading.Tasks;
+using Horde.Build.Acls;
+using Horde.Build.Models;
+using Horde.Build.Server;
+using Horde.Build.Utilities;
+using MongoDB.Bson;
+using MongoDB.Bson.Serialization.Attributes;
+using MongoDB.Driver;
+
+namespace Horde.Build.Collections.Impl
+{
+	using DeviceId = StringId<IDevice>;
+	using DevicePlatformId = StringId<IDevicePlatform>;
+	using DevicePoolId = StringId<IDevicePool>;
+	using ProjectId = StringId<IProject>;
+	using UserId = ObjectId<IUser>;
+
+	/// <summary>
+	/// Collection of device documents
+	/// </summary>
+	public class DeviceCollection : IDeviceCollection
+	{
+
+		/// <summary>
+		/// Document representing a device platform
+		/// </summary>
+		class DevicePlatformDocument : IDevicePlatform
+		{
+			[BsonRequired, BsonId]
+			public DevicePlatformId Id { get; set; }
+
+			public string Name { get; set; }
+
+			public List<string> Models { get; set; } = new List<string>();
+
+			IReadOnlyList<string> IDevicePlatform.Models => Models;
+
+			[BsonConstructor]
+			private DevicePlatformDocument()
+			{
+				Name = null!;
+			}
+
+			public DevicePlatformDocument(DevicePlatformId id, string name)
+			{
+				Id = id;
+				Name = name;
+			}
+		}
+
+		/// <summary>
+		/// Document representing a pool of devices
+		/// </summary>
+		class DevicePoolDocument : IDevicePool
+		{
+			[BsonRequired, BsonId]
+			public DevicePoolId Id { get; set; }
+
+			[BsonRequired]
+            public DevicePoolType PoolType { get; set; }
+
+			[BsonIgnoreIfNull]
+			public List<ProjectId>? ProjectIds { get; set; }
+
+			[BsonRequired]
+			public string Name { get; set; } = null!;
+
+			[BsonIgnoreIfNull]
+			public Acl? Acl { get; set; }
+
+			[BsonConstructor]
+			private DevicePoolDocument()
+			{
+            }
+
+			public DevicePoolDocument(DevicePoolId id, string name, DevicePoolType poolType, List<ProjectId>? projectIds)
+			{
+				Id = id;
+				Name = name;
+                PoolType = poolType;
+				ProjectIds = projectIds;
+            }
+		}
+
+		/// <summary>
+		/// Document representing a reservation of devices
+		/// </summary>
+		class DeviceReservationDocument : IDeviceReservation
+		{
+			/// <summary>
+			/// Randomly generated unique id for this reservation.
+			/// </summary>
+			[BsonRequired, BsonId]
+			public ObjectId Id { get; set; }
+
+			[BsonRequired]
+			public DevicePoolId PoolId { get; set; }
+
+			[BsonIgnoreIfNull]
+			public string? JobId { get; set; }
+
+			[BsonIgnoreIfNull]
+			public string? StepId { get; set; }
+
+            /// <summary>
+            /// Reservations held by a user, requires a token like download code
+            /// </summary>
+            [BsonIgnoreIfNull]
+			public UserId? UserId { get; set; }
+
+            /// <summary>
+            /// The hostname of the machine which has made the reservation
+            /// </summary>
+            [BsonIgnoreIfNull]
+			public string? Hostname { get; set; }
+
+			/// <summary>
+			/// Optional string holding details about the reservation
+			/// </summary>
+			[BsonIgnoreIfNull]
+			public string? ReservationDetails { get; set; }
+
+			/// <summary>
+			/// DeviceIds in the reservation
+			/// </summary>
+			public List<DeviceId> Devices { get; set; } = new List<DeviceId>();
+
+			public List<string> RequestedDevicePlatforms { get; set; } = new List<string>();
+
+			public DateTime CreateTimeUtc { get; set; }
+
+			public DateTime UpdateTimeUtc { get; set; }
+
+			// Legacy Guid
+			public string LegacyGuid { get; set; } = null!;
+
+			[BsonConstructor]
+			private DeviceReservationDocument()
+			{
+
+			}
+
+			public DeviceReservationDocument(ObjectId id, DevicePoolId poolId, List<DeviceId> devices, List<string> requestedDevicePlatforms, DateTime createTimeUtc, string? hostname, string? reservationDetails, string? jobId, string? stepId)
+			{
+				Id = id;
+				PoolId = poolId;
+				Devices = devices;
+				RequestedDevicePlatforms = requestedDevicePlatforms;
+				CreateTimeUtc = createTimeUtc;
+				UpdateTimeUtc = createTimeUtc;
+				Hostname = hostname;
+				ReservationDetails = reservationDetails;
+                JobId = jobId;
+                StepId = stepId;
+
+                LegacyGuid = Guid.NewGuid().ToString();
+			}
+		}
+
+		/// <summary>
+		/// Concrete implementation of an device document
+		/// </summary>
+		class DeviceDocument : IDevice
+		{
+			[BsonRequired, BsonId]
+			public DeviceId Id { get; set; }
+
+			public DevicePlatformId PlatformId { get; set; }
+
+			public DevicePoolId PoolId { get; set; }
+
+			public string Name { get; set; } = null!;
+
+			public bool Enabled { get; set; }
+
+			[BsonIgnoreIfNull]
+			public string? ModelId { get; set; }
+
+			[BsonIgnoreIfNull]
+			public string? Address { get; set; }
+
+			[BsonIgnoreIfNull]
+			public string? CheckedOutByUser { get; set; }
+
+			[BsonIgnoreIfNull]
+			public DateTime? CheckOutTime { get; set; }
+
+			[BsonIgnoreIfNull]
+			public bool? CheckoutExpiringNotificationSent { get; set; }
+
+			[BsonIgnoreIfNull]
+			public DateTime? ProblemTimeUtc { get; set; }
+
+			[BsonIgnoreIfNull]
+			public DateTime? MaintenanceTimeUtc { get; set; }
+
+			/// <summary>
+			///  The last time this device was reserved, used to cycle devices for reservations
+			/// </summary>
+			[BsonIgnoreIfNull]
+			public DateTime? ReservationTimeUtc { get; set; }
+
+			[BsonIgnoreIfNull]
+			public string? ModifiedByUser { get; set; }
+
+			[BsonIgnoreIfNull]
+			public string? Notes { get; set; }
+
+			[BsonIgnoreIfNull]
+			public List<DeviceUtilizationTelemetry>? Utilization { get; set; }
+
+			[BsonIgnoreIfNull]
+			public Acl? Acl { get; set; }
+
+			[BsonConstructor]
+			private DeviceDocument()
+			{
+			}
+
+			public DeviceDocument(DeviceId id, DevicePlatformId platformId, DevicePoolId poolId, string name, bool enabled, string? address, string? modelId, UserId? userId)
+			{
+				Id = id;
+				PlatformId = platformId;
+				PoolId = poolId;
+				Name = name;
+				Enabled = enabled;
+				Address = address;
+				ModelId = modelId;
+                ModifiedByUser = userId?.ToString();
+            }
+		}
+
+		readonly IMongoCollection<DevicePlatformDocument> _platforms;
+
+		readonly IMongoCollection<DeviceDocument> _devices;
+
+		readonly IMongoCollection<DevicePoolDocument> _pools;
+
+		readonly IMongoCollection<DeviceReservationDocument> _reservations;
+
+		readonly int _checkoutDays = 7;
+
+		/// <summary>
+		/// Constructor
+		/// </summary>
+		public DeviceCollection(MongoService mongoService)
+		{
+			_devices = mongoService.GetCollection<DeviceDocument>("Devices", keys => keys.Ascending(x => x.Name), unique: true);
+			_platforms = mongoService.GetCollection<DevicePlatformDocument>("Devices.Platforms", keys => keys.Ascending(x => x.Name), unique: true);
+			_pools = mongoService.GetCollection<DevicePoolDocument>("Devices.Pools", keys => keys.Ascending(x => x.Name), unique: true);
+			_reservations = mongoService.GetCollection<DeviceReservationDocument>("Devices.Reservations");
+		}
+
+		/// <inheritdoc/>
+		public async Task<IDevice?> TryAddDeviceAsync(DeviceId id, string name, DevicePlatformId platformId, DevicePoolId poolId, bool? enabled, string? address, string? modelId, UserId? userId)
+		{
+			DeviceDocument newDevice = new DeviceDocument(id, platformId, poolId, name, enabled ?? true, address, modelId, userId);
+			await _devices.InsertOneAsync(newDevice);
+			return newDevice;
+		}
+
+		/// <inheritdoc/>
+		public async Task<IDevicePlatform?> TryAddPlatformAsync(DevicePlatformId id, string name)
+		{
+			DevicePlatformDocument newPlatform = new DevicePlatformDocument(id, name);
+
+			try
+			{
+				await _platforms.InsertOneAsync(newPlatform);
+				return newPlatform;
+			}
+			catch (MongoWriteException ex)
+			{
+				if (ex.WriteError.Category == ServerErrorCategory.DuplicateKey)
+				{
+					return null;
+				}
+				else
+				{
+					throw;
+				}
+			}
+		}
+
+		/// <inheritdoc/>
+		public async Task<List<IDevicePlatform>> FindAllPlatformsAsync()
+		{
+			List<DevicePlatformDocument> results = await _platforms.Find(x => true).ToListAsync();
+			return results.OrderBy(x => x.Name).Select<DevicePlatformDocument, IDevicePlatform>(x => x).ToList();
+		}
+
+		/// <inheritdoc/>
+		public async Task<bool> UpdatePlatformAsync(DevicePlatformId platformId, string[]? modelIds)
+		{
+			UpdateDefinitionBuilder<DevicePlatformDocument> updateBuilder = Builders<DevicePlatformDocument>.Update;
+
+			List<UpdateDefinition<DevicePlatformDocument>> updates = new List<UpdateDefinition<DevicePlatformDocument>>();
+
+			if (modelIds != null)
+			{
+				updates.Add(updateBuilder.Set(x => x.Models, modelIds.ToList()));
+			}
+
+			if (updates.Count > 0)
+			{
+				await _platforms.FindOneAndUpdateAsync<DevicePlatformDocument>(x => x.Id == platformId, updateBuilder.Combine(updates));
+			}
+
+			return true;
+
+		}
+
+		/// <inheritdoc/>
+		public async Task<IDevicePool?> TryAddPoolAsync(DevicePoolId id, string name, DevicePoolType poolType, List<ProjectId>? projectIds)
+		{
+			DevicePoolDocument newPool = new DevicePoolDocument(id, name, poolType, projectIds);
+
+			try
+			{
+				await _pools.InsertOneAsync(newPool);
+				return newPool;
+			}
+			catch (MongoWriteException ex)
+			{
+				if (ex.WriteError.Category == ServerErrorCategory.DuplicateKey)
+				{
+					return null;
+				}
+				else
+				{
+					throw;
+				}
+			}
+		}
+
+		/// <inheritdoc/>
+		public async Task UpdatePoolAsync(DevicePoolId id, List<ProjectId>? projectIds)
+		{
+			UpdateDefinitionBuilder<DevicePoolDocument> updateBuilder = Builders<DevicePoolDocument>.Update;
+
+			List<UpdateDefinition<DevicePoolDocument>> updates = new List<UpdateDefinition<DevicePoolDocument>>();
+
+			if (projectIds != null)
+			{
+				updates.Add(updateBuilder.Set(x => x.ProjectIds, projectIds));
+			}
+
+			if (updates.Count > 0)
+			{
+				await _pools.FindOneAndUpdateAsync<DevicePoolDocument>(x => x.Id == id, updateBuilder.Combine(updates));
+			}
+		}
+
+		/// <inheritdoc/>
+		public async Task<List<IDevice>> FindAllDevicesAsync(List<DeviceId>? deviceIds = null)
+		{
+			FilterDefinition<DeviceDocument> filter = Builders<DeviceDocument>.Filter.Empty;
+			if (deviceIds != null)
+			{
+				filter &= Builders<DeviceDocument>.Filter.In(x => x.Id, deviceIds);
+			}
+			List<DeviceDocument> results = await _devices.Find(filter).ToListAsync();
+			return results.OrderBy(x => x.Name).Select<DeviceDocument, IDevice>(x => x).ToList();
+		}
+
+		/// <inheritdoc/>
+		public async Task<List<IDevicePool>> FindAllPoolsAsync()
+		{
+			List<DevicePoolDocument> results = await _pools.Find(x => true).ToListAsync();
+			return results.OrderBy(x => x.Name).Select<DevicePoolDocument, IDevicePool>(x => x).ToList();
+		}
+
+		/// <inheritdoc/>
+		public async Task<IDevicePlatform?> GetPlatformAsync(DevicePlatformId platformId)
+		{
+			return await _platforms.Find<DevicePlatformDocument>(x => x.Id == platformId).FirstOrDefaultAsync();
+		}
+
+		/// <inheritdoc/>
+		public async Task<IDevicePool?> GetPoolAsync(DevicePoolId poolId)
+		{
+			return await _pools.Find<DevicePoolDocument>(x => x.Id == poolId).FirstOrDefaultAsync();
+		}
+
+		/// <inheritdoc/>
+		public async Task<IDevice?> GetDeviceAsync(DeviceId deviceId)
+		{
+			return await _devices.Find<DeviceDocument>(x => x.Id == deviceId).FirstOrDefaultAsync();
+		}
+
+		/// <inheritdoc/>
+		public async Task<IDevice?> GetDeviceByNameAsync(string deviceName)
+		{
+			return await _devices.Find<DeviceDocument>(x => x.Name == deviceName).FirstOrDefaultAsync();
+		}
+
+		/// <inheritdoc/>
+		public async Task CheckoutDeviceAsync(DeviceId deviceId, UserId? checkedOutByUserId)
+		{
+			UpdateDefinitionBuilder<DeviceDocument> updateBuilder = Builders<DeviceDocument>.Update;
+
+			List<UpdateDefinition<DeviceDocument>> updates = new List<UpdateDefinition<DeviceDocument>>();
+
+            string? userId = checkedOutByUserId?.ToString();
+
+            updates.Add(updateBuilder.Set(x => x.CheckedOutByUser, String.IsNullOrEmpty(userId) ? null : userId));			
+
+			if (checkedOutByUserId != null)
+			{
+				updates.Add(updateBuilder.Set(x => x.CheckOutTime, DateTime.UtcNow));
+				updates.Add(updateBuilder.Set(x => x.CheckoutExpiringNotificationSent, false));
+			}
+			else
+			{
+				updates.Add(updateBuilder.Set(x => x.CheckoutExpiringNotificationSent, true));
+			}
+
+			await _devices.FindOneAndUpdateAsync<DeviceDocument>(x => x.Id == deviceId, updateBuilder.Combine(updates));
+		}
+
+		/// <inheritdoc/>
+		public async Task UpdateDeviceAsync(DeviceId deviceId, DevicePoolId? newPoolId, string? newName, string? newAddress, string? newModelId, string? newNotes, bool? newEnabled, bool? newProblem, bool? newMaintenance, UserId? modifiedByUserId = null)
+		{
+			UpdateDefinitionBuilder<DeviceDocument> updateBuilder = Builders<DeviceDocument>.Update;
+
+			List<UpdateDefinition<DeviceDocument>> updates = new List<UpdateDefinition<DeviceDocument>>();
+
+			DateTime utcNow = DateTime.UtcNow;
+
+			if (modifiedByUserId != null)
+			{
+				updates.Add(updateBuilder.Set(x => x.ModifiedByUser, modifiedByUserId.ToString()));
+			}
+
+			if (newPoolId != null)
+			{
+				updates.Add(updateBuilder.Set(x => x.PoolId, newPoolId.Value));
+			}
+
+			if (newName != null)
+			{
+				updates.Add(updateBuilder.Set(x => x.Name, newName));
+			}
+
+			if (newEnabled != null)
+			{
+				updates.Add(updateBuilder.Set(x => x.Enabled, newEnabled));
+			}
+
+			if (newAddress != null)
+			{
+				updates.Add(updateBuilder.Set(x => x.Address, newAddress));
+			}
+
+			if (!String.IsNullOrEmpty(newModelId))
+			{
+				if (newModelId == "Base")
+				{
+					updates.Add(updateBuilder.Set(x => x.ModelId, null));
+				}
+				else
+				{
+					updates.Add(updateBuilder.Set(x => x.ModelId, newModelId));
+				}				
+			}
+
+			if (newNotes != null)
+			{
+				updates.Add(updateBuilder.Set(x => x.Notes, newNotes));
+			}
+
+			if (newProblem.HasValue)
+			{
+				DateTime? problemTime = null;
+				if (newProblem.Value)
+				{
+					problemTime = utcNow;
+				}
+
+				updates.Add(updateBuilder.Set(x => x.ProblemTimeUtc, problemTime));
+			}
+
+			if (newMaintenance.HasValue)
+			{
+				DateTime? maintenanceTime = null;
+				if (newMaintenance.Value)
+				{
+					maintenanceTime = utcNow;
+				}
+
+				updates.Add(updateBuilder.Set(x => x.MaintenanceTimeUtc, maintenanceTime));
+			}
+
+			if (updates.Count > 0)
+			{
+				await _devices.FindOneAndUpdateAsync<DeviceDocument>(x => x.Id == deviceId, updateBuilder.Combine(updates));
+			}
+		}
+
+		/// <inheritdoc/>
+		public async Task<bool> DeleteDeviceAsync(DeviceId deviceId)
+		{
+			FilterDefinition<DeviceDocument> filter = Builders<DeviceDocument>.Filter.Eq(x => x.Id, deviceId);
+			DeleteResult result = await _devices.DeleteOneAsync(filter);
+			return result.DeletedCount > 0;
+
+		}
+
+		/// <inheritdoc/>
+		public async Task<List<IDeviceReservation>> FindAllDeviceReservationsAsync(DevicePoolId? poolId = null)
+		{
+			return await _reservations.Find(a => poolId == null || a.PoolId == poolId).ToListAsync<DeviceReservationDocument, IDeviceReservation>();
+		}
+
+		/// <inheritdoc/>
+		public async Task<IDeviceReservation?> TryAddReservationAsync(DevicePoolId poolId, List<DeviceRequestData> request, string? hostname, string? reservationDetails, string? jobId, string? stepId)
+		{
+
+			if (request.Count == 0)
+			{
+				return null;
+			}
+
+			DevicePoolDocument? pool = await _pools.Find<DevicePoolDocument>(x => x.Id == poolId).FirstOrDefaultAsync();
+
+			if (pool == null || pool.PoolType != DevicePoolType.Automation)
+			{
+                return null;
+            }
+			
+			HashSet<DeviceId> allocated = new HashSet<DeviceId>();
+			Dictionary<DeviceId, string> platformRequestMap = new Dictionary<DeviceId, string>();
+
+			List<DeviceReservationDocument> poolReservations = await _reservations.Find(x => x.PoolId == poolId).ToListAsync();
+
+			DateTime reservationTimeUtc = DateTime.UtcNow;
+
+			// Get available devices
+			List<DeviceDocument> poolDevices = await _devices.Find(x =>
+				x.PoolId == poolId &&
+				x.Enabled &&
+				x.MaintenanceTimeUtc == null).ToListAsync();
+
+			// filter out problem devices
+			poolDevices = poolDevices.FindAll(x => (x.ProblemTimeUtc == null || ((reservationTimeUtc - x.ProblemTimeUtc).Value.TotalMinutes > 30)));
+
+			// filter out currently reserved devices
+			poolDevices = poolDevices.FindAll(x => poolReservations.FirstOrDefault(p => p.Devices.Contains(x.Id)) == null);
+
+			// sort to use last reserved first to cycle devices
+			poolDevices.Sort((a, b) =>
+			{
+				DateTime? aTime = a.ReservationTimeUtc;
+				DateTime? bTime = b.ReservationTimeUtc;
+
+				if (aTime == bTime)
+				{
+					return 0;
+				}
+
+				if (aTime == null && bTime != null)
+				{
+					return -1;
+				}
+				if (aTime != null && bTime == null)
+				{
+					return 1;
+				}
+
+				return aTime < bTime ? -1 : 1;
+
+			});
+
+			foreach (DeviceRequestData data in request)
+			{
+				DeviceDocument? device = poolDevices.FirstOrDefault(a =>
+				{
+
+					if (allocated.Contains(a.Id) || a.PlatformId != data.PlatformId)
+					{
+						return false;
+					}
+
+					if (data.IncludeModels.Count > 0 && (a.ModelId == null || !data.IncludeModels.Contains(a.ModelId)))
+					{
+						return false;
+					}
+
+					if (data.ExcludeModels.Count > 0 && (a.ModelId != null && data.ExcludeModels.Contains(a.ModelId)))
+					{
+						return false;
+					}
+
+					return true;
+				});
+
+				if (device == null)
+				{
+					// can't fulfill request
+					return null;
+				}
+
+				allocated.Add(device.Id);
+				platformRequestMap.Add(device.Id, data.RequestedPlatform);
+			}
+
+			// update reservation time and utilization for allocated devices 
+			foreach (DeviceId id in allocated)
+			{
+				DeviceDocument device = poolDevices.First((device) => device.Id == id);
+
+				List<DeviceUtilizationTelemetry>? utilization = device.Utilization;
+
+				if (utilization == null)
+				{
+					utilization = new List<DeviceUtilizationTelemetry>();
+				}
+
+				// keep up to 100, maintaining order
+				if (utilization.Count > 99)
+				{
+					utilization = utilization.GetRange(0, 99);
+				}
+
+				utilization.Insert(0, new DeviceUtilizationTelemetry(reservationTimeUtc) { JobId = jobId, StepId = stepId });
+
+				UpdateDefinitionBuilder<DeviceDocument> deviceBuilder = Builders<DeviceDocument>.Update;
+				List<UpdateDefinition<DeviceDocument>> deviceUpdates = new List<UpdateDefinition<DeviceDocument>>();
+
+				deviceUpdates.Add(deviceBuilder.Set(x => x.ReservationTimeUtc, reservationTimeUtc));
+				deviceUpdates.Add(deviceBuilder.Set(x => x.Utilization, utilization));
+
+				await _devices.FindOneAndUpdateAsync<DeviceDocument>(x => x.Id == id, deviceBuilder.Combine(deviceUpdates));
+			}
+
+			List<DeviceId> deviceIds = allocated.ToList();
+			List<string> requestedPlatforms = deviceIds.Select(x => platformRequestMap[x]).ToList();
+
+			// Create new reservation
+			DeviceReservationDocument newReservation = new DeviceReservationDocument(ObjectId.GenerateNewId(), poolId, deviceIds, requestedPlatforms, reservationTimeUtc, hostname, reservationDetails, jobId, stepId);
+			await _reservations.InsertOneAsync(newReservation);
+
+			return newReservation;
+
+		}
+
+		/// <inheritdoc/>
+		public async Task<bool> TryUpdateReservationAsync(ObjectId id)
+		{
+			UpdateResult result = await _reservations.UpdateOneAsync(x => x.Id == id, Builders<DeviceReservationDocument>.Update.Set(x => x.UpdateTimeUtc, DateTime.UtcNow));
+			return result.ModifiedCount == 1;
+		}
+
+		/// <inheritdoc/>
+		public async Task<bool> DeleteReservationAsync(ObjectId id)
+		{
+			FilterDefinition<DeviceReservationDocument> filter = Builders<DeviceReservationDocument>.Filter.Eq(x => x.Id, id);
+			DeleteResult result = await _reservations.DeleteOneAsync(filter);
+			return result.DeletedCount > 0;
+		}
+
+		/// <summary>
+		/// Deletes expired reservations
+		/// </summary>
+		public async Task<bool> ExpireReservationsAsync()
+		{
+			List<IDeviceReservation> reserves = await _reservations.Find(a => true).ToListAsync<DeviceReservationDocument, IDeviceReservation>();
+
+			DateTime utcNow = DateTime.UtcNow;
+
+			reserves = reserves.FindAll(r => (utcNow - r.UpdateTimeUtc).TotalMinutes > 10).ToList();
+
+			if (reserves.Count > 0)
+			{
+				DeleteResult result = await _reservations.DeleteManyAsync(Builders<DeviceReservationDocument>.Filter.In(x => x.Id, reserves.Select(y => y.Id)));
+				return (result.DeletedCount > 0);
+			}
+
+			return false;
+
+		}
+
+		/// <summary>
+		/// Deletes expired user checkouts
+		/// </summary>
+		public async Task<List<(UserId, IDevice)>?> ExpireCheckedOutAsync()
+		{
+			FilterDefinition<DeviceDocument> filter = Builders<DeviceDocument>.Filter.Empty;
+			filter &= Builders<DeviceDocument>.Filter.Where(x => x.CheckedOutByUser != null && x.CheckOutTime != null);
+
+			List<DeviceDocument> checkedOutDevices = await _devices.Find(filter).ToListAsync();
+
+			DateTime utcNow = DateTime.UtcNow;
+			List<DeviceDocument> expiredDevices = checkedOutDevices.FindAll(x => (utcNow - x.CheckOutTime!.Value).TotalDays >= _checkoutDays).ToList();
+
+			if (expiredDevices.Count > 0)
+			{
+
+				FilterDefinition<DeviceDocument> updateFilter = Builders<DeviceDocument>.Filter.In(x => x.Id, expiredDevices.Select(y => y.Id));
+
+				UpdateDefinitionBuilder<DeviceDocument> deviceBuilder = Builders<DeviceDocument>.Update;
+				List<UpdateDefinition<DeviceDocument>> deviceUpdates = new List<UpdateDefinition<DeviceDocument>>();
+
+				deviceUpdates.Add(deviceBuilder.Set(x => x.CheckedOutByUser, null));
+				deviceUpdates.Add(deviceBuilder.Set(x => x.CheckOutTime, null));
+				deviceUpdates.Add(deviceBuilder.Set(x => x.CheckoutExpiringNotificationSent, true));
+
+				UpdateResult result = await _devices.UpdateManyAsync(Builders<DeviceDocument>.Filter.In(x => x.Id, expiredDevices.Select(y => y.Id)), deviceBuilder.Combine(deviceUpdates));
+
+				if (result.ModifiedCount > 0)
+				{
+					return expiredDevices.Select(x => (UserId.Parse(x.CheckedOutByUser!), (IDevice) x)).ToList();
+				}
+			}
+
+			return null;
+
+		}
+
+		/// <summary>
+		/// Gets a list of users to notify  that their device is about to expire in the next 24 hours
+		/// </summary>
+		public async Task<List<(UserId, IDevice)>?> ExpireNotificatonsAsync()
+		{
+			FilterDefinition<DeviceDocument> filter = Builders<DeviceDocument>.Filter.Empty;
+			filter &= Builders<DeviceDocument>.Filter.Where(x => x.CheckedOutByUser != null && x.CheckoutExpiringNotificationSent != true);
+
+			List<DeviceDocument> checkedOutDevices = await _devices.Find(filter).ToListAsync();
+
+			DateTime utcNow = DateTime.UtcNow;
+			List<DeviceDocument> expiredDevices = checkedOutDevices.FindAll(x => (utcNow - x.CheckOutTime!.Value).TotalDays >= (_checkoutDays - 1)).ToList();
+
+			if (expiredDevices.Count > 0)
+			{
+				FilterDefinition<DeviceDocument> updateFilter = Builders<DeviceDocument>.Filter.In(x => x.Id, expiredDevices.Select(y => y.Id));
+
+				UpdateDefinitionBuilder<DeviceDocument> deviceBuilder = Builders<DeviceDocument>.Update;
+				List<UpdateDefinition<DeviceDocument>> deviceUpdates = new List<UpdateDefinition<DeviceDocument>>();
+
+				deviceUpdates.Add(deviceBuilder.Set(x => x.CheckoutExpiringNotificationSent, true));
+
+				UpdateResult result = await _devices.UpdateManyAsync(Builders<DeviceDocument>.Filter.In(x => x.Id, expiredDevices.Select(y => y.Id)), deviceBuilder.Combine(deviceUpdates));
+
+				if (result.ModifiedCount > 0)
+				{
+					return expiredDevices.Select(x => (UserId.Parse(x.CheckedOutByUser!), (IDevice)x)).ToList();
+				}
+			}
+
+			return null;
+		}
+
+		/// <inheritdoc/>
+		public async Task<List<IDeviceReservation>> FindAllReservationsAsync()
+		{
+			List<DeviceReservationDocument> results = await _reservations.Find(x => true).ToListAsync();
+			return results.OrderBy(x => x.CreateTimeUtc.Ticks).Select<DeviceReservationDocument, IDeviceReservation>(x => x).ToList();
+		}
+
+		/// <inheritdoc/>
+		public async Task<IDeviceReservation?> TryGetReservationFromLegacyGuidAsync(string legacyGuid)
+		{
+			return await _reservations.Find<DeviceReservationDocument>(r => r.LegacyGuid == legacyGuid).FirstOrDefaultAsync();
+		}
+
+		/// <inheritdoc/>
+		public async Task<IDeviceReservation?> TryGetDeviceReservationAsync(DeviceId id)
+		{
+			List<DeviceReservationDocument> results = await _reservations.Find(x => true).ToListAsync();
+			return results.FirstOrDefault(r => r.Devices.Contains(id));
+        }
+	}
+}
diff --git a/Engine/Source/Programs/Horde/Horde.Build/Devices/DeviceService.cs b/Engine/Source/Programs/Horde/Horde.Build/Devices/DeviceService.cs
new file mode 100644
--- /dev/null
+++ b/Engine/Source/Programs/Horde/Horde.Build/Devices/DeviceService.cs
@@ -0,0 +1,642 @@
+// Copyright Epic Games, Inc. All Rights Reserved.
+
+using System;
+using System.Collections.Generic;
+using System.Linq;
+using System.Security.Claims;
+using System.Threading;
+using System.Threading.Tasks;
+using Horde.Build.Acls;
+using Horde.Build.Api;
+using Horde.Build.Collections;
+using Horde.Build.Models;
+using Horde.Build.Notifications;
+using Horde.Build.Utilities;
+using HordeCommon;
+using Microsoft.Extensions.Hosting;
+using Microsoft.Extensions.Logging;
+using MongoDB.Bson;
+
+namespace Horde.Build.Services
+{
+	using DeviceId = StringId<IDevice>;
+	using DevicePlatformId = StringId<IDevicePlatform>;
+	using DevicePoolId = StringId<IDevicePool>;
+	using JobId = ObjectId<IJob>;
+	using ProjectId = StringId<IProject>;
+	using UserId = ObjectId<IUser>;
+
+	/// <summary>
+	///  Device Pool Authorization (convenience class for pool ACL's)
+	/// </summary>
+	public class DevicePoolAuthorization
+	{
+		/// <summary>
+		/// The device pool
+		/// </summary>
+		public IDevicePool Pool { get; private set; }
+
+		/// <summary>
+		///  Read access to pool
+		/// </summary>
+		public bool Read { get; private set; }
+
+		/// <summary>
+		///  Write access to pool
+		/// </summary>
+		public bool Write { get; private set; }
+
+		/// <summary>
+		/// Constructor
+		/// </summary>
+		/// <param name="pool"></param>
+		/// <param name="read"></param>
+		/// <param name="write"></param>
+		public DevicePoolAuthorization(IDevicePool pool, bool read, bool write)
+		{
+			Pool = pool;
+			Read = read;
+			Write = write;				
+		}
+	}
+
+	/// <summary>
+	/// Platform map required by V1 API
+	/// </summary>
+	[SingletonDocument("6165a2e26fd5f104e31e6862")]
+	public class DevicePlatformMapV1 : SingletonBase
+	{
+		/// <summary>
+		/// Platform V1 => Platform Id
+		/// </summary>
+		public Dictionary<string, DevicePlatformId> PlatformMap { get; set; } = new Dictionary<string, DevicePlatformId>();
+
+		/// <summary>
+		/// Platform Id => Platform V1
+		/// </summary>
+		public Dictionary<DevicePlatformId, string> PlatformReverseMap { get; set; } = new Dictionary<DevicePlatformId, string>();
+
+		/// <summary>
+		/// Perfspec V1 => Model
+		/// </summary>
+		public Dictionary<DevicePlatformId, string> PerfSpecHighMap { get; set; } = new Dictionary<DevicePlatformId, string>();
+
+	}
+
+	/// <summary>
+	/// Device management service
+	/// </summary>
+	public sealed class DeviceService : IHostedService, IDisposable
+	{
+		/// <summary>
+		/// The ACL service instance
+		/// </summary>
+		readonly AclService _aclService;
+
+		/// <summary>
+		/// Instance of the notification service
+		/// </summary>
+		readonly INotificationService _notificationService;
+
+		/// <summary>
+		/// Singleton instance of the job service
+		/// </summary>
+		readonly JobService _jobService;
+
+		/// <summary>
+		/// Singleton instance of the stream service
+		/// </summary>
+		readonly StreamService _streamService;
+
+		/// <summary>
+		/// Singleton instance of the project service
+		/// </summary>
+		readonly ProjectService _projectService;
+
+		/// <summary>
+		/// The user collection instance
+		/// </summary>
+		IUserCollection UserCollection { get; set; }
+
+		/// <summary>
+		/// Log output writer
+		/// </summary>
+		readonly ILogger<DeviceService> _logger;
+
+		/// <summary>
+		/// Device collection
+		/// </summary>
+		readonly IDeviceCollection _devices;
+		readonly ITicker _ticker;
+
+		/// <summary>
+		/// Platform map V1 singleton
+		/// </summary>
+		readonly ISingletonDocument<DevicePlatformMapV1> _platformMapSingleton;
+
+		/// <summary>
+		/// Device service constructor
+		/// </summary>
+		public DeviceService(IDeviceCollection devices, ISingletonDocument<DevicePlatformMapV1> platformMapSingleton, IUserCollection userCollection, JobService jobService, ProjectService projectService, StreamService streamService, AclService aclService, INotificationService notificationService, IClock clock, ILogger<DeviceService> logger)
+		{
+			UserCollection = userCollection;
+			_devices = devices;
+            _jobService = jobService;
+			_projectService = projectService;
+            _streamService = streamService;
+            _aclService = aclService;
+            _notificationService = notificationService;
+			_ticker = clock.AddSharedTicker<DeviceService>(TimeSpan.FromMinutes(1.0), TickAsync, logger);
+            _logger = logger;
+			
+			_platformMapSingleton = platformMapSingleton;
+
+		}
+
+		/// <inheritdoc/>
+		public Task StartAsync(CancellationToken cancellationToken) => _ticker.StartAsync();
+
+		/// <inheritdoc/>
+		public Task StopAsync(CancellationToken cancellationToken) => _ticker.StopAsync();
+
+		/// <inheritdoc/>
+		public void Dispose() => _ticker.Dispose();
+
+		/// <summary>
+		/// Ticks service
+		/// </summary>
+		async ValueTask TickAsync(CancellationToken stoppingToken)
+		{
+
+			if (!stoppingToken.IsCancellationRequested)
+			{
+				await _devices.ExpireReservationsAsync();
+
+				List<(UserId, IDevice)>? expireNotifications = await _devices.ExpireNotificatonsAsync();
+				if (expireNotifications != null && expireNotifications.Count > 0)
+				{
+					foreach ((UserId, IDevice) expiredDevice in expireNotifications)
+					{
+						await NotifyDeviceServiceAsync($"Device {expiredDevice.Item2.PlatformId.ToString().ToUpperInvariant()} / {expiredDevice.Item2.Name} checkout will expire in 24 hours.  Please visit https://horde.devtools.epicgames.com/devices to renew the checkout if needed.", null, null, null, expiredDevice.Item1);
+					}
+				}
+
+				List<(UserId, IDevice)>? expireCheckouts = await _devices.ExpireCheckedOutAsync();
+				if (expireCheckouts != null && expireCheckouts.Count > 0)
+				{
+					foreach ((UserId, IDevice) expiredDevice in expireCheckouts)
+					{
+						await NotifyDeviceServiceAsync($"Device {expiredDevice.Item2.PlatformId.ToString().ToUpperInvariant()} / {expiredDevice.Item2.Name} checkout has expired.  The device has been returned to the shared pool and should no longer be accessed.  Please visit https://horde.devtools.epicgames.com/devices to checkout devices as needed.", null, null, null, expiredDevice.Item1);
+					}
+				}
+			}
+		}
+
+		/// <summary>
+		/// Create a new device platform
+		/// </summary>
+		public Task<IDevicePlatform?> TryCreatePlatformAsync(DevicePlatformId id, string name)
+		{
+			return _devices.TryAddPlatformAsync(id, name);
+		}
+
+		/// <summary>
+		/// Get a list of existing device platforms
+		/// </summary>
+		public Task<List<IDevicePlatform>> GetPlatformsAsync()
+		{
+			return _devices.FindAllPlatformsAsync();
+		}
+
+		/// <summary>
+		/// Update an existing platform
+		/// </summary>
+		public Task<bool> UpdatePlatformAsync(DevicePlatformId platformId, string[]? modelIds)
+		{
+			return _devices.UpdatePlatformAsync(platformId, modelIds);
+		}
+
+		/// <summary>
+		/// Get a specific device platform
+		/// </summary>
+		public Task<IDevicePlatform?> GetPlatformAsync(DevicePlatformId id)
+		{
+			return _devices.GetPlatformAsync(id);
+		}
+
+		/// <summary>
+		/// Get a device pool by id
+		/// </summary>
+		public Task<IDevicePool?> GetPoolAsync(DevicePoolId id)
+		{
+			return _devices.GetPoolAsync(id);
+		}
+
+		/// <summary>
+		/// Create a new device pool
+		/// </summary>
+		public Task<IDevicePool?> TryCreatePoolAsync(DevicePoolId id, string name, DevicePoolType poolType, List<ProjectId>? projectIds)
+		{
+			return _devices.TryAddPoolAsync(id, name, poolType, projectIds);
+		}
+
+		/// <summary>
+		/// Update a device pool
+		/// </summary>
+		public Task UpdatePoolAsync(DevicePoolId id, List<ProjectId>? projectIds)
+		{
+			return _devices.UpdatePoolAsync(id, projectIds);
+		}
+
+		/// <summary>
+		/// Get a list of existing device pools
+		/// </summary>
+		public Task<List<IDevicePool>> GetPoolsAsync()
+		{
+			return _devices.FindAllPoolsAsync();
+		}
+
+		/// <summary>
+		/// Get a list of devices, optionally filtered to provided ids
+		/// </summary>
+		public Task<List<IDevice>> GetDevicesAsync(List<DeviceId>? deviceIds = null)
+		{
+			return _devices.FindAllDevicesAsync(deviceIds);
+		}
+
+		/// <summary>
+		/// Get a specific device
+		/// </summary>
+		public Task<IDevice?> GetDeviceAsync(DeviceId id)
+		{
+			return _devices.GetDeviceAsync(id);
+		}
+
+		/// <summary>
+		/// Get a device by name
+		/// </summary>
+		public Task<IDevice?> GetDeviceByNameAsync(string deviceName)
+		{
+			return _devices.GetDeviceByNameAsync(deviceName);
+		}
+
+		/// <summary>
+		/// Delete a device
+		/// </summary>
+		public async Task<bool> DeleteDeviceAsync(DeviceId id)
+		{
+			return await _devices.DeleteDeviceAsync(id);
+		}
+
+		/// <summary>
+		/// Create a new device
+		/// </summary>
+		/// <param name="id">Unique id of the device</param>
+		/// <param name="name">Friendly name of the device</param>
+		/// <param name="platformId">The device platform</param>
+		/// <param name="poolId">Which pool to add the device</param>
+		/// <param name="enabled">Whether the device is enabled</param>
+		/// <param name="address">Address or hostname of device</param>
+		/// <param name="modelId">Vendor model id</param>
+        /// <param name="userId">User adding the device</param>
+		/// <returns></returns>
+		public Task<IDevice?> TryCreateDeviceAsync(DeviceId id, string name, DevicePlatformId platformId, DevicePoolId poolId, bool? enabled, string? address, string? modelId, UserId? userId = null)
+		{
+			return _devices.TryAddDeviceAsync(id, name, platformId, poolId, enabled, address, modelId, userId);
+		}
+
+		/// <summary>
+		/// Update a device
+		/// </summary>
+		public Task UpdateDeviceAsync(DeviceId deviceId, DevicePoolId? newPoolId = null, string? newName = null, string? newAddress = null, string? newModelId = null, string? newNotes = null, bool? newEnabled = null, bool? newProblem = null, bool? newMaintenance = null, UserId? modifiedByUserId = null)
+		{
+			return _devices.UpdateDeviceAsync(deviceId, newPoolId, newName, newAddress, newModelId, newNotes, newEnabled, newProblem, newMaintenance, modifiedByUserId);
+		}
+
+		/// <summary>
+		/// Checkout a device
+		/// </summary>
+		public Task CheckoutDeviceAsync(DeviceId deviceId, UserId? userId)
+		{
+            return _devices.CheckoutDeviceAsync(deviceId, userId);
+        }
+
+		/// <summary>
+		/// Try to create a reservation satisfying the specified device platforms and models
+		/// </summary>
+		public Task<IDeviceReservation?> TryCreateReservationAsync(DevicePoolId pool, List<DeviceRequestData> request, string? hostname = null, string? reservationDetails = null, string? jobId = null, string? stepId = null)
+		{
+			return _devices.TryAddReservationAsync(pool, request, hostname, reservationDetails, jobId, stepId);
+		}
+
+		/// <summary>
+		/// Update/renew an existing reservation
+		/// </summary>
+		public Task<bool> TryUpdateReservationAsync(ObjectId id)
+		{
+			return _devices.TryUpdateReservationAsync(id);
+		}
+
+		/// <summary>
+		///  Delete an existing reservation
+		/// </summary>
+		public Task<bool> DeleteReservationAsync(ObjectId id)
+		{
+			return _devices.DeleteReservationAsync(id);
+		}
+
+		/// <summary>
+		/// Get a reservation from a legacy guid
+		/// </summary>
+		public Task<IDeviceReservation?> TryGetReservationFromLegacyGuidAsync(string legacyGuid)
+		{
+			return _devices.TryGetReservationFromLegacyGuidAsync(legacyGuid);
+		}
+
+		/// <summary>
+		/// Get a reservation from a device id
+		/// </summary>
+		public Task<IDeviceReservation?> TryGetDeviceReservation(DeviceId deviceId)
+		{
+			return _devices.TryGetDeviceReservationAsync(deviceId);
+		}
+
+		/// <summary>
+		/// Get a list of existing device reservations
+		/// </summary>
+		public Task<List<IDeviceReservation>> GetReservationsAsync()
+		{
+			return _devices.FindAllReservationsAsync();
+		}
+
+		/// <summary>
+		/// 
+		/// </summary>
+		/// <param name="message"></param>
+		/// <param name="deviceId"></param>
+		/// <param name="jobId"></param>
+		/// <param name="stepId"></param>
+		/// <param name="userId"></param>
+		public async Task NotifyDeviceServiceAsync(string message, DeviceId? deviceId = null, string? jobId = null, string? stepId = null, UserId? userId = null)
+		{
+			try 
+			{
+				IDevice? device = null;
+				IDevicePool? pool = null;
+				IJob? job = null;
+				IJobStep? step = null;
+				INode? node = null;
+				IStream? stream = null;
+				IUser? user = null;
+
+				if (userId.HasValue)
+				{
+					user = await UserCollection.GetUserAsync(userId.Value);
+					if (user == null)
+					{
+						_logger.LogError("Unable to send device notification, can't find User {UserId}", userId.Value);
+						return;
+					}
+				}
+
+				if (deviceId.HasValue)
+				{
+					device = await GetDeviceAsync(deviceId.Value);
+					pool = await GetPoolAsync(device!.PoolId);
+				}
+
+				if (jobId != null)
+				{
+					job = await _jobService.GetJobAsync(new JobId(jobId));
+
+					if (job != null)
+					{
+						stream = await _streamService.GetStreamAsync(job.StreamId);
+
+						if (stepId != null)
+						{
+							IGraph graph = await _jobService.GetGraphAsync(job)!;
+
+							SubResourceId stepIdValue = SubResourceId.Parse(stepId);
+							IJobStepBatch? batch = job.Batches.FirstOrDefault(b => b.Steps.FirstOrDefault(s => s.Id == stepIdValue) != null);
+							if (batch != null)
+							{
+								step = batch.Steps.FirstOrDefault(s => s.Id == stepIdValue)!;
+								INodeGroup group = graph.Groups[batch.GroupIdx];
+								node = group.Nodes[step.NodeIdx];
+							}
+						}
+					}
+				}
+
+				_notificationService.NotifyDeviceService(message, device, pool, stream, job, step, node, user);
+
+			}
+			catch (Exception ex)
+			{
+                _logger.LogError(ex, "Error on device notification {Message}", ex.Message);
+            }
+        }
+
+		/// <summary>
+		/// Authorize device action
+		/// </summary>
+		/// <param name="action"></param>
+		/// <param name="user"></param>
+		/// <returns></returns>
+		public async Task<bool> AuthorizeAsync(AclAction action, ClaimsPrincipal user)
+		{
+			// Tihs is deprecated and device auth should be going through GetUserPoolAuthorizationsAsync
+
+			if (action == AclAction.DeviceRead)
+			{
+				return true;
+			}
+
+			if (user.IsInRole("Internal-Employees"))
+			{
+				return true;
+			}
+
+			if (await _aclService.AuthorizeAsync(AclAction.AdminWrite, user))
+			{
+				return true;
+			}
+
+			return false;
+
+		}
+
+		/// <summary>
+		/// Get list of pool authorizations for a user
+		/// </summary>
+		/// <param name="user"></param>
+		/// <returns></returns>
+		public async Task<List<DevicePoolAuthorization>> GetUserPoolAuthorizationsAsync(ClaimsPrincipal user)
+		{
+			
+			List<DevicePoolAuthorization> authPools = new List<DevicePoolAuthorization>();
+
+			List<IDevicePool> allPools = await GetPoolsAsync();
+			List<IProject> projects = await _projectService.GetProjectsAsync();
+
+			// Set of projects associated with device pools
+			HashSet<ProjectId> projectIds = new HashSet<ProjectId>(allPools.Where(x => x.ProjectIds != null).SelectMany(x => x.ProjectIds!));
+
+			Dictionary<ProjectId, bool> deviceRead = new Dictionary<ProjectId, bool>();
+			Dictionary<ProjectId, bool> deviceWrite = new Dictionary<ProjectId, bool>();
+			ProjectPermissionsCache permissionsCache = new ProjectPermissionsCache();
+
+			foreach (ProjectId projectId in projectIds)
+			{
+
+				if (projects.Where(x => x.Id == projectId).FirstOrDefault() == null)
+				{
+					_logger.LogWarning("Device pool authorization references missing project id {ProjectId}", projectId);
+					continue;
+				}
+
+				deviceRead.Add(projectId,  await _projectService.AuthorizeAsync(projectId, AclAction.DeviceRead, user, permissionsCache));
+				deviceWrite.Add(projectId, await _projectService.AuthorizeAsync(projectId, AclAction.DeviceWrite, user, permissionsCache));
+			}
+
+			// for global pools which aren't associated with a project
+			bool globalPoolAccess = user.IsInRole("Internal-Employees");
+
+			if (!globalPoolAccess)
+			{
+				if (await _aclService.AuthorizeAsync(AclAction.AdminWrite, user))
+				{
+					globalPoolAccess = true;
+				}
+			}
+
+			foreach (IDevicePool pool in allPools)
+			{				
+				if (pool.ProjectIds == null || pool.ProjectIds.Count == 0)
+				{
+					if (pool.PoolType == DevicePoolType.Shared && !globalPoolAccess)
+					{
+						authPools.Add(new DevicePoolAuthorization(pool, false, false));
+						continue;
+					}
+
+					authPools.Add(new DevicePoolAuthorization(pool, true, true));
+					continue;
+				}
+
+				bool read = false;
+				bool write = false;
+
+				foreach (ProjectId projectId in pool.ProjectIds)
+				{
+					bool value;
+
+					if (!read && deviceRead.TryGetValue(projectId, out value))
+					{
+						read = value;
+					}
+
+					if (!write && deviceWrite.TryGetValue(projectId, out value))
+					{
+						write = value;
+					}
+				}
+
+				authPools.Add(new DevicePoolAuthorization(pool, read, write));
+
+			}
+
+			return authPools;
+		}
+
+		/// <summary>
+		/// Get list of pool authorizations for a user
+		/// </summary>
+		/// <param name="id"></param>
+		/// <param name="user"></param>
+		/// <returns></returns>
+		public async Task<DevicePoolAuthorization?> GetUserPoolAuthorizationAsync(DevicePoolId id, ClaimsPrincipal user)
+		{
+			List<DevicePoolAuthorization> auth = await GetUserPoolAuthorizationsAsync(user);
+			return auth.Where(x => x.Pool.Id == id).FirstOrDefault();
+		}
+
+		/// <summary>
+		/// Update a device
+		/// </summary>
+		public async Task UpdateDevicePoolAsync(DevicePoolId poolId, List<ProjectId>? projectIds)
+		{
+			await _devices.UpdatePoolAsync(poolId, projectIds);
+		}
+
+		/// <summary>
+		/// Get Platform mappings for V1 API
+		/// </summary>		
+		public async Task<DevicePlatformMapV1> GetPlatformMapV1()
+		{
+			return await _platformMapSingleton.GetAsync();
+		}
+
+		/// <summary>
+		/// Updates the platform mapping information required for the V1 api
+		/// </summary>				
+		public async Task<bool> UpdatePlatformMapAsync(UpdatePlatformMapRequest request)
+		{
+
+			List<IDevicePlatform> platforms = await GetPlatformsAsync();
+
+			// Update the platform map
+			for (int i = 0; i < 10; i++)
+			{
+				DevicePlatformMapV1 instance = await _platformMapSingleton.GetAsync();
+
+				instance.PlatformMap = new Dictionary<string, DevicePlatformId>();
+				instance.PlatformReverseMap = new Dictionary<DevicePlatformId, string>();
+				instance.PerfSpecHighMap = new Dictionary<DevicePlatformId, string>();
+
+				foreach(KeyValuePair<string, string> entry in request.PlatformMap)
+				{
+					IDevicePlatform? platform = platforms.FirstOrDefault(p => p.Id == new DevicePlatformId(entry.Value));
+					if (platform == null)
+					{
+						throw new Exception($"Unknowm platform in map {entry.Key} : {entry.Value}");
+					}
+
+					instance.PlatformMap.Add(entry.Key, platform.Id);
+				}
+
+				foreach (KeyValuePair<string, string> entry in request.PlatformReverseMap)
+				{
+					IDevicePlatform? platform = platforms.FirstOrDefault(p => p.Id == new DevicePlatformId(entry.Key));
+					if (platform == null)
+					{
+						throw new Exception($"Unknowm platform in reverse map {entry.Key} : {entry.Value}");
+					}
+
+					instance.PlatformReverseMap.Add(platform.Id, entry.Value);
+				}
+
+				foreach (KeyValuePair<string, string> entry in request.PerfSpecHighMap)
+				{
+					IDevicePlatform? platform = platforms.FirstOrDefault(p => p.Id == new DevicePlatformId(entry.Key));
+					if (platform == null)
+					{
+						throw new Exception($"Unknowm platform in spec map {entry.Key} : {entry.Value}");
+					}
+
+					instance.PerfSpecHighMap.Add(platform.Id, entry.Value);
+				}
+
+				if (await _platformMapSingleton.TryUpdateAsync(instance))
+				{
+					return true;
+				}
+
+				Thread.Sleep(1000);
+			}
+
+			return false;
+		}
+	}
+}
diff --git a/Engine/Source/Programs/Horde/Horde.Build/Devices/DevicesController.Messages.cs b/Engine/Source/Programs/Horde/Horde.Build/Devices/DevicesController.Messages.cs
new file mode 100644
--- /dev/null
+++ b/Engine/Source/Programs/Horde/Horde.Build/Devices/DevicesController.Messages.cs
@@ -0,0 +1,738 @@
+// Copyright Epic Games, Inc. All Rights Reserved.
+
+using System;
+using System.Collections.Generic;
+using System.ComponentModel.DataAnnotations;
+using System.Linq;
+using Horde.Build.Models;
+
+namespace Horde.Build.Api
+{
+
+	/// <summary>
+	/// Create device platform request
+	/// </summary>
+	public class CreateDevicePlatformRequest
+	{
+		/// <summary>
+		/// The name of the platform
+		/// </summary>
+		[Required]
+		public string Name { get; set; } = null!;
+
+	}
+
+	/// <summary>
+	/// Create device platform response object
+	/// </summary>
+	public class CreateDevicePlatformResponse
+	{
+		/// <summary>
+		/// Id of newly created platform
+		/// </summary>
+		public string Id { get; set; }
+
+		/// <summary>
+		/// Constructor
+		/// </summary>
+		public CreateDevicePlatformResponse(string id)
+		{
+			Id = id;
+		}
+	}
+
+	/// <summary>
+	/// Update requesty object for a device platform
+	/// </summary>
+	public class UpdateDevicePlatformRequest
+	{
+		/// <summary>
+		/// The vendor model ids for the platform
+		/// </summary>
+		public string[] ModelIds { get; set; } = null!;
+	}
+
+	/// <summary>
+	/// Get object response which describes a device platform
+	/// </summary>
+	public class GetDevicePlatformResponse
+	{
+		/// <summary>
+		/// Unique id of device platform
+		/// </summary>
+		public string Id { get; set; }
+
+		/// <summary>
+		/// Friendly name of device platform
+		/// </summary>
+		public string Name { get; set; }
+
+		/// <summary>
+		/// Platform vendor models
+		/// </summary>
+		public string[] ModelIds { get; set; }
+
+		/// <summary>
+		/// Response constructor
+		/// </summary>
+		public GetDevicePlatformResponse(string id, string name, string[] modelIds)
+		{
+			Id = id;
+			Name = name;
+			ModelIds = modelIds;
+		}
+	}
+
+	/// <summary>
+	/// Device pool creation request object
+	/// </summary>
+	public class CreateDevicePoolRequest
+	{
+		/// <summary>
+		/// The name for the new pool
+		/// </summary>
+		[Required]
+		public string Name { get; set; } = null!;
+
+		/// <summary>
+		/// The name for the new pool
+		/// </summary>
+		[Required]
+		public DevicePoolType PoolType { get; set; }
+
+		/// <summary>
+		/// Projects associated with this device pool
+		/// </summary>
+		public List<string>? ProjectIds { get; set; }
+	}
+
+	/// <summary>
+	/// Device pool update request object
+	/// </summary>
+	public class UpdateDevicePoolRequest
+	{
+		/// <summary>
+		/// Id of the device pool to update
+		/// </summary>
+		public string Id { get; set; } = null!;
+
+		/// <summary>
+		/// Projects associated with this device pool
+		/// </summary>
+		public List<string>? ProjectIds { get; set; }
+	}
+
+	/// <summary>
+	/// Device pool creation response object
+	/// </summary>
+	public class CreateDevicePoolResponse
+	{
+		/// <summary>
+		/// Id of the newly created device pool
+		/// </summary>
+		public string Id { get; set; }
+
+		/// <summary>
+		/// Constructor
+		/// </summary>
+		public CreateDevicePoolResponse(string id)
+		{
+			Id = id;
+		}
+	}
+
+	/// <summary>
+	/// Device pool response object
+	/// </summary>
+	public class GetDevicePoolResponse
+	{
+		/// <summary>
+		/// Id of  the device pool
+		/// </summary>
+		public string Id { get; set; }
+
+		/// <summary>
+		/// Name of the device pool
+		/// </summary>
+		public string Name { get; set; }
+
+		/// <summary>
+		/// Type of the device pool
+		/// </summary>
+		public DevicePoolType PoolType { get; set; }
+
+		/// <summary>
+		/// Whether there is write access to the pool
+		/// </summary>
+		public bool WriteAccess { get; set; }
+
+		/// <summary>
+		/// Constructor
+		/// </summary>
+		public GetDevicePoolResponse(string id, string name, DevicePoolType poolType, bool writeAccess)
+		{
+			Id = id;
+			Name = name;
+			PoolType = poolType;
+			WriteAccess = writeAccess;
+		}
+	}
+
+	// Devices 
+
+	/// <summary>
+	/// Device creation request object
+	/// </summary>
+	public class CreateDeviceRequest
+	{
+		/// <summary>
+		/// The platform of the device 
+		/// </summary>
+		[Required]
+		public string PlatformId { get; set; } = null!;
+
+		/// <summary>
+		/// The pool to assign the device
+		/// </summary>
+		[Required]
+		public string PoolId { get; set; } = null!;
+
+		/// <summary>
+		/// The friendly name of the device
+		/// </summary>
+		[Required]
+		public string Name { get; set; } = null!;
+
+		/// <summary>
+		/// Whether to create the device in enabled state
+		/// </summary>
+		public bool? Enabled { get; set; }
+
+		/// <summary>
+		/// The network address of the device
+		/// </summary>
+		public string? Address { get; set; }
+
+		/// <summary>
+		/// The vendor model id of the device
+		/// </summary>
+		public string? ModelId { get; set; }
+	}
+
+	/// <summary>
+	/// Device creation response object
+	/// </summary>
+	public class CreateDeviceResponse
+	{
+		/// <summary>
+		/// The id of the newly created device
+		/// </summary>
+		[Required]
+		public string Id { get; set; } = null!;
+
+		/// <summary>
+		/// Constructor
+		/// </summary>
+		/// <param name="id"></param>
+		public CreateDeviceResponse(string id)
+		{
+			Id = id;
+		}
+	}
+
+	/// <summary>
+	/// Get response object which describes a device
+	/// </summary>
+	public class GetDeviceUtilizationResponse
+	{
+		/// <summary>
+		/// The job id which utilized device
+		/// </summary>
+		public string? JobId { get; set; }
+
+		/// <summary>
+		/// The job's step id
+		/// </summary>
+		public string? StepId { get; set; }
+
+		/// <summary>
+		/// The time device was reserved
+		/// </summary>
+		public DateTime ReservationStartUtc { get; set; }
+
+		/// <summary>
+		/// The time device was freed
+		/// </summary>
+		public DateTime? ReservationFinishUtc { get; set; }
+
+		/// <summary>
+		/// Constructor
+		/// </summary>
+		/// <param name="telemetry"></param>
+		public GetDeviceUtilizationResponse(DeviceUtilizationTelemetry telemetry)
+		{
+			JobId = telemetry.JobId;
+			StepId = telemetry.StepId;
+			ReservationStartUtc = telemetry.ReservationStartUtc;
+			ReservationFinishUtc = telemetry.ReservationFinishUtc;
+		}
+	}
+
+	/// <summary>
+	/// Get response object which describes a device
+	/// </summary>
+	public class GetDeviceResponse
+	{
+		/// <summary>
+		/// The unique id of the device
+		/// </summary>
+		public string Id { get; set; }
+
+		/// <summary>
+		/// The platform of the device
+		/// </summary>
+		public string PlatformId { get; set; }
+
+		/// <summary>
+		/// The pool the device belongs to
+		/// </summary>
+		public string PoolId { get; set; }
+
+		/// <summary>
+		/// The friendly name of the device
+		/// </summary>
+		public string Name { get; set; }
+
+		/// <summary>
+		/// Whether the device is currently enabled
+		/// </summary>
+		public bool Enabled { get; set; }
+
+		/// <summary>
+		///  The address of the device (if it allows network connections)
+		/// </summary>
+		public string? Address { get; set; }
+
+		/// <summary>
+		/// The vendor model id of the device
+		/// </summary>
+		public string? ModelId { get; set; }
+
+		/// <summary>
+		/// Any notes provided for the device
+		/// </summary>
+		public string? Notes { get; set; }
+
+		/// <summary>
+		/// If the device has a marked problem
+		/// </summary>
+		public DateTime? ProblemTime { get; set; }
+
+		/// <summary>
+		/// If the device is in maintenance mode
+		/// </summary>
+		public DateTime? MaintenanceTime { get; set; }
+
+		/// <summary>
+		/// The user id that has the device checked out
+		/// </summary>
+		public string? CheckedOutByUserId { get; set; }
+
+		/// <summary>
+		/// The last time the device was checked out
+		/// </summary>
+		public DateTime? CheckOutTime { get; set; }
+
+		/// <summary>
+		/// The last user to modifiy the device
+		/// </summary>
+		public string? ModifiedByUser { get; set; }
+
+		/// <summary>
+		///  Device Utilization data
+		/// </summary>
+		public List<GetDeviceUtilizationResponse>? Utilization { get; set; }
+
+		/// <summary>
+		/// Device response constructor
+		/// </summary>
+		public GetDeviceResponse(string id, string platformId, string poolId, string name, bool enabled, string? address, string? modelId, string? modifiedByUser, string? notes, DateTime? problemTime, DateTime? maintenanceTime, List<DeviceUtilizationTelemetry>? utilization, string? checkedOutByUser = null, DateTime? checkOutTime = null)
+		{
+			Id = id;
+			Name = name;
+			PlatformId = platformId;
+			PoolId = poolId;
+			Enabled = enabled;
+			Address = address;
+			ModelId = modelId;
+			ModifiedByUser = modifiedByUser;
+			Notes = notes;
+			ProblemTime = problemTime;
+			MaintenanceTime = maintenanceTime;
+			Utilization = utilization?.Select(u => new GetDeviceUtilizationResponse(u)).ToList();
+			CheckedOutByUserId = checkedOutByUser;
+			CheckOutTime = checkOutTime;
+		}
+	}
+
+	/// <summary>
+	/// Device update request object
+	/// </summary>
+	public class UpdateDeviceRequest
+	{
+		/// <summary>
+		/// The device pool id
+		/// </summary>
+		public string? PoolId { get; set; }
+
+		/// <summary>
+		/// The device name
+		/// </summary>
+		public string? Name { get; set; }
+
+		/// <summary>
+		/// IP address or hostname of device
+		/// </summary>
+		public string? Address { get; set; }
+
+		/// <summary>
+		/// Device vendor model id
+		/// </summary>
+		public string? ModelId { get; set; }
+
+		/// <summary>
+		/// Markdown notes
+		/// </summary>
+		public string? Notes { get; set; }
+
+		/// <summary>
+		/// Whether device is enabled
+		/// </summary>
+		public bool? Enabled { get; set; }
+
+		/// <summary>
+		/// Whether the device is in maintenance mode
+		/// </summary>
+		public bool? Maintenance { get; set; }
+
+		/// <summary>
+		/// Whether to set or clear any device problem state
+		/// </summary>
+		public bool? Problem { get; set; }
+	}
+
+	/// <summary>
+	/// Device checkout request object
+	/// </summary>
+	public class CheckoutDeviceRequest
+	{
+		/// <summary>
+		/// Whether to checkout or in the device
+		/// </summary>
+		public bool Checkout { get; set; }
+	}
+
+	/// <summary>
+	/// Device reservation request object
+	/// </summary>
+	public class DeviceReservationRequest
+	{
+		/// <summary>
+		/// Device reservation platform id
+		/// </summary>
+		[Required]
+		public string PlatformId { get; set; } = null!;
+
+		/// <summary>
+		/// The optional vendor model ids to include for this device
+		/// </summary>
+		public List<string>? IncludeModels { get; set; }
+
+		/// <summary>
+		/// The optional vendor model ids to exclude for this device
+		/// </summary>
+		public List<string>? ExcludeModels { get; set; }
+	}
+
+	/// <summary>
+	/// Reservation request object
+	/// </summary>
+	public class CreateDeviceReservationRequest
+	{
+		/// <summary>
+		/// What pool to reserve devices in
+		/// </summary>
+		[Required]
+		public string PoolId { get; set; } = null!;
+
+		/// <summary>
+		/// Devices to reserve
+		/// </summary>
+		[Required]
+		public List<DeviceReservationRequest> Devices { get; set; } = null!;
+	}
+
+	/// <summary>
+	/// Device reservation response object
+	/// </summary>
+	public class CreateDeviceReservationResponse
+	{
+		/// <summary>
+		/// The reservation id of newly created reservation
+		/// </summary>
+		public string Id { get; set; } = null!;
+
+		/// <summary>
+		/// The devices that were reserved
+		/// </summary>
+		public List<GetDeviceResponse> Devices { get; set; } = null!;
+	}
+
+	/// <summary>
+	/// A reservation containing one or more devices
+	/// </summary>
+	public class GetDeviceReservationResponse
+	{
+		/// <summary>
+		/// Randomly generated unique id for this reservation
+		/// </summary>
+		[Required]
+		public string Id { get; set; } = null!;
+
+		/// <summary>
+		/// Which device pool the reservation is in
+		/// </summary>
+		[Required]
+		public string PoolId { get; set; } = null!;
+
+		/// <summary>
+		/// The reserved devices
+		/// </summary>
+		[Required]
+		public List<string> Devices { get; set; } = null!;
+
+		/// <summary>
+		/// JobID holding reservation
+		/// </summary>
+		public string? JobId { get; set; }
+
+		/// <summary>
+		/// Job step id holding reservation
+		/// </summary>
+		public string? StepId { get; set; }
+
+		/// <summary>
+		/// Reservations held by a user, requires a token
+		/// </summary>
+		public string? UserId { get; set; }
+
+		/// <summary>
+		/// The hostname of machine holding reservation
+		/// </summary>
+		public string? Hostname { get; set; }
+
+		/// <summary>
+		/// The optional reservation details 
+		/// </summary>
+		public string? ReservationDetails { get; set; }
+
+		/// <summary>
+		/// The UTC time when the reservation was created
+		/// </summary>
+		public DateTime CreateTimeUtc { get; set; }
+
+		/// <summary>
+		/// The legacy reservation system guid, to be removed once can update Gauntlet client in all streams
+		/// </summary>
+		public string LegacyGuid { get; set; } = null!;
+	}
+
+	// Legacy clients
+
+	/// <summary>
+	/// Reservation request for legacy clients
+	/// </summary>
+	public class LegacyCreateReservationRequest
+	{
+		/// <summary>
+		/// The device types to reserve, these are mapped to platforms
+		/// </summary>
+		[Required]
+		public string[] DeviceTypes { get; set; } = null!;
+
+		/// <summary>
+		/// The hostname of machine reserving devices 
+		/// </summary>
+		[Required]
+		public string Hostname { get; set; } = null!;
+
+		/// <summary>
+		/// The duration of reservation
+		/// </summary>
+		[Required]
+		public string Duration { get; set; } = null!;
+
+		/// <summary>
+		/// Reservation details string
+		/// </summary>
+		public string? ReservationDetails { get; set; } = null;
+
+		/// <summary>
+		/// The PoolId of reservation request 
+		/// </summary>
+		public string? PoolId { get; set; } = null;
+
+		/// <summary>
+		/// The JobId of reservation request 
+		/// </summary>
+		public string? JobId { get; set; } = null;
+
+		/// <summary>
+		/// The StepId of reservation request 
+		/// </summary>
+		public string? StepId { get; set; } = null;
+
+	}
+
+	/// <summary>
+	/// Reservation response for legacy clients
+	/// </summary>
+	public class GetLegacyReservationResponse
+	{
+		/// <summary>
+		/// The names of the devices that were reserved
+		/// </summary>
+		[Required]
+		public string[] DeviceNames { get; set; } = null!;
+
+		/// <summary>
+		/// The corresponding perf specs of the reserved devices
+		/// </summary>
+		[Required]
+		public string[] DevicePerfSpecs { get; set; } = null!;
+
+		/// <summary>
+		/// The corresponding perf specs of the reserved devices
+		/// </summary>
+		[Required]
+		public string[] DeviceModels { get; set; } = null!;
+
+		/// <summary>
+		/// The host name of the machine making the reservation
+		/// </summary>
+		[Required]
+		public string HostName { get; set; } = null!;
+
+		/// <summary>
+		/// The start time of the reservation
+		/// </summary>
+		[Required]
+		public string StartDateTime { get; set; } = null!;
+
+		/// <summary>
+		/// The duration of the reservation (before renew)
+		/// </summary>
+		[Required]
+		public string Duration { get; set; } = null!;
+
+		/// <summary>
+		/// The JobId of reservation request 
+		/// </summary>
+		public string? JobId { get; set; } = null;
+
+		/// <summary>
+		/// The StepId of reservation request 
+		/// </summary>
+		public string? StepId { get; set; } = null;
+
+		/// <summary>
+		/// The legacy guid of the reservation
+		/// </summary>
+		[Required]
+		public string Guid { get; set; } = null!;
+	}
+
+	/// <summary>
+	/// Device response for legacy clients
+	/// </summary>
+	public class GetLegacyDeviceResponse
+	{
+		/// <summary>
+		/// The  name of the reserved device
+		/// </summary>
+		[Required]
+		public string Name { get; set; } = null!;
+
+		/// <summary>
+		/// The (legacy) type of the device, mapped from platform id
+		/// </summary>
+		[Required]
+		public string Type { get; set; } = null!;
+
+		/// <summary>
+		/// The IP or hostname of device
+		/// </summary>
+		[Required]
+		public string IPOrHostName { get; set; } = null!;
+
+		/// <summary>
+		/// The (legacy) perf spec of the device
+		/// </summary>
+		[Required]
+		public string PerfSpec { get; set; } = null!;
+
+		/// <summary>
+		/// The device model information
+		/// </summary>
+		[Required]
+		public string Model { get; set; } = null!;
+
+		/// <summary>
+		/// The available start time which is parsed client side
+		/// </summary>
+		[Required]
+		public string AvailableStartTime { get; set; } = null!;
+
+		/// <summary>
+		/// The available end time which is parsed client side
+		/// </summary>
+		[Required]
+		public string AvailableEndTime { get; set; } = null!;
+
+		/// <summary>
+		/// Whether device is enabled
+		/// </summary>
+		[Required]
+		public bool Enabled { get; set; }
+
+		/// <summary>
+		/// Associated device data
+		/// </summary>
+		[Required]
+		public string DeviceData { get; set; } = null!;
+	}
+
+	/// <summary>
+	/// Platform mapping update for V1 devices
+	/// </summary>
+	public class UpdatePlatformMapRequest
+	{
+		/// <summary>
+		/// Platform name => platform id
+		/// </summary>
+		[Required]
+		public Dictionary<string, string> PlatformMap { get; set; } = null!;
+
+		/// <summary>
+		/// Platform id => name
+		/// </summary>
+		[Required]
+		public Dictionary<string, string> PlatformReverseMap { get; set; } = null!;
+
+		/// <summary>
+		/// Perfspec v1 => model names
+		/// </summary>
+		[Required]
+		public Dictionary<string, string> PerfSpecHighMap { get; set; } = null!;
+	}
+}
diff --git a/Engine/Source/Programs/Horde/Horde.Build/Devices/DevicesController.cs b/Engine/Source/Programs/Horde/Horde.Build/Devices/DevicesController.cs
new file mode 100644
--- /dev/null
+++ b/Engine/Source/Programs/Horde/Horde.Build/Devices/DevicesController.cs
@@ -0,0 +1,1060 @@
+// Copyright Epic Games, Inc. All Rights Reserved.
+
+using System;
+using System.Collections.Generic;
+using System.Linq;
+using System.Text.Json;
+using System.Threading.Tasks;
+using Horde.Build.Acls;
+using Horde.Build.Api;
+using Horde.Build.Collections;
+using Horde.Build.Models;
+using Horde.Build.Services;
+using Horde.Build.Utilities;
+using Microsoft.AspNetCore.Authorization;
+using Microsoft.AspNetCore.Mvc;
+using Microsoft.Extensions.Logging;
+using MongoDB.Bson;
+using MongoDB.Driver;
+
+namespace Horde.Build.Controllers
+{
+	using DeviceId = StringId<IDevice>;
+	using DevicePlatformId = StringId<IDevicePlatform>;
+	using DevicePoolId = StringId<IDevicePool>;
+	using ProjectId = StringId<IProject>;
+
+	/// <summary>
+	/// Controller for device service
+	/// </summary>
+	public class DevicesController : ControllerBase
+	{
+
+		/// <summary>
+		/// The acl service singleton
+		/// </summary>
+		readonly AclService _aclService;
+
+		/// <summary>
+		/// The user collection instance
+		/// </summary>
+		IUserCollection UserCollection { get; set; }
+
+		/// <summary>
+		/// Singleton instance of the device service
+		/// </summary>
+		readonly DeviceService _deviceService;
+
+		/// <summary>
+		///  Logger for controller
+		/// </summary>
+		private readonly ILogger<DevicesController> _logger;
+
+		/// <summary>
+		/// Constructor
+		/// </summary>
+		public DevicesController(DeviceService deviceService, AclService aclService, IUserCollection userCollection, ILogger<DevicesController> logger)
+		{
+			UserCollection = userCollection;
+			_deviceService = deviceService;
+			_logger = logger;
+			_aclService = aclService;
+		}
+
+		// DEVICES
+
+		/// <summary>
+		/// Create a new device
+		/// </summary>
+		[HttpPost]
+		[Authorize]
+		[Route("/api/v2/devices")]
+		public async Task<ActionResult<CreateDeviceResponse>> CreateDeviceAsync([FromBody] CreateDeviceRequest deviceRequest)
+		{
+
+			DevicePoolAuthorization? poolAuth = await _deviceService.GetUserPoolAuthorizationAsync(new DevicePoolId(deviceRequest.PoolId!), User);
+
+			if (poolAuth == null || !poolAuth.Write)
+			{
+				return Forbid();
+			}
+
+			IUser? internalUser = await UserCollection.GetUserAsync(User);
+			if (internalUser == null)
+			{
+				return NotFound();
+			}
+
+			IDevicePlatform? platform = await _deviceService.GetPlatformAsync(new DevicePlatformId(deviceRequest.PlatformId!));
+
+			if (platform == null)
+			{
+				return BadRequest($"Bad platform id {deviceRequest.PlatformId} on request");
+			}
+
+			IDevicePool? pool = await _deviceService.GetPoolAsync(new DevicePoolId(deviceRequest.PoolId!));
+
+			if (pool == null)
+			{
+				return BadRequest($"Bad pool id {deviceRequest.PoolId} on request");
+			}
+
+			string? modelId = null;
+
+			if (!String.IsNullOrEmpty(deviceRequest.ModelId))
+			{
+				modelId = new string(deviceRequest.ModelId);
+
+				if (platform.Models?.FirstOrDefault(x => x == modelId) == null)
+				{
+					return BadRequest($"Bad model id {modelId} for platform {platform.Id} on request");
+				}
+			}
+
+			string name = deviceRequest.Name.Trim();
+			string? address = deviceRequest.Address?.Trim();
+
+			IDevice? device = await _deviceService.TryCreateDeviceAsync(DeviceId.Sanitize(name), name, platform.Id, pool.Id, deviceRequest.Enabled, address, modelId, internalUser.Id);
+
+			if (device == null)
+			{
+				return BadRequest($"Unable to create device");
+			}
+
+			return new CreateDeviceResponse(device.Id.ToString());
+		}
+
+		/// <summary>
+		/// Get list of devices
+		/// </summary>
+		[HttpGet]
+		[Authorize]
+		[Route("/api/v2/devices")]
+		[ProducesResponseType(typeof(List<GetDeviceResponse>), 200)]
+		public async Task<ActionResult<List<object>>> GetDevicesAsync()
+		{
+
+			List<DevicePoolAuthorization> poolAuth = await _deviceService.GetUserPoolAuthorizationsAsync(User);
+
+			List<IDevice> devices = await _deviceService.GetDevicesAsync();
+
+			List<object> responses = new List<object>();
+
+			foreach (IDevice device in devices)
+			{
+				DevicePoolAuthorization? auth = poolAuth.Where(x => x.Pool.Id == device.PoolId).FirstOrDefault();
+
+				if (auth == null || !auth.Read)
+				{
+					continue;
+				}
+
+				responses.Add(new GetDeviceResponse(device.Id.ToString(), device.PlatformId.ToString(), device.PoolId.ToString(), device.Name, device.Enabled, device.Address, device.ModelId?.ToString(), device.ModifiedByUser, device.Notes, device.ProblemTimeUtc, device.MaintenanceTimeUtc, device.Utilization, device.CheckedOutByUser, device.CheckOutTime));
+			}
+
+			return responses;
+		}
+
+		/// <summary>
+		/// Get a specific device
+		/// </summary>
+		[HttpGet]
+		[Authorize]
+		[Route("/api/v2/devices/{deviceId}")]
+		[ProducesResponseType(typeof(GetDeviceResponse), 200)]
+		public async Task<ActionResult<object>> GetDeviceAsync(string deviceId)
+		{
+
+			DeviceId deviceIdValue = new DeviceId(deviceId);
+
+			IDevice? device = await _deviceService.GetDeviceAsync(deviceIdValue);
+
+			if (device == null)
+			{
+				return BadRequest($"Unable to find device with id {deviceId}");
+			}
+
+			DevicePoolAuthorization? poolAuth = await _deviceService.GetUserPoolAuthorizationAsync(device.PoolId, User);
+
+			if (poolAuth == null || !poolAuth.Write)
+			{
+				return Forbid();
+			}
+
+			return new GetDeviceResponse(device.Id.ToString(), device.PlatformId.ToString(), device.PoolId.ToString(), device.Name, device.Enabled, device.Address, device.ModelId?.ToString(), device.ModifiedByUser?.ToString(), device.Notes, device.ProblemTimeUtc, device.MaintenanceTimeUtc, device.Utilization, device.CheckedOutByUser, device.CheckOutTime);
+		}
+
+		/// <summary>
+		/// Update a specific device
+		/// </summary>
+		[HttpPut]
+		[Authorize]
+		[Route("/api/v2/devices/{deviceId}")]
+		[ProducesResponseType(typeof(List<GetDeviceResponse>), 200)]
+		public async Task<ActionResult> UpdateDeviceAsync(string deviceId, [FromBody] UpdateDeviceRequest update)
+		{
+			IUser? internalUser = await UserCollection.GetUserAsync(User);
+			if (internalUser == null)
+			{
+				return NotFound();
+			}
+
+			DeviceId deviceIdValue = new DeviceId(deviceId);
+
+			IDevice? device = await _deviceService.GetDeviceAsync(deviceIdValue);
+
+			if (device == null)
+			{
+				return BadRequest($"Device with id ${deviceId} does not exist");
+			}
+
+			DevicePoolAuthorization? poolAuth = await _deviceService.GetUserPoolAuthorizationAsync(device.PoolId, User);
+
+			if (poolAuth == null || !poolAuth.Write)
+			{
+				return Forbid();
+			}
+
+			IDevicePlatform? platform = await _deviceService.GetPlatformAsync(device.PlatformId);
+
+			if (platform == null)
+			{
+				return BadRequest($"Platform id ${device.PlatformId} does not exist");
+			}
+
+			DevicePoolId? poolIdValue = null;
+
+			if (update.PoolId != null)
+			{
+				poolIdValue = new DevicePoolId(update.PoolId);
+			}
+
+			string? modelIdValue = null;
+
+			if (update.ModelId != null)
+			{
+				modelIdValue = new string(update.ModelId);
+
+				if (platform.Models?.FirstOrDefault(x => x == modelIdValue) == null)
+				{
+					return BadRequest($"Bad model id {update.ModelId} for platform {platform.Id} on request");
+				}
+			}
+
+			string? name = update.Name?.Trim();
+			string? address = update.Address?.Trim();
+
+			await _deviceService.UpdateDeviceAsync(deviceIdValue, poolIdValue, name, address, modelIdValue, update.Notes, update.Enabled, update.Problem, update.Maintenance, internalUser.Id);
+
+			return Ok();
+		}
+
+		/// <summary>
+		/// Checkout a device
+		/// </summary>
+		[HttpPut]
+		[Authorize]
+		[Route("/api/v2/devices/{deviceId}/checkout")]
+		[ProducesResponseType(typeof(List<GetDeviceResponse>), 200)]
+		public async Task<ActionResult> CheckoutDeviceAsync(string deviceId, [FromBody] CheckoutDeviceRequest request)
+		{
+
+			IUser? internalUser = await UserCollection.GetUserAsync(User);
+			if (internalUser == null)
+			{
+				return NotFound();
+			}
+
+			DeviceId deviceIdValue = new DeviceId(deviceId);
+
+			IDevice? device = await _deviceService.GetDeviceAsync(deviceIdValue);
+
+			if (device == null)
+			{
+				return BadRequest($"Device with id ${deviceId} does not exist");
+			}
+
+			DevicePoolAuthorization? poolAuth = await _deviceService.GetUserPoolAuthorizationAsync(device.PoolId, User);
+
+			if (poolAuth == null || !poolAuth.Write)
+			{
+				return Forbid();
+			}
+
+			if (request.Checkout)
+			{
+				if (!String.IsNullOrEmpty(device.CheckedOutByUser))
+				{
+					return BadRequest($"Already checked out by user {device.CheckedOutByUser}");
+				}
+
+				await _deviceService.CheckoutDeviceAsync(deviceIdValue, internalUser.Id);
+
+			}
+			else
+			{
+				await _deviceService.CheckoutDeviceAsync(deviceIdValue, null);
+			}
+
+			return Ok();
+		}
+
+		/// <summary>
+		/// Delete a specific device
+		/// </summary>
+		[HttpDelete]
+		[Authorize]
+		[Route("/api/v2/devices/{deviceId}")]
+		public async Task<ActionResult> DeleteDeviceAsync(string deviceId)
+		{
+			if (!await _deviceService.AuthorizeAsync(AclAction.DeviceWrite, User))
+			{
+				return Forbid();
+			}
+
+			DeviceId deviceIdValue = new DeviceId(deviceId);
+
+			IDevice? device = await _deviceService.GetDeviceAsync(deviceIdValue);
+			if (device == null)
+			{
+				return NotFound();
+			}
+
+			DevicePoolAuthorization? poolAuth = await _deviceService.GetUserPoolAuthorizationAsync(device.PoolId, User);
+
+			if (poolAuth == null || !poolAuth.Write)
+			{
+				return Forbid();
+			}
+
+			await _deviceService.DeleteDeviceAsync(deviceIdValue);
+			return Ok();
+		}
+
+		// PLATFORMS
+
+		/// <summary>
+		/// Create a new device platform
+		/// </summary>
+		[HttpPost]
+		[Authorize]
+		[Route("/api/v2/devices/platforms")]
+		public async Task<ActionResult<CreateDevicePlatformResponse>> CreatePlatformAsync([FromBody] CreateDevicePlatformRequest request)
+		{
+			if (!await _deviceService.AuthorizeAsync(AclAction.DeviceWrite, User))
+			{
+				return Forbid();
+			}
+
+			string name = request.Name.Trim();
+
+			IDevicePlatform? platform = await _deviceService.TryCreatePlatformAsync(DevicePlatformId.Sanitize(name), name);
+
+			if (platform == null)
+			{
+				return BadRequest($"Unable to create platform for {name}");
+			}
+
+			return new CreateDevicePlatformResponse(platform.Id.ToString());
+		}
+
+		/// <summary>
+		/// Update a device platform
+		/// </summary>
+		[HttpPut]
+		[Authorize]
+		[Route("/api/v2/devices/platforms/{platformId}")]
+		public async Task<ActionResult<CreateDevicePlatformResponse>> UpdatePlatformAsync(string platformId, [FromBody] UpdateDevicePlatformRequest request)
+		{
+			if (!await _deviceService.AuthorizeAsync(AclAction.DeviceWrite, User))
+			{
+				return Forbid();
+			}
+
+			await _deviceService.UpdatePlatformAsync(new DevicePlatformId(platformId), request.ModelIds);
+
+			return Ok();
+		}
+
+		/// <summary>
+		/// Get a list of supported device platforms
+		/// </summary>
+		[HttpGet]
+		[Authorize]
+		[Route("/api/v2/devices/platforms")]
+		[ProducesResponseType(typeof(List<GetDevicePlatformResponse>), 200)]
+		public async Task<ActionResult<List<object>>> GetDevicePlatformsAsync()
+		{
+
+			if (!await _deviceService.AuthorizeAsync(AclAction.DeviceRead, User))
+			{
+				return Forbid();
+			}
+
+			List<IDevicePlatform> platforms = await _deviceService.GetPlatformsAsync();
+
+			List<object> responses = new List<object>();
+
+			foreach (IDevicePlatform platform in platforms)
+			{
+				// @todo: ACL per platform
+				responses.Add(new GetDevicePlatformResponse(platform.Id.ToString(), platform.Name, platform.Models?.ToArray() ?? Array.Empty<string>()));
+			}
+
+			return responses;
+		}
+
+		// POOLS
+
+		/// <summary>
+		/// Create a new device pool
+		/// </summary>
+		[HttpPost]
+		[Authorize]
+		[Route("/api/v2/devices/pools")]
+		public async Task<ActionResult<CreateDevicePoolResponse>> CreatePoolAsync([FromBody] CreateDevicePoolRequest request)
+		{
+			if (!await _deviceService.AuthorizeAsync(AclAction.DeviceWrite, User))
+			{
+				return Forbid();
+			}
+
+			string name = request.Name.Trim();
+
+			IDevicePool? pool = await _deviceService.TryCreatePoolAsync(DevicePoolId.Sanitize(name), name, request.PoolType, request.ProjectIds?.Select(x => new ProjectId(x)).ToList());
+
+			if (pool == null)
+			{
+				return BadRequest($"Unable to create pool for {request.Name}");
+			}
+
+			return new CreateDevicePoolResponse(pool.Id.ToString());
+
+		}
+
+		/// <summary>
+		/// Update a device pool
+		/// </summary>
+		[HttpPut]
+		[Authorize]
+		[Route("/api/v2/devices/pools")]
+		public async Task<ActionResult> UpdatePoolAsync([FromBody] UpdateDevicePoolRequest request)
+		{
+
+			DevicePoolAuthorization? poolAuth = await _deviceService.GetUserPoolAuthorizationAsync(new DevicePoolId(request.Id), User);
+
+			if (poolAuth == null || !poolAuth.Write)
+			{
+				return Forbid();
+			}
+
+			List<ProjectId>? projectIds = null;
+
+			if (request.ProjectIds != null)
+			{
+				projectIds = request.ProjectIds.Select(x => new ProjectId(x)).ToList();
+			}
+
+			await _deviceService.UpdatePoolAsync(new DevicePoolId(request.Id), projectIds);
+
+			return Ok();
+		}
+
+		/// <summary>
+		/// Get a list of existing device pools
+		/// </summary>
+		[HttpGet]
+		[Authorize]
+		[Route("/api/v2/devices/pools")]
+		[ProducesResponseType(typeof(List<GetDevicePoolResponse>), 200)]
+		public async Task<ActionResult<List<object>>> GetDevicePoolsAsync()
+		{
+
+			List<DevicePoolAuthorization> poolAuth = await _deviceService.GetUserPoolAuthorizationsAsync(User);
+
+			List<IDevicePool> pools = await _deviceService.GetPoolsAsync();
+
+			List<object> responses = new List<object>();
+
+			foreach (IDevicePool pool in pools)
+			{
+				DevicePoolAuthorization? auth = poolAuth.Where(x => x.Pool.Id == pool.Id).FirstOrDefault();
+
+				if (auth == null || !auth.Read)
+				{
+					continue;
+				}
+
+				responses.Add(new GetDevicePoolResponse(pool.Id.ToString(), pool.Name, pool.PoolType, auth.Write));
+			}
+
+			return responses;
+		}
+
+		// RESERVATIONS
+
+		/// <summary>
+		/// Create a new device reservation
+		/// </summary>
+		[HttpPost]
+		[Authorize]
+		[Route("/api/v2/devices/reservations")]
+		[ProducesResponseType(typeof(CreateDeviceReservationResponse), 200)]
+		public async Task<ActionResult<CreateDeviceReservationResponse>> CreateDeviceReservation([FromBody] CreateDeviceReservationRequest request)
+		{
+
+			if (!await _deviceService.AuthorizeAsync(AclAction.DeviceWrite, User))
+			{
+				return Forbid();
+			}
+
+			List<IDevicePool> pools = await _deviceService.GetPoolsAsync();
+			List<IDevicePlatform> platforms = await _deviceService.GetPlatformsAsync();
+
+			DevicePoolId poolIdValue = new DevicePoolId(request.PoolId);
+			IDevicePool? pool = pools.FirstOrDefault(x => x.Id == poolIdValue);
+			if (pool == null)
+			{
+				return BadRequest($"Unknown pool {request.PoolId} on device reservation request");
+			}
+
+			List<DeviceRequestData> requestedDevices = new List<DeviceRequestData>();
+
+			foreach (DeviceReservationRequest deviceRequest in request.Devices)
+			{
+				DevicePlatformId platformIdValue = new DevicePlatformId(deviceRequest.PlatformId);
+				IDevicePlatform? platform = platforms.FirstOrDefault(x => x.Id == platformIdValue);
+
+				if (platform == null)
+				{
+					return BadRequest($"Unknown platform {deviceRequest.PlatformId} on device reservation request");
+				}
+
+				if (deviceRequest.IncludeModels != null)
+				{
+					foreach (string model in deviceRequest.IncludeModels)
+					{
+						if (platform.Models?.FirstOrDefault(x => x == model) == null)
+						{
+							return BadRequest($"Unknown model {model} for platform {deviceRequest.PlatformId} on device reservation request");
+						}
+					}
+				}
+
+				if (deviceRequest.ExcludeModels != null)
+				{
+					foreach (string model in deviceRequest.ExcludeModels)
+					{
+						if (platform.Models?.FirstOrDefault(x => x == model) == null)
+						{
+							return BadRequest($"Unknown model {model} for platform {deviceRequest.PlatformId} on device reservation request");
+						}
+					}
+				}
+
+				requestedDevices.Add(new DeviceRequestData(platformIdValue, platformIdValue.ToString(), deviceRequest.IncludeModels, deviceRequest.ExcludeModels));
+			}
+
+			IDeviceReservation? reservation = await _deviceService.TryCreateReservationAsync(poolIdValue, requestedDevices);
+
+			if (reservation == null)
+			{
+				return Conflict("Unable to allocated devices for reservation");
+			}
+
+			List<IDevice> devices = await _deviceService.GetDevicesAsync(reservation.Devices);
+
+			CreateDeviceReservationResponse response = new CreateDeviceReservationResponse();
+
+			response.Id = reservation.Id.ToString();
+
+			foreach (IDevice device in devices)
+			{
+				response.Devices.Add(new GetDeviceResponse(device.Id.ToString(), device.PlatformId.ToString(), device.PoolId.ToString(), device.Name, device.Enabled, device.Address, device.ModelId?.ToString(), device.ModifiedByUser, device.Notes, device.ProblemTimeUtc, device.MaintenanceTimeUtc, device.Utilization));
+			}
+
+			return response;
+		}
+
+		/// <summary>
+		/// Get active device reservations
+		/// </summary>		
+		[HttpGet]
+		[Authorize]
+		[Route("/api/v2/devices/reservations")]
+		[ProducesResponseType(typeof(List<GetDeviceReservationResponse>), 200)]
+		public async Task<ActionResult<List<GetDeviceReservationResponse>>> GetDeviceReservations()
+		{
+			if (!await _deviceService.AuthorizeAsync(AclAction.DeviceRead, User))
+			{
+				return Forbid();
+			}
+
+			List<IDeviceReservation> reservations = await _deviceService.GetReservationsAsync();
+
+			List<GetDeviceReservationResponse> response = new List<GetDeviceReservationResponse>();
+
+			foreach (IDeviceReservation reservation in reservations)
+			{
+				GetDeviceReservationResponse reservationResponse = new GetDeviceReservationResponse()
+				{
+					Id = reservation.Id.ToString(),
+					PoolId = reservation.PoolId.ToString(),
+					Devices = reservation.Devices.Select(x => x.ToString()).ToList(),
+					JobId = reservation.JobId?.ToString(),
+					StepId = reservation.StepId?.ToString(),
+					UserId = reservation.UserId?.ToString(),
+					Hostname = reservation.Hostname,
+					ReservationDetails = reservation.ReservationDetails,
+					CreateTimeUtc = reservation.CreateTimeUtc,
+					LegacyGuid = reservation.LegacyGuid
+				};
+
+				response.Add(reservationResponse);
+
+			}
+
+			return response;
+		}
+
+		/// <summary>
+		/// Renew an existing reservation
+		/// </summary>
+		[HttpPut]
+		[Authorize]
+		[Route("/api/v2/devices/reservations/{reservationId}")]
+		public async Task<ActionResult> UpdateReservationAsync(string reservationId)
+		{
+			if (!await _deviceService.AuthorizeAsync(AclAction.DeviceWrite, User))
+			{
+				return Forbid();
+			}
+
+			ObjectId reservationIdValue = new ObjectId(reservationId);
+
+			bool updated = await _deviceService.TryUpdateReservationAsync(reservationIdValue);
+
+			if (!updated)
+			{
+				return BadRequest("Failed to update reservation");
+			}
+
+			return Ok();
+		}
+
+		/// <summary>
+		/// Delete a reservation
+		/// </summary>
+		[HttpDelete]
+		[Authorize]
+		[Route("/api/v2/devices/reservations/{reservationId}")]
+		public async Task<ActionResult> DeleteReservationAsync(string reservationId)
+		{
+			if (!await _deviceService.AuthorizeAsync(AclAction.DeviceWrite, User))
+			{
+				return Forbid();
+			}
+
+			ObjectId reservationIdValue = new ObjectId(reservationId);
+
+			bool deleted = await _deviceService.DeleteReservationAsync(reservationIdValue);
+
+			if (!deleted)
+			{
+				return BadRequest("Failed to delete reservation");
+			}
+
+			return Ok();
+		}
+
+		// LEGACY V1 API
+
+		enum LegacyPerfSpec
+		{
+			Unspecified,
+			Minimum,
+			Recommended,
+			High
+		};
+
+		/// <summary>
+		/// Create a device reservation
+		/// </summary>
+		[HttpPost]
+		[Route("/api/v1/reservations")]
+		[ProducesResponseType(typeof(GetLegacyReservationResponse), 200)]
+		public async Task<ActionResult<GetLegacyReservationResponse>> CreateDeviceReservationV1Async([FromBody] LegacyCreateReservationRequest request)
+		{
+
+			List<IDevicePool> pools = await _deviceService.GetPoolsAsync();
+			List<IDevicePlatform> platforms = await _deviceService.GetPlatformsAsync();
+
+			string? poolId = request.PoolId;
+
+			// @todo: Remove this once all streams are updated to provide jobid
+			string details = "";
+			if ((String.IsNullOrEmpty(request.JobId) || String.IsNullOrEmpty(request.StepId)))
+			{
+				if (!String.IsNullOrEmpty(request.ReservationDetails))
+				{
+					details = $" - {request.ReservationDetails}";
+				}
+				else
+				{
+					details = $" - Host {request.Hostname}";
+				}
+			}
+
+			if (String.IsNullOrEmpty(poolId))
+			{
+				_logger.LogError("No pool specified, defaulting to UE4 {Details} JobId: {JobId}, StepId: {StepId}", details, request.JobId, request.StepId);
+				string message = $"No pool specified, defaulting to UE4" + details;
+				await _deviceService.NotifyDeviceServiceAsync(message, null, request.JobId, request.StepId);
+				poolId = "ue4";
+				//return BadRequest(Message);
+			}
+
+			DevicePoolId poolIdValue = DevicePoolId.Sanitize(poolId);
+			IDevicePool? pool = pools.FirstOrDefault(x => x.Id == poolIdValue);
+			if (pool == null)
+			{
+				_logger.LogError("Unknown pool {PoolId} {Details}", poolId, details);
+				string message = $"Unknown pool {poolId} " + details;
+				await _deviceService.NotifyDeviceServiceAsync(message, null, request.JobId, request.StepId);
+				return BadRequest(message);
+			}
+
+			List<DeviceRequestData> requestedDevices = new List<DeviceRequestData>();
+
+			HashSet<string> legacyPerfSpecs = new HashSet<string> { "Unspecified", "Minimum", "Recommended", "High" };
+
+			List<string> perfSpecs = new List<string>();
+
+			foreach (string deviceType in request.DeviceTypes)
+			{
+
+				string platformName = deviceType;
+				string? constraint = null;
+
+				if (deviceType.Contains(':', StringComparison.Ordinal))
+				{
+					string[] tokens = deviceType.Split(":");
+					platformName = tokens[0];
+					constraint = tokens[1];
+				}
+
+				DevicePlatformId platformId;
+
+				DevicePlatformMapV1 mapV1 = await _deviceService.GetPlatformMapV1();
+
+				if (!mapV1.PlatformMap.TryGetValue(platformName, out platformId))
+				{
+					string message = $"Unknown platform {platformName}" + details;
+					_logger.LogError("Unknown platform {PlatformName} {Details}", platformName, details);
+					await _deviceService.NotifyDeviceServiceAsync(message, null, request.JobId, request.StepId);
+					return BadRequest(message);
+				}
+
+				IDevicePlatform? platform = platforms.FirstOrDefault(x => x.Id == platformId);
+
+				if (platform == null)
+				{
+					string message = $"Unknown platform {platformId}" + details;
+					_logger.LogError("Unknown platform {PlatformId} {Details}", platformId, details);
+					await _deviceService.NotifyDeviceServiceAsync(message, null, request.JobId, request.StepId);
+					return BadRequest(message);
+				}
+
+				List<string> includeModels = new List<string>();
+				List<string> excludeModels = new List<string>();
+
+				if (constraint == null)
+				{
+					perfSpecs.Add("Unspecified");
+				}
+				else if (legacyPerfSpecs.Contains(constraint))
+				{
+					perfSpecs.Add(constraint);
+
+					if (constraint == "High")
+					{
+						string? model = null;
+						if (mapV1.PerfSpecHighMap.TryGetValue(platformId, out model))
+						{
+							includeModels.Add(model);
+						}
+					}
+					else if (constraint == "Minimum" || constraint == "Recommended")
+					{
+						string? model = null;
+						if (mapV1.PerfSpecHighMap.TryGetValue(platformId, out model))
+						{
+							excludeModels.Add(model);
+						}
+					}
+				}
+				else
+				{
+					string? model = platform.Models?.FirstOrDefault(x => x == constraint);
+					if (model == null)
+					{
+						return NotFound($"Platform {platform.Id} has no model {model}");
+					}
+
+					string modelPerfSpec = "Minimum";
+					string? specModel = null;
+
+					if (mapV1.PerfSpecHighMap.TryGetValue(platform.Id, out specModel))
+					{
+						if (model == specModel)
+						{
+							modelPerfSpec = "High";
+						}
+					}
+
+					perfSpecs.Add(modelPerfSpec);
+
+					includeModels.Add(constraint);
+				}
+
+				requestedDevices.Add(new DeviceRequestData(platformId, platformName, includeModels, excludeModels));
+			}
+
+			IDeviceReservation? reservation = await _deviceService.TryCreateReservationAsync(poolIdValue, requestedDevices, request.Hostname, request.ReservationDetails, request.JobId, request.StepId);
+
+			if (reservation == null)
+			{
+				return Conflict("Unable to allocated devices for reservation");
+			}
+
+			List<IDevice> devices = await _deviceService.GetDevicesAsync(reservation.Devices);
+
+			GetLegacyReservationResponse response = new GetLegacyReservationResponse();
+
+			response.Guid = reservation.LegacyGuid;
+			response.DeviceNames = devices.Select(x => x.Name).ToArray();
+			response.DevicePerfSpecs = perfSpecs.ToArray();
+			response.DeviceModels = devices.Select(x => x.ModelId ?? "Base").ToArray();
+			response.HostName = request.Hostname;
+			response.StartDateTime = reservation.CreateTimeUtc.ToString("O", System.Globalization.CultureInfo.InvariantCulture);
+			response.Duration = $"{request.Duration}";
+
+			return new JsonResult(response, new JsonSerializerOptions() { PropertyNamingPolicy = null });
+
+		}
+
+		/// <summary>
+		/// Renew a reservation
+		/// </summary>
+		[HttpPut]
+		[Route("/api/v1/reservations/{reservationGuid}")]
+		[ProducesResponseType(typeof(GetLegacyReservationResponse), 200)]
+		public async Task<ActionResult<GetLegacyReservationResponse>> UpdateReservationV1Async(string reservationGuid /* [FromBody] string Duration */)
+		{
+			IDeviceReservation? reservation = await _deviceService.TryGetReservationFromLegacyGuidAsync(reservationGuid);
+
+			if (reservation == null)
+			{
+				_logger.LogError("Unable to find reservation for legacy guid {ReservationGuid}", reservationGuid);
+				return BadRequest();
+			}
+
+			bool updated = await _deviceService.TryUpdateReservationAsync(reservation.Id);
+
+			if (!updated)
+			{
+				_logger.LogError("Unable to find reservation for reservation {ReservationId}", reservation.Id);
+				return BadRequest();
+			}
+
+			List<IDevice> devices = await _deviceService.GetDevicesAsync(reservation.Devices);
+
+			GetLegacyReservationResponse response = new GetLegacyReservationResponse();
+
+			response.Guid = reservation.LegacyGuid;
+			response.DeviceNames = reservation.Devices.Select(deviceId => devices.First(d => d.Id == deviceId).Name).ToArray();
+			response.HostName = reservation.Hostname ?? "";
+			response.StartDateTime = reservation.CreateTimeUtc.ToString("O", System.Globalization.CultureInfo.InvariantCulture);
+			response.Duration = "00:10:00"; // matches gauntlet duration
+
+			return new JsonResult(response, new JsonSerializerOptions() { PropertyNamingPolicy = null });
+		}
+
+		/// <summary>
+		/// Delete a reservation
+		/// </summary>
+		[HttpDelete]
+		[Route("/api/v1/reservations/{reservationGuid}")]
+		public async Task<ActionResult> DeleteReservationV1Async(string reservationGuid)
+		{
+			IDeviceReservation? reservation = await _deviceService.TryGetReservationFromLegacyGuidAsync(reservationGuid);
+
+			if (reservation == null)
+			{
+				return BadRequest($"Unable to find reservation for guid {reservationGuid}");
+			}
+
+			bool deleted = await _deviceService.DeleteReservationAsync(reservation.Id);
+
+			if (!deleted)
+			{
+				return BadRequest("Failed to delete reservation");
+			}
+
+			return Ok();
+		}
+
+		/// <summary>
+		/// Get device info for a reserved device
+		/// </summary>
+		[HttpGet]
+		[Route("/api/v1/devices/{deviceName}")]
+		[ProducesResponseType(typeof(GetLegacyDeviceResponse), 200)]
+		public async Task<ActionResult<GetLegacyDeviceResponse>> GetDeviceV1Async(string deviceName)
+		{
+			IDevice? device = await _deviceService.GetDeviceByNameAsync(deviceName);
+
+			if (device == null)
+			{
+				return BadRequest($"Unknown device {deviceName}");
+			}
+
+			IDeviceReservation? reservation = await _deviceService.TryGetDeviceReservation(device.Id);
+
+			string? platformName = null;
+
+			if (reservation != null && reservation.Devices.Count == reservation.RequestedDevicePlatforms.Count)
+			{
+				for (int i = 0; i < reservation.Devices.Count; i++)
+				{
+					if (reservation.Devices[i] == device.Id)
+					{
+						platformName = reservation.RequestedDevicePlatforms[i];
+					}
+				}
+			}
+
+			DevicePlatformMapV1 mapV1 = await _deviceService.GetPlatformMapV1();
+
+			if (String.IsNullOrEmpty(platformName))
+			{
+				if (!mapV1.PlatformReverseMap.TryGetValue(device.PlatformId, out platformName))
+				{
+					return BadRequest($"Unable to map platform for {deviceName} : {device.PlatformId}");
+				}
+			}
+
+			if (String.IsNullOrEmpty(platformName))
+			{
+				return BadRequest($"Unable to get platform for {deviceName} from reservation or mapping : {device.PlatformId}");
+			}
+
+			GetLegacyDeviceResponse response = new GetLegacyDeviceResponse();
+
+			response.Name = device.Name;
+			response.Type = platformName;
+			response.IPOrHostName = device.Address ?? "";
+			response.AvailableStartTime = "00:00:00";
+			response.AvailableEndTime = "00:00:00";
+			response.Enabled = true;
+			response.Model = device.ModelId ?? "Base";
+			response.DeviceData = "";
+			response.PerfSpec = "Minimum";
+
+			if (device.ModelId != null)
+			{
+				string? model;
+				if (mapV1.PerfSpecHighMap.TryGetValue(device.PlatformId, out model))
+				{
+					if (model == device.ModelId)
+					{
+						response.PerfSpec = "High";
+					}
+				}
+			}
+
+			return new JsonResult(response, new JsonSerializerOptions() { PropertyNamingPolicy = null });
+		}
+
+		/// <summary>
+		/// Mark a problem device
+		/// </summary>
+		[HttpPut]
+		[Route("/api/v1/deviceerror/{deviceName}")]
+		public async Task<ActionResult> PutDeviceErrorAsync(string deviceName)
+		{
+			IDevice? device = await _deviceService.GetDeviceByNameAsync(deviceName);
+
+			if (device == null)
+			{
+				_logger.LogError("Device error reported for unknown device {DeviceName}", deviceName);
+				return BadRequest($"Unknown device {deviceName}");
+			}
+
+			await _deviceService.UpdateDeviceAsync(device.Id, newProblem: true);
+
+			string message = $"Device problem, {device.Name} : {device.PoolId.ToString().ToUpperInvariant()}";
+
+			IDeviceReservation? reservation = await _deviceService.TryGetDeviceReservation(device.Id);
+
+			string? jobId = null;
+			string? stepId = null;
+			if (reservation != null)
+			{
+				jobId = !String.IsNullOrEmpty(reservation.JobId) ? reservation.JobId : null;
+				stepId = !String.IsNullOrEmpty(reservation.StepId) ? reservation.StepId : null;
+
+				if ((jobId == null || stepId == null))
+				{
+					if (!String.IsNullOrEmpty(reservation.ReservationDetails))
+					{
+						message += $" - {reservation.ReservationDetails}";
+					}
+					else
+					{
+						message += $" - Host {reservation.Hostname}";
+					}
+				}
+			}
+
+			_logger.LogError("{DeviceError}", message);
+
+			await _deviceService.NotifyDeviceServiceAsync(message, device.Id, jobId, stepId);
+
+			return Ok();
+
+		}
+
+		/// <summary>
+		/// Updates the platform map for v1 requests
+		/// </summary>
+		[HttpPut]
+		[Route("/api/v1/devices/platformmap")]
+		public async Task<ActionResult> UpdatePlatformMapV1([FromBody] UpdatePlatformMapRequest request)
+		{
+			if (!await _aclService.AuthorizeAsync(AclAction.AdminWrite, User))
+			{
+				return Forbid();
+			}
+
+			bool result;
+			try
+			{
+				result = await _deviceService.UpdatePlatformMapAsync(request);
+			}
+			catch (Exception ex)
+			{
+				_logger.LogError(ex, "Error updating device platform map {Message}", ex.Message);
+				throw;
+			}
+
+			if (!result)
+			{
+				_logger.LogError("Unable to update device platform mapping");
+				return BadRequest();
+			}
+
+			return Ok();
+		}
+	}
+}
diff --git a/Engine/Source/Programs/Horde/Horde.Build/Issues/Handlers/CompileIssueHandler.cs b/Engine/Source/Programs/Horde/Horde.Build/Issues/Handlers/CompileIssueHandler.cs
new file mode 100644
--- /dev/null
+++ b/Engine/Source/Programs/Horde/Horde.Build/Issues/Handlers/CompileIssueHandler.cs
@@ -0,0 +1,76 @@
+// Copyright Epic Games, Inc. All Rights Reserved.
+
+using System;
+using System.Collections.Generic;
+using System.Diagnostics.CodeAnalysis;
+using System.Linq;
+using EpicGames.Core;
+using Horde.Build.Collections;
+using Horde.Build.Models;
+using Microsoft.Extensions.Logging;
+using MongoDB.Driver;
+
+namespace Horde.Build.IssueHandlers.Impl
+{
+	/// <summary>
+	/// Instance of a particular compile error
+	/// </summary>
+	class CompileIssueHandler : SourceFileIssueHandler
+	{
+		/// <summary>
+		/// Annotation describing the compile type
+		/// </summary>
+		const string CompileTypeAnnotation = "CompileType";
+
+		/// <inheritdoc/>
+		public override string Type => "Compile";
+
+		/// <inheritdoc/>
+		public override int Priority => 10;
+
+		/// <summary>
+		/// Determines if the given event id matches
+		/// </summary>
+		/// <param name="eventId">The event id to compare</param>
+		/// <returns>True if the given event id matches</returns>
+		public static bool IsMatchingEventId(EventId? eventId)
+		{
+			return eventId == KnownLogEvents.Compiler || eventId == KnownLogEvents.AutomationTool_SourceFileLine || eventId == KnownLogEvents.MSBuild;
+		}
+
+		/// <inheritdoc/>
+		public override bool TryGetFingerprint(IJob job, INode node, IReadOnlyNodeAnnotations annotations, ILogEventData eventData, [NotNullWhen(true)] out NewIssueFingerprint? fingerprint)
+		{
+			if (!IsMatchingEventId(eventData.EventId))
+			{
+				fingerprint = null;
+				return false;
+			}
+
+			List<string> newFileNames = new List<string>();
+			GetSourceFiles(eventData, newFileNames);
+
+			string? compileType;
+			if (!annotations.TryGetValue(CompileTypeAnnotation, out compileType))
+			{
+				compileType = "Compile";
+			}
+
+			List<string> newMetadata = new List<string>();
+			newMetadata.Add($"{CompileTypeAnnotation}={compileType}");
+
+			fingerprint = new NewIssueFingerprint(Type, newFileNames, null, newMetadata);
+			return true;
+		}
+
+		/// <inheritdoc/>
+		public override string GetSummary(IIssueFingerprint fingerprint, IssueSeverity severity)
+		{
+			List<string> types = fingerprint.GetMetadataValues(CompileTypeAnnotation).ToList();
+			string type = (types.Count == 1) ? types[0] : "Compile";
+			string level = (severity == IssueSeverity.Warning) ? "warnings" : "errors";
+			string list = StringUtils.FormatList(fingerprint.Keys.Where(x => !x.StartsWith(NotePrefix, StringComparison.Ordinal)).ToArray(), 2);
+			return $"{type} {level} in {list}";
+		}
+	}
+}
