Horde: Finish re-org of Horde solution started a while ago.

diff --git a/Engine/Source/Programs/Horde/Horde.Build/Issues/Handlers/ContentIssueHandler.cs b/Engine/Source/Programs/Horde/Horde.Build/Issues/Handlers/ContentIssueHandler.cs
new file mode 100644
--- /dev/null
+++ b/Engine/Source/Programs/Horde/Horde.Build/Issues/Handlers/ContentIssueHandler.cs
@@ -0,0 +1,98 @@
+// Copyright Epic Games, Inc. All Rights Reserved.
+
+using System;
+using System.Collections.Generic;
+using System.Diagnostics.CodeAnalysis;
+using System.Linq;
+using EpicGames.Core;
+using Horde.Build.Collections;
+using Horde.Build.Models;
+using Horde.Build.Utilities;
+using Microsoft.Extensions.Logging;
+using MongoDB.Driver;
+
+namespace Horde.Build.IssueHandlers.Impl
+{
+	/// <summary>
+	/// Instance of a particular compile error
+	/// </summary>
+	class ContentIssueHandler : IIssueHandler
+	{
+		/// <inheritdoc/>
+		public string Type => "Content";
+
+		/// <inheritdoc/>
+		public int Priority => 10;
+
+		/// <summary>
+		/// Determines if the given event id matches
+		/// </summary>
+		/// <param name="eventId">The event id to compare</param>
+		/// <returns>True if the given event id matches</returns>
+		public static bool IsMatchingEventId(EventId? eventId)
+		{
+			return eventId == KnownLogEvents.Engine_AssetLog;
+		}
+
+		/// <summary>
+		/// Adds all the assets from the given log event
+		/// </summary>
+		/// <param name="eventData">The log event to parse</param>
+		/// <param name="assetNames">Receives the referenced asset names</param>
+		public static void GetAssetNames(ILogEventData eventData, HashSet<string> assetNames)
+		{
+			foreach (ILogEventLine line in eventData.Lines)
+			{
+				string? relativePath;
+				if (line.Data.TryGetNestedProperty("properties.asset.relativePath", out relativePath))
+				{
+					int endIdx = relativePath.LastIndexOfAny(new char[] { '/', '\\' }) + 1;
+					string fileName = relativePath.Substring(endIdx);
+					assetNames.Add(fileName);
+				}
+			}
+		}
+
+		/// <inheritdoc/>
+		public bool TryGetFingerprint(IJob job, INode node, IReadOnlyNodeAnnotations annotations, ILogEventData eventData, [NotNullWhen(true)] out NewIssueFingerprint? fingerprint)
+		{
+			if (!IsMatchingEventId(eventData.EventId))
+			{
+				fingerprint = null;
+				return false;
+			}
+
+			HashSet<string> newAssetNames = new HashSet<string>();
+			GetAssetNames(eventData, newAssetNames);
+			fingerprint = new NewIssueFingerprint(Type, newAssetNames, null, null);
+			return true;
+		}
+
+		/// <inheritdoc/>
+		public string GetSummary(IIssueFingerprint fingerprint, IssueSeverity severity)
+		{
+			string type = (severity == IssueSeverity.Warning) ? "Warnings" : "Errors";
+			string list = StringUtils.FormatList(fingerprint.Keys.ToArray(), 2);
+			return $"{type} in {list}";
+		}
+
+		/// <inheritdoc/>
+		public void RankSuspects(IIssueFingerprint fingerprint, List<SuspectChange> suspects)
+		{
+			foreach (SuspectChange suspect in suspects)
+			{
+				if (suspect.ContainsContent)
+				{
+					if (suspect.Details.Files.Any(x => fingerprint.Keys.Any(y => x.Path.Contains(y, StringComparison.OrdinalIgnoreCase))))
+					{
+						suspect.Rank += 20;
+					}
+					else
+					{
+						suspect.Rank += 10;
+					}
+				}
+			}
+		}
+	}
+}
diff --git a/Engine/Source/Programs/Horde/Horde.Build/Issues/Handlers/CopyrightIssueHandler.cs b/Engine/Source/Programs/Horde/Horde.Build/Issues/Handlers/CopyrightIssueHandler.cs
new file mode 100644
--- /dev/null
+++ b/Engine/Source/Programs/Horde/Horde.Build/Issues/Handlers/CopyrightIssueHandler.cs
@@ -0,0 +1,101 @@
+// Copyright Epic Games, Inc. All Rights Reserved.
+
+using System;
+using System.Collections.Generic;
+using System.Diagnostics.CodeAnalysis;
+using System.Linq;
+using EpicGames.Core;
+using Horde.Build.Collections;
+using Horde.Build.Models;
+using Horde.Build.Utilities;
+using Microsoft.Extensions.Logging;
+using MongoDB.Driver;
+
+namespace Horde.Build.IssueHandlers.Impl
+{
+	/// <summary>
+	/// Instance of a particular compile error
+	/// </summary>
+	class CopyrightIssueHandler : IIssueHandler
+	{
+		/// <inheritdoc/>
+		public string Type => "Copyright";
+
+		/// <inheritdoc/>
+		public int Priority => 10;
+
+		/// <summary>
+		/// Filenames containing errors or warnings
+		/// </summary>
+		public HashSet<string> FileNames { get; set; } = new HashSet<string>();
+
+		/// <summary>
+		/// Determines if the given event id matches
+		/// </summary>
+		/// <param name="eventId">The event id to compare</param>
+		/// <returns>True if the given event id matches</returns>
+		public static bool IsMatchingEventId(EventId? eventId)
+		{
+			return eventId == KnownLogEvents.AutomationTool_MissingCopyright;
+		}
+
+		/// <summary>
+		/// Extracts a list of source files from an event
+		/// </summary>
+		/// <param name="logEventData">The event data</param>
+		/// <param name="sourceFiles">List of source files</param>
+		public static void GetSourceFiles(ILogEventData logEventData, HashSet<string> sourceFiles)
+		{
+			foreach (ILogEventLine line in logEventData.Lines)
+			{
+				string? relativePath;
+				if (line.Data.TryGetNestedProperty("properties.file.relativePath", out relativePath) || line.Data.TryGetNestedProperty("properties.file", out relativePath))
+				{
+					int endIdx = relativePath.LastIndexOfAny(new char[] { '/', '\\' }) + 1;
+					string fileName = relativePath.Substring(endIdx);
+					sourceFiles.Add(fileName);
+				}
+			}
+		}
+
+		/// <inheritdoc/>
+		public bool TryGetFingerprint(IJob job, INode node, IReadOnlyNodeAnnotations annotations, ILogEventData eventData, [NotNullWhen(true)] out NewIssueFingerprint? fingerprint)
+		{
+			if (!IsMatchingEventId(eventData.EventId))
+			{
+				fingerprint = null;
+				return false;
+			}
+
+			HashSet<string> newFileNames = new HashSet<string>(StringComparer.OrdinalIgnoreCase);
+			GetSourceFiles(eventData, newFileNames);
+			fingerprint = new NewIssueFingerprint(Type, newFileNames, null, null);
+			return true;
+		}
+
+		/// <inheritdoc/>
+		public void RankSuspects(IIssueFingerprint fingerprint, List<SuspectChange> suspects)
+		{
+			foreach (SuspectChange suspect in suspects)
+			{
+				if (suspect.ContainsCode)
+				{
+					if (FileNames.Any(x => suspect.ModifiesFile(x)))
+					{
+						suspect.Rank += 20;
+					}
+					else
+					{
+						suspect.Rank += 10;
+					}
+				}
+			}
+		}
+
+		/// <inheritdoc/>
+		public string GetSummary(IIssueFingerprint fingerprint, IssueSeverity severity)
+		{
+			return $"Missing copyright notice in {StringUtils.FormatList(fingerprint.Keys.ToArray(), 2)}";
+		}
+	}
+}
diff --git a/Engine/Source/Programs/Horde/Horde.Build/Issues/Handlers/DefaultIssueHandler.cs b/Engine/Source/Programs/Horde/Horde.Build/Issues/Handlers/DefaultIssueHandler.cs
new file mode 100644
--- /dev/null
+++ b/Engine/Source/Programs/Horde/Horde.Build/Issues/Handlers/DefaultIssueHandler.cs
@@ -0,0 +1,53 @@
+// Copyright Epic Games, Inc. All Rights Reserved.
+
+using System.Collections.Generic;
+using System.Diagnostics.CodeAnalysis;
+using System.Linq;
+using Horde.Build.Collections;
+using Horde.Build.Models;
+
+namespace Horde.Build.IssueHandlers.Impl
+{
+	/// <summary>
+	/// Instance of a particular compile error
+	/// </summary>
+	class DefaultIssueHandler : IIssueHandler
+	{
+		/// <summary>
+		/// Name of the handler
+		/// </summary>
+		public const string Type = "Default";
+
+		/// <inheritdoc/>
+		string IIssueHandler.Type => Type;
+
+		/// <inheritdoc/>
+		public int Priority => 0;
+
+		/// <inheritdoc/>
+		public bool TryGetFingerprint(IJob job, INode node, IReadOnlyNodeAnnotations annotations, ILogEventData eventData, [NotNullWhen(true)] out NewIssueFingerprint? fingerprint)
+		{
+			fingerprint = new NewIssueFingerprint(Type, new[] { node.Name }, null, null);
+			return true;
+		}
+
+		/// <inheritdoc/>
+		public void RankSuspects(IIssueFingerprint fingerprint, List<SuspectChange> suspects)
+		{
+		}
+
+		/// <inheritdoc/>
+		public string GetSummary(IIssueFingerprint fingerprint, IssueSeverity severity)
+		{
+			string nodeName = fingerprint.Keys.FirstOrDefault() ?? "(unknown)";
+			if(severity == IssueSeverity.Warning)
+			{
+				return $"Warnings in {nodeName}";
+			}
+			else
+			{
+				return $"Errors in {nodeName}";
+			}
+		}
+	}
+}
diff --git a/Engine/Source/Programs/Horde/Horde.Build/Issues/Handlers/GauntletIssueHandler.cs b/Engine/Source/Programs/Horde/Horde.Build/Issues/Handlers/GauntletIssueHandler.cs
new file mode 100644
--- /dev/null
+++ b/Engine/Source/Programs/Horde/Horde.Build/Issues/Handlers/GauntletIssueHandler.cs
@@ -0,0 +1,171 @@
+// Copyright Epic Games, Inc. All Rights Reserved.
+
+using System;
+using System.Collections.Generic;
+using System.Diagnostics.CodeAnalysis;
+using System.Linq;
+using System.Text.Json;
+using EpicGames.Core;
+using Horde.Build.Collections;
+using Horde.Build.Models;
+using Horde.Build.Utilities;
+using Microsoft.Extensions.Logging;
+
+namespace Horde.Build.IssueHandlers.Impl
+{
+	/// <summary>
+	/// Instance of a particular Gauntlet error
+	/// </summary>
+	class GauntletIssueHandler : IIssueHandler
+	{
+		/// <summary>
+		/// Prefix for unit test keys
+		/// </summary>
+		const string UnitTestPrefix = "UnitTest:";
+
+		/// <summary>
+		/// Prefix for screenshot test keys
+		/// </summary>
+		const string ScreenshotTestPrefix = "Screenshot:";
+
+		/// <inheritdoc/>
+		public string Type => "Gauntlet";
+
+		/// <inheritdoc/>
+		public int Priority => 10;
+
+		/// <summary>
+		/// Determines if the given event id matches
+		/// </summary>
+		/// <param name="eventId">The event id to compare</param>
+		/// <returns>True if the given event id matches</returns>
+		public static bool IsMatchingEventId(EventId? eventId)
+		{
+			return eventId == KnownLogEvents.Gauntlet_UnitTest || eventId == KnownLogEvents.Gauntlet_ScreenshotTest;
+		}
+
+		/// <summary>
+		/// Parses symbol names from a log event
+		/// </summary>
+		/// <param name="eventData">The log event data</param>
+		/// <param name="unitTestNames">Receives a set of the unit test names</param>
+		private static void GetUnitTestNames(ILogEventData eventData, HashSet<string> unitTestNames)
+		{
+			foreach (ILogEventLine line in eventData.Lines)
+			{
+				string? group = null;
+				string? name = null;
+
+				string? value;
+				if (line.Data.TryGetNestedProperty("properties.group", out value))
+				{
+					group = value;
+				}
+				if (line.Data.TryGetNestedProperty("properties.name", out value))
+				{
+					name = value;
+				}
+
+				if (group != null && name != null)
+				{
+					unitTestNames.Add($"{UnitTestPrefix}:{group}/{name}");
+				}
+			}
+		}
+
+		/// <summary>
+		/// Parses screenshot test names from a log event
+		/// </summary>
+		/// <param name="eventData">The event data</param>
+		/// <param name="screenshotTestNames">Receives the parsed screenshot test names</param>
+		private static void GetScreenshotTestNames(ILogEventData eventData, HashSet<string> screenshotTestNames)
+		{
+			foreach (ILogEventLine line in eventData.Lines)
+			{
+				if (line.Data.TryGetProperty("properties", JsonValueKind.Object, out JsonElement properties))
+				{
+					foreach (JsonProperty property in properties.EnumerateObject())
+					{
+						JsonElement value = property.Value;
+						if (value.ValueKind == JsonValueKind.Object && value.HasStringProperty("type", "Screenshot") && value.TryGetStringProperty("name", out string? name))
+						{
+							screenshotTestNames.Add($"{ScreenshotTestPrefix}:{name}");
+						}
+					}
+				}
+			}
+		}
+
+		/// <inheritdoc/>
+		public bool TryGetFingerprint(IJob job, INode node, IReadOnlyNodeAnnotations annotations, ILogEventData eventData, [NotNullWhen(true)] out NewIssueFingerprint? fingerprint)
+		{
+			if(!IsMatchingEventId(eventData.EventId))
+			{
+				fingerprint = null;
+				return false;
+			}
+
+			HashSet<string> keys = new HashSet<string>();
+			GetUnitTestNames(eventData, keys);
+			GetScreenshotTestNames(eventData, keys);
+
+			if (keys.Count == 0)
+			{
+				fingerprint = null;
+				return false;
+			}
+
+			fingerprint = new NewIssueFingerprint(Type, keys, null, null);
+			return true;
+		}
+
+		/// <inheritdoc/>
+		public void RankSuspects(IIssueFingerprint fingerprint, List<SuspectChange> changes)
+		{
+			foreach (SuspectChange change in changes)
+			{
+				if (change.ContainsCode)
+				{
+					change.Rank += 10;
+				}
+			}
+		}
+
+		/// <inheritdoc/>
+		public string GetSummary(IIssueFingerprint fingerprint, IssueSeverity severity)
+		{
+			List<string> unitTestNames = fingerprint.Keys.Where(x => x.StartsWith(UnitTestPrefix, StringComparison.Ordinal)).Select(x => x.Substring(UnitTestPrefix.Length + 1)).ToList();
+			if (unitTestNames.Count > 0)
+			{
+				HashSet<string> groupNames = new HashSet<string>(StringComparer.OrdinalIgnoreCase);
+				HashSet<string> testNames = new HashSet<string>(StringComparer.OrdinalIgnoreCase);
+				foreach (string unitTestName in unitTestNames)
+				{
+					int idx = unitTestName.IndexOf('/', StringComparison.OrdinalIgnoreCase);
+					if (idx != -1)
+					{
+						groupNames.Add(unitTestName.Substring(0, idx));
+					}
+					testNames.Add(unitTestName.Substring(idx + 1));
+				}
+
+				if (groupNames.Count == 1)
+				{
+					return $"{groupNames.First()} test failures: {StringUtils.FormatList(testNames.ToArray(), 3)}";
+				}
+				else
+				{
+					return $"{StringUtils.FormatList(groupNames.OrderBy(x => x).ToArray(), 100)} test failures";
+				}
+			}
+
+			List<string> screenshotTestNames = fingerprint.Keys.Where(x => x.StartsWith(ScreenshotTestPrefix, StringComparison.Ordinal)).Select(x => x.Substring(ScreenshotTestPrefix.Length + 1)).ToList();
+			if (screenshotTestNames.Count > 0)
+			{
+				return $"Screenshot test failures: {StringUtils.FormatList(screenshotTestNames.ToArray(), 3)}";
+			}
+
+			return "Test failures";
+		}
+	}
+}
diff --git a/Engine/Source/Programs/Horde/Horde.Build/Issues/Handlers/SymbolIssueHandler.cs b/Engine/Source/Programs/Horde/Horde.Build/Issues/Handlers/SymbolIssueHandler.cs
new file mode 100644
--- /dev/null
+++ b/Engine/Source/Programs/Horde/Horde.Build/Issues/Handlers/SymbolIssueHandler.cs
@@ -0,0 +1,131 @@
+// Copyright Epic Games, Inc. All Rights Reserved.
+
+using System;
+using System.Collections.Generic;
+using System.Diagnostics.CodeAnalysis;
+using System.Linq;
+using System.Text;
+using EpicGames.Core;
+using Horde.Build.Collections;
+using Horde.Build.Models;
+using Horde.Build.Utilities;
+using Microsoft.Extensions.Logging;
+
+namespace Horde.Build.IssueHandlers.Impl
+{
+	/// <summary>
+	/// Instance of a particular compile error
+	/// </summary>
+	class SymbolIssueHandler : IIssueHandler
+	{
+		const string NodeName = "Node";
+		const string EventIdName = "EventId";
+
+		static readonly string DuplicateEventIdMetadata = $"{EventIdName}={KnownLogEvents.Linker_DuplicateSymbol.Id}";
+
+		/// <inheritdoc/>
+		public string Type => "Symbol";
+
+		/// <inheritdoc/>
+		public int Priority => 10;
+
+		/// <summary>
+		/// Determines if the given event id matches
+		/// </summary>
+		/// <param name="eventId">The event id to compare</param>
+		/// <returns>True if the given event id matches</returns>
+		public static bool IsMatchingEventId(EventId? eventId)
+		{
+			return eventId == KnownLogEvents.Linker_UndefinedSymbol || eventId == KnownLogEvents.Linker_DuplicateSymbol || eventId == KnownLogEvents.Linker;
+		}
+
+		/// <summary>
+		/// Parses symbol names from a log event
+		/// </summary>
+		/// <param name="eventData">The log event data</param>
+		/// <param name="symbolNames">Receives the list of symbol names</param>
+		public static void GetSymbolNames(ILogEventData eventData, SortedSet<string> symbolNames)
+		{
+			foreach (ILogEventLine line in eventData.Lines)
+			{
+				string? identifier;
+				if (line.Data.TryGetNestedProperty("properties.symbol.identifier", out identifier))
+				{
+					symbolNames.Add(identifier);
+				}
+			}
+		}
+
+		/// <inheritdoc/>
+		public void RankSuspects(IIssueFingerprint fingerprint, List<SuspectChange> changes)
+		{
+			HashSet<string> names = new HashSet<string>();
+			foreach (string name in fingerprint.Keys)
+			{
+				names.UnionWith(name.Split("::", StringSplitOptions.RemoveEmptyEntries));
+			}
+
+			foreach (SuspectChange change in changes)
+			{
+				if (change.ContainsCode)
+				{
+					int matches = names.Count(x => change.Details.Files.Any(y => y.Path.Contains(x, StringComparison.OrdinalIgnoreCase)));
+					change.Rank += 10 + (10 * matches);
+				}
+			}
+		}
+
+		/// <inheritdoc/>
+		public string GetSummary(IIssueFingerprint fingerprint, IssueSeverity severity)
+		{
+			HashSet<string> symbols = fingerprint.Keys;
+			if (symbols.Count == 0)
+			{
+				string[] nodes = fingerprint.GetMetadataValues(NodeName).ToArray();
+
+				StringBuilder summary = new StringBuilder("Linker ");
+				summary.Append((severity == IssueSeverity.Warning) ? "warnings" : "errors");
+				if (nodes.Length > 0)
+				{
+					summary.Append($" in {StringUtils.FormatList(nodes, 2)}");
+				}
+
+				return summary.ToString();
+			}
+			else
+			{
+				string problemType = fingerprint.HasMetadataValue(DuplicateEventIdMetadata) ? "Duplicate" : "Undefined";
+				if (symbols.Count == 1)
+				{
+					return $"{problemType} symbol '{symbols.First()}'";
+				}
+				else
+				{
+					return $"{problemType} symbols: {StringUtils.FormatList(symbols.ToArray(), 3)}";
+				}
+			}
+		}
+
+		public bool TryGetFingerprint(IJob job, INode node, IReadOnlyNodeAnnotations annotations, ILogEventData eventData, [NotNullWhen(true)] out NewIssueFingerprint? fingerprint)
+		{
+			if (!IsMatchingEventId(eventData.EventId))
+			{
+				fingerprint = null;
+				return false;
+			}
+
+			SortedSet<string> symbolNames = new SortedSet<string>();
+			GetSymbolNames(eventData, symbolNames);
+
+			List<string> metadata = new List<string>();
+			metadata.Add($"{NodeName}={node.Name}");
+			if (eventData.EventId != null)
+			{
+				metadata.Add($"{EventIdName}={eventData.EventId.Value.Id}");
+			}
+
+			fingerprint = new NewIssueFingerprint(Type, symbolNames, null, metadata);
+			return true;
+		}
+	}
+}
diff --git a/Engine/Source/Programs/Horde/Horde.Build/Issues/Handlers/SystemicIssueHandler.cs b/Engine/Source/Programs/Horde/Horde.Build/Issues/Handlers/SystemicIssueHandler.cs
new file mode 100644
--- /dev/null
+++ b/Engine/Source/Programs/Horde/Horde.Build/Issues/Handlers/SystemicIssueHandler.cs
@@ -0,0 +1,87 @@
+// Copyright Epic Games, Inc. All Rights Reserved.
+
+using System;
+using System.Collections.Generic;
+using System.Diagnostics.CodeAnalysis;
+using System.Linq;
+using EpicGames.Core;
+using Horde.Build.Collections;
+using Horde.Build.Models;
+using Microsoft.Extensions.Logging;
+using MongoDB.Driver;
+
+namespace Horde.Build.IssueHandlers.Impl
+{
+	/// <summary>
+	/// Instance of a particular systemic error
+	/// </summary>
+	class SystemicIssueHandler : IIssueHandler
+	{
+		/// <inheritdoc/>
+		public string Type => "Systemic";
+
+		/// <inheritdoc/>
+		public int Priority => 10;
+
+		/// <summary>
+		///  Known systemic errors
+		/// </summary>
+		static readonly HashSet<EventId> s_knownSystemic = new HashSet<EventId> { KnownLogEvents.Horde, KnownLogEvents.Horde_InvalidPreflight };
+
+		/// <summary>
+		/// Determines if the given event id matches
+		/// </summary>
+		/// <param name="eventId">The event id to compare</param>
+		/// <returns>True if the given event id matches</returns>
+		public static bool IsMatchingEventId(EventId eventId)
+		{
+			return s_knownSystemic.Contains(eventId) || (eventId.Id >= KnownLogEvents.Systemic.Id && eventId.Id <= KnownLogEvents.Systemic_Max.Id);
+		}
+
+		/// <inheritdoc/>
+		public void RankSuspects(IIssueFingerprint fingerprint, List<SuspectChange> suspects)
+		{
+		}
+
+		/// <inheritdoc/>
+		public bool TryGetFingerprint(IJob job, INode node, IReadOnlyNodeAnnotations annotations, ILogEventData eventData, [NotNullWhen(true)] out NewIssueFingerprint? fingerprint)
+		{
+			fingerprint = null;
+
+			if (eventData.EventId == KnownLogEvents.ExitCode)
+			{
+				for (int i = 0; i < eventData.Lines.Count; i++)
+				{
+					string message = eventData.Lines[i].Message;
+
+					string[] systemicExitMessages = new string[] { "AutomationTool exiting with ExitCode", "BUILD FAILED", "tool returned code" };
+
+					for (int j = 0; j < systemicExitMessages.Length; j++)
+					{
+						if (message.Contains(systemicExitMessages[j], StringComparison.InvariantCultureIgnoreCase))
+						{
+							fingerprint = new NewIssueFingerprint(Type, new[] { node.Name }, null, null);
+							return true;
+						}
+					}
+				}
+			}
+
+			if (eventData.EventId == null || !IsMatchingEventId(eventData.EventId.Value))
+			{				
+				return false;
+			}
+
+			fingerprint = new NewIssueFingerprint(Type, new[] { node.Name }, null, null);
+			return true;
+		}
+
+		/// <inheritdoc/>
+		public string GetSummary(IIssueFingerprint fingerprint, IssueSeverity severity)
+		{
+			string type = (severity == IssueSeverity.Warning) ? "Systemic warnings" : "Systemic errors";
+			string nodeName = fingerprint.Keys.FirstOrDefault() ?? "(unknown)";
+			return $"{type} in {nodeName}";
+		}
+	}
+}
diff --git a/Engine/Source/Programs/Horde/Horde.Build/Issues/IssueCollection.cs b/Engine/Source/Programs/Horde/Horde.Build/Issues/IssueCollection.cs
new file mode 100644
--- /dev/null
+++ b/Engine/Source/Programs/Horde/Horde.Build/Issues/IssueCollection.cs
@@ -0,0 +1,1380 @@
+// Copyright Epic Games, Inc. All Rights Reserved.
+
+using System;
+using System.Collections.Generic;
+using System.Diagnostics;
+using System.Linq;
+using System.Threading.Tasks;
+using EpicGames.Redis.Utility;
+using Horde.Build.Models;
+using Horde.Build.Server;
+using Horde.Build.Utilities;
+using Microsoft.Extensions.Caching.Memory;
+using Microsoft.Extensions.Logging;
+using MongoDB.Bson;
+using MongoDB.Bson.Serialization.Attributes;
+using MongoDB.Driver;
+
+namespace Horde.Build.Collections.Impl
+{
+	using JobId = ObjectId<IJob>;
+	using LogId = ObjectId<ILogFile>;
+	using StreamId = StringId<IStream>;
+	using TemplateRefId = StringId<TemplateRef>;
+	using UserId = ObjectId<IUser>;
+
+	class IssueCollection : IIssueCollection
+	{
+		[SingletonDocument("5e4c226440ce25fa3207a9af")]
+		class IssueLedger : SingletonBase
+		{
+			public int NextId { get; set; }
+		}
+
+		class Issue : IIssue
+		{
+			[BsonId]
+			public int Id { get; set; }
+
+			public string Summary { get; set; }
+
+			[BsonIgnoreIfNull]
+			public string? UserSummary { get; set; }
+
+			[BsonIgnoreIfNull]
+			public string? Description { get; set; }
+
+			[BsonIgnoreIfNull]
+			public IssueFingerprint? Fingerprint { get; set; }
+			public List<IssueFingerprint>? Fingerprints { get; set; }
+
+			public IssueSeverity Severity { get; set; }
+
+			[BsonElement("PromotedV2")]
+			public bool Promoted { get; set; }
+
+			[BsonIgnoreIfNull]
+			public bool? ManuallyPromoted { get; set; }
+
+			[BsonIgnoreIfNull, BsonElement("Promoted")]
+			public bool? ManuallyPromotedDeprecated { get; set; }
+
+			[BsonIgnoreIfNull]
+			public UserId? OwnerId { get; set; }
+
+			[BsonIgnoreIfNull]
+			public UserId? DefaultOwnerId { get; set; }
+
+			[BsonIgnoreIfNull]
+			public UserId? NominatedById { get; set; }
+
+			public DateTime CreatedAt { get; set; }
+
+			[BsonIgnoreIfNull]
+			public DateTime? NominatedAt { get; set; }
+
+			[BsonIgnoreIfNull]
+			public DateTime? AcknowledgedAt { get; set; }
+
+			[BsonIgnoreIfNull]
+			public DateTime? ResolvedAt { get; set; }
+
+			[BsonIgnoreIfNull]
+			public UserId? ResolvedById { get; set; }
+
+			[BsonIgnoreIfNull]
+			public DateTime? VerifiedAt { get; set; }
+
+			public DateTime LastSeenAt { get; set; }
+
+			[BsonIgnoreIfNull]
+			public int? FixChange { get; set; }
+
+			public List<IssueStream> Streams { get; set; } = new List<IssueStream>();
+
+			public int MinSuspectChange { get; set; }
+			public int MaxSuspectChange { get; set; }
+
+			[BsonElement("Suspects"), BsonIgnoreIfNull]
+			public List<IssueSuspect>? SuspectsDeprecated { get; set; }
+
+			[BsonIgnoreIfNull]
+			public List<ObjectId>? ExcludeSpans { get; set; }
+
+			public int UpdateIndex { get; set; }
+
+			IReadOnlyList<IIssueFingerprint> IIssue.Fingerprints => Fingerprints ?? ((Fingerprint == null)? new List<IssueFingerprint>() : new List<IssueFingerprint> { Fingerprint });
+			UserId? IIssue.OwnerId => OwnerId ?? DefaultOwnerId ?? GetDefaultOwnerId();
+			IReadOnlyList<IIssueStream> IIssue.Streams => Streams;
+			DateTime IIssue.LastSeenAt => (LastSeenAt == default) ? DateTime.UtcNow : LastSeenAt;
+
+			[BsonIgnoreIfNull]
+			public string? ExternalIssueKey { get; set; }
+
+			[BsonIgnoreIfNull]
+			public UserId? QuarantinedByUserId { get; set; }
+
+			[BsonIgnoreIfNull]
+			public DateTime? QuarantineTimeUtc { get; set; }
+
+
+			[BsonConstructor]
+			private Issue()
+			{
+				Summary = String.Empty;
+				Fingerprint = null!;
+			}
+
+			public Issue(int id, string summary)
+			{
+				Id = id;
+				Summary = summary;
+				CreatedAt = DateTime.UtcNow;
+				LastSeenAt = DateTime.UtcNow;
+			}
+
+			UserId? GetDefaultOwnerId()
+			{
+				if (SuspectsDeprecated != null && SuspectsDeprecated.Count > 0)
+				{
+					UserId possibleOwner = SuspectsDeprecated[0].AuthorId;
+					if (SuspectsDeprecated.All(x => x.AuthorId == possibleOwner) && SuspectsDeprecated.Any(x => x.DeclinedAt == null))
+					{
+						return SuspectsDeprecated[0].AuthorId;
+					}
+				}
+				return null;
+			}
+
+			public string FingerprintsDesc
+			{
+				get
+				{
+					if (Fingerprints == null || Fingerprints.Count == 0)
+					{
+						return String.Empty;
+					}
+
+					return String.Join(", ", Fingerprints.Select(x => {
+					   return $"(Type: {x.Type} / Keys: {String.Join(", ", x.Keys)} / RejectKeys: {String.Join(", ", x.RejectKeys ?? new CaseInsensitiveStringSet(new string[] {"No Reject Keys"}))})";
+				   }));
+				}
+			}
+		}
+
+		class IssueStream : IIssueStream
+		{
+			public StreamId StreamId { get; set; }
+			public bool? MergeOrigin { get; set; }
+			public bool? ContainsFix { get; set; }
+
+			public IssueStream()
+			{
+			}
+
+			public IssueStream(IIssueStream other)
+			{
+				StreamId = other.StreamId;
+				MergeOrigin = other.MergeOrigin;
+				ContainsFix = other.ContainsFix;
+			}
+		}
+
+		class IssueSuspect : IIssueSuspect
+		{
+			public ObjectId Id { get; set; }
+			public int IssueId { get; set; }
+			public UserId AuthorId { get; set; }
+			public int Change { get; set; }
+			public DateTime? DeclinedAt { get; set; }
+			public DateTime? ResolvedAt { get; set; } // Degenerate
+
+			private IssueSuspect()
+			{
+			}
+
+			public IssueSuspect(int issueId, NewIssueSuspectData newSuspect, DateTime? resolvedAt)
+			{
+				Id = ObjectId.GenerateNewId();
+				IssueId = issueId;
+				AuthorId = newSuspect.AuthorId;
+				Change = newSuspect.Change;
+				ResolvedAt = resolvedAt;
+			}
+
+			public IssueSuspect(int issueId, IIssueSpanSuspect suspect)
+				: this(issueId, suspect.AuthorId, suspect.OriginatingChange ?? suspect.Change, null, null)
+			{
+			}
+
+			public IssueSuspect(int issueId, UserId authorId, int change, DateTime? declinedAt, DateTime? resolvedAt)
+			{
+				Id = ObjectId.GenerateNewId();
+				IssueId = issueId;
+				AuthorId = authorId;
+				Change = change;
+				DeclinedAt = declinedAt;
+				ResolvedAt = resolvedAt;
+			}
+		}
+
+		class IssueFingerprint : IIssueFingerprint
+		{
+			public string Type { get; set; }
+			public CaseInsensitiveStringSet Keys { get; set; }
+			public CaseInsensitiveStringSet? RejectKeys { get; set; }
+			public CaseInsensitiveStringSet? Metadata { get; set; }
+
+			[BsonConstructor]
+			private IssueFingerprint()
+			{
+				Type = String.Empty;
+				Keys = new CaseInsensitiveStringSet();
+			}
+
+			public IssueFingerprint(IIssueFingerprint fingerprint)
+			{
+				Type = fingerprint.Type;
+				Keys = fingerprint.Keys;
+				RejectKeys = fingerprint.RejectKeys;
+				Metadata = fingerprint.Metadata;
+			}
+		}
+
+		class IssueSpan : IIssueSpan
+		{
+			public ObjectId Id { get; set; }
+
+			[BsonRequired]
+			public StreamId StreamId { get; set; }
+
+			[BsonRequired]
+			public string StreamName { get; set; }
+
+			[BsonRequired]
+			public TemplateRefId TemplateRefId { get; set; }
+
+			[BsonRequired]
+			public string NodeName { get; set; }
+			public DateTime? ResolvedAt { get; set; } // Propagated from the owning issue
+
+			[BsonRequired]
+			public IssueFingerprint Fingerprint { get; set; }
+
+			public int MinChange { get; set; }
+			public int MaxChange { get; set; } = Int32.MaxValue;
+
+			public IssueStep? LastSuccess { get; set; }
+
+			[BsonRequired]
+			public IssueStep FirstFailure { get; set; }
+
+			[BsonRequired]
+			public IssueStep LastFailure { get; set; }
+
+			public IssueStep? NextSuccess { get; set; }
+
+			public bool? PromoteByDefault { get; set; }
+
+			[BsonElement("NotifySuspects"), BsonIgnoreIfDefault(false)]
+			public bool NotifySuspectsDeprecated { get; set; }
+
+			bool IIssueSpan.PromoteByDefault => PromoteByDefault ?? NotifySuspectsDeprecated;
+
+			public List<IssueSpanSuspect> Suspects { get; set; }
+			public int IssueId { get; set; }
+			public int UpdateIndex { get; set; }
+
+			IIssueStep? IIssueSpan.LastSuccess => LastSuccess;
+			IIssueStep IIssueSpan.FirstFailure => FirstFailure;
+			IIssueStep IIssueSpan.LastFailure => LastFailure;
+			IIssueStep? IIssueSpan.NextSuccess => NextSuccess;
+			IReadOnlyList<IIssueSpanSuspect> IIssueSpan.Suspects => Suspects;
+			IIssueFingerprint IIssueSpan.Fingerprint => Fingerprint;
+
+			private IssueSpan()
+			{
+				StreamName = null!;
+				NodeName = null!;
+				Fingerprint = null!;
+				FirstFailure = null!;
+				LastFailure = null!;
+				Suspects = new List<IssueSpanSuspect>();
+			}
+
+			public IssueSpan(int issueId, NewIssueSpanData newSpan)
+			{
+				Id = ObjectId.GenerateNewId();
+				StreamId = newSpan.StreamId;
+				StreamName = newSpan.StreamName;
+				TemplateRefId = newSpan.TemplateRefId;
+				NodeName = newSpan.NodeName;
+				Fingerprint = new IssueFingerprint(newSpan.Fingerprint);
+				if (newSpan.LastSuccess != null)
+				{
+					MinChange = newSpan.LastSuccess.Change;
+					LastSuccess = new IssueStep(Id, newSpan.LastSuccess);
+				}
+				FirstFailure = new IssueStep(Id, newSpan.FirstFailure);
+				LastFailure = new IssueStep(Id, newSpan.FirstFailure);
+				if (newSpan.NextSuccess != null)
+				{
+					MaxChange = newSpan.NextSuccess.Change;
+					NextSuccess = new IssueStep(Id, newSpan.NextSuccess);
+				}
+				PromoteByDefault = newSpan.FirstFailure.PromoteByDefault;
+				Suspects = newSpan.Suspects.ConvertAll(x => new IssueSpanSuspect(x));
+				IssueId = issueId;
+			}
+		}
+
+		class IssueSpanSuspect : IIssueSpanSuspect
+		{
+			public int Change { get; set; }
+			public UserId AuthorId { get; set; }
+			public int? OriginatingChange { get; set; }
+
+			[BsonConstructor]
+			private IssueSpanSuspect()
+			{
+			}
+
+			public IssueSpanSuspect(NewIssueSpanSuspectData newSuspectData)
+			{
+				Change = newSuspectData.Change;
+				AuthorId = newSuspectData.AuthorId;
+				OriginatingChange = newSuspectData.OriginatingChange;
+			}
+		}
+
+		class IssueStep : IIssueStep
+		{
+			public ObjectId Id { get; set; }
+			public ObjectId SpanId { get; set; }
+
+			public int Change { get; set; }
+			public IssueSeverity Severity { get; set; }
+
+			[BsonRequired]
+			public string JobName { get; set; }
+
+			[BsonRequired]
+			public JobId JobId { get; set; }
+
+			[BsonRequired]
+			public SubResourceId BatchId { get; set; }
+
+			[BsonRequired]
+			public SubResourceId StepId { get; set; }
+
+			public DateTime StepTime { get; set; }
+
+			public LogId? LogId { get; set; }
+
+			[BsonIgnoreIfNull]
+			public NodeAnnotations? Annotations { get; set; }
+
+			public bool? PromoteByDefault { get; set; }
+
+			[BsonElement("NotifySuspects"), BsonIgnoreIfDefault(false)]
+			public bool NotifySuspectsDeprecated { get; set; }
+
+			IReadOnlyNodeAnnotations IIssueStep.Annotations => Annotations ?? NodeAnnotations.Empty;
+			bool IIssueStep.PromoteByDefault => PromoteByDefault ?? NotifySuspectsDeprecated;
+
+			[BsonConstructor]
+			private IssueStep()
+			{
+				JobName = null!;
+			}
+
+			public IssueStep(ObjectId spanId, NewIssueStepData stepData)
+			{
+				Id = ObjectId.GenerateNewId();
+				SpanId = spanId;
+				Change = stepData.Change;
+				Severity = stepData.Severity;
+				JobName = stepData.JobName;
+				JobId = stepData.JobId;
+				BatchId = stepData.BatchId;
+				StepId = stepData.StepId;
+				StepTime = stepData.StepTime;
+				LogId = stepData.LogId;
+				Annotations = stepData.Annotations;
+				PromoteByDefault = stepData.PromoteByDefault;
+			}
+		}
+
+		readonly RedisService _redisService;
+		readonly IUserCollection _userCollection;
+		readonly ISingletonDocument<IssueLedger> _ledgerSingleton;
+		readonly IMongoCollection<Issue> _issues;
+		readonly IMongoCollection<IssueSpan> _issueSpans;
+		readonly IMongoCollection<IssueStep> _issueSteps;
+		readonly IMongoCollection<IssueSuspect> _issueSuspects;
+		readonly IAuditLog<int> _auditLog;
+		readonly ILogger _logger;
+
+		public IssueCollection(MongoService mongoService, RedisService redisService, IUserCollection userCollection, IAuditLogFactory<int> auditLogFactory, ILogger<IssueCollection> logger)
+		{
+			_redisService = redisService;
+			_userCollection = userCollection;
+			_logger = logger;
+
+			_ledgerSingleton = new SingletonDocument<IssueLedger>(mongoService);
+
+			List<MongoIndex<Issue>> issueIndexes = new List<MongoIndex<Issue>>();
+			issueIndexes.Add(keys => keys.Ascending(x => x.ResolvedAt));
+			issueIndexes.Add(keys => keys.Ascending(x => x.VerifiedAt));
+			_issues = mongoService.GetCollection<Issue>("IssuesV2");
+
+			List<MongoIndex<IssueSpan>> issueSpanIndexes = new List<MongoIndex<IssueSpan>>();
+			issueSpanIndexes.Add(keys => keys.Ascending(x => x.IssueId));
+			issueSpanIndexes.Add(keys => keys.Ascending(x => x.StreamId).Ascending(x => x.MinChange).Ascending(x => x.MaxChange));
+			issueSpanIndexes.Add("StreamChanges", keys => keys.Ascending(x => x.StreamId).Ascending(x => x.TemplateRefId).Ascending(x => x.NodeName).Ascending(x => x.MinChange).Ascending(x => x.MaxChange));
+			_issueSpans = mongoService.GetCollection<IssueSpan>("IssuesV2.Spans", issueSpanIndexes);
+
+			List<MongoIndex<IssueStep>> issueStepIndexes = new List<MongoIndex<IssueStep>>();
+			issueStepIndexes.Add(keys => keys.Ascending(x => x.SpanId));
+			issueStepIndexes.Add(keys => keys.Ascending(x => x.JobId).Ascending(x => x.BatchId).Ascending(x => x.StepId));
+			_issueSteps = mongoService.GetCollection<IssueStep>("IssuesV2.Steps", issueStepIndexes);
+
+			List<MongoIndex<IssueSuspect>> issueSuspectIndexes = new List<MongoIndex<IssueSuspect>>();
+			issueSuspectIndexes.Add(keys => keys.Ascending(x => x.Change));
+			issueSuspectIndexes.Add(keys => keys.Ascending(x => x.AuthorId).Ascending(x => x.ResolvedAt));
+			issueSuspectIndexes.Add(keys => keys.Ascending(x => x.IssueId).Ascending(x => x.Change), unique: true);
+			_issueSuspects = mongoService.GetCollection<IssueSuspect>("IssuesV2.Suspects", issueSuspectIndexes);
+
+			_auditLog = auditLogFactory.Create("IssuesV2.History", "IssueId");
+
+			if (!mongoService.ReadOnlyMode)
+			{
+				
+
+			}
+		}
+
+		/// <inheritdoc/>
+		public async Task<IAsyncDisposable> EnterCriticalSectionAsync()
+		{
+			Stopwatch timer = Stopwatch.StartNew();
+			TimeSpan nextNotifyTime = TimeSpan.FromSeconds(2.0);
+
+			RedisLock issueLock = new RedisLock(_redisService.Database, "issues/lock");
+			while (!await issueLock.AcquireAsync(TimeSpan.FromMinutes(1)))
+			{
+				if (timer.Elapsed > nextNotifyTime)
+				{
+					_logger.LogWarning("Waiting on lock over issue collection for {TimeSpan}", timer.Elapsed);
+					nextNotifyTime *= 2;
+				}
+				await Task.Delay(TimeSpan.FromMilliseconds(100));
+			}
+			return issueLock;
+		}
+
+		async Task<Issue?> TryUpdateIssueAsync(IIssue issue, UpdateDefinition<Issue> update)
+		{
+			Issue issueDocument = (Issue)issue;
+
+			int prevUpdateIndex = issueDocument.UpdateIndex;
+			update = update.Set(x => x.UpdateIndex, prevUpdateIndex + 1);
+
+			FindOneAndUpdateOptions<Issue, Issue> options = new FindOneAndUpdateOptions<Issue, Issue> { ReturnDocument = ReturnDocument.After };
+			return await _issues.FindOneAndUpdateAsync<Issue>(x => x.Id == issueDocument.Id && x.UpdateIndex == prevUpdateIndex, update, options);
+		}
+
+		async Task<IssueSpan?> TryUpdateSpanAsync(IIssueSpan issueSpan, UpdateDefinition<IssueSpan> update)
+		{
+			IssueSpan issueSpanDocument = (IssueSpan)issueSpan;
+
+			int prevUpdateIndex = issueSpanDocument.UpdateIndex;
+			update = update.Set(x => x.UpdateIndex, prevUpdateIndex + 1);
+
+			FindOneAndUpdateOptions<IssueSpan, IssueSpan> options = new FindOneAndUpdateOptions<IssueSpan, IssueSpan> { ReturnDocument = ReturnDocument.After };
+			return await _issueSpans.FindOneAndUpdateAsync<IssueSpan>(x => x.Id == issueSpanDocument.Id && x.UpdateIndex == prevUpdateIndex, update, options);
+		}
+
+		#region Issues
+
+		/// <inheritdoc/>
+		public async Task<IIssue> AddIssueAsync(string summary)
+		{
+			IssueLedger ledger = await _ledgerSingleton.UpdateAsync(x => x.NextId++);
+
+			Issue newIssue = new Issue(ledger.NextId, summary);
+			await _issues.InsertOneAsync(newIssue);
+
+			ILogger issueLogger = GetLogger(newIssue.Id);
+			issueLogger.LogInformation("Created issue {IssueId}", newIssue.Id);
+
+			return newIssue;
+		}
+
+		static void LogIssueChanges(ILogger issueLogger, Issue oldIssue, Issue newIssue)
+		{
+			if (newIssue.Severity != oldIssue.Severity)
+			{
+				issueLogger.LogInformation("Changed severity to {Severity}", newIssue.Severity);
+			}
+			if (newIssue.Summary != oldIssue.Summary)
+			{
+				issueLogger.LogInformation("Changed summary to \"{Summary}\"", newIssue.Summary);
+			}
+			if (newIssue.Description != oldIssue.Description)
+			{
+				issueLogger.LogInformation("Description set to {Value}", newIssue.Description);
+			}
+			if (((IIssue)newIssue).Promoted != ((IIssue)oldIssue).Promoted)
+			{
+				issueLogger.LogInformation("Promoted set to {Value}", ((IIssue)newIssue).Promoted);
+			}
+			if (newIssue.OwnerId != oldIssue.OwnerId)
+			{
+				if (newIssue.NominatedById != null)
+				{
+					issueLogger.LogInformation("User {UserId} was nominated by {NominatedByUserId}", newIssue.OwnerId, newIssue.NominatedById);
+				}
+				else
+				{
+					issueLogger.LogInformation("User {UserId} was nominated by default", newIssue.OwnerId);
+				}
+			}
+			if (newIssue.AcknowledgedAt != oldIssue.AcknowledgedAt)
+			{
+				if (newIssue.AcknowledgedAt == null)
+				{
+					issueLogger.LogInformation("Issue was un-acknowledged by {UserId}", oldIssue.OwnerId);
+				}
+				else
+				{
+					issueLogger.LogInformation("Issue was acknowledged by {UserId}", oldIssue.OwnerId);
+				}
+			}
+			if (newIssue.FixChange != oldIssue.FixChange)
+			{
+				if (newIssue.FixChange == 0)
+				{
+					issueLogger.LogInformation("Issue was marked as not fixed");
+				}
+				else
+				{
+					issueLogger.LogInformation("Issue was marked as fixed in {Change}", newIssue.FixChange);
+				}
+			}
+			if (newIssue.ResolvedById != oldIssue.ResolvedById)
+			{
+				if (newIssue.ResolvedById == null)
+				{
+					issueLogger.LogInformation("Marking as unresolved");
+				}
+				else
+				{
+					issueLogger.LogInformation("Resolved by {UserId}", newIssue.ResolvedById);
+				}
+			}
+
+			if (newIssue.ResolvedAt != oldIssue.ResolvedAt)
+			{
+				if (newIssue.ResolvedAt == null)
+				{
+					issueLogger.LogInformation("Clearing resolved at time");
+				}
+				else
+				{
+					issueLogger.LogInformation("Setting resolved at time");
+				}
+			}
+
+			if (newIssue.VerifiedAt != oldIssue.VerifiedAt)
+			{
+				if (newIssue.VerifiedAt == null)
+				{
+					issueLogger.LogInformation("Clearing verified at time");
+				}
+				else
+				{
+					issueLogger.LogInformation("Setting verified at time");
+				}
+			}
+
+
+			if (newIssue.ExternalIssueKey != oldIssue.ExternalIssueKey)
+			{
+				if (newIssue.ExternalIssueKey != null)
+				{
+					issueLogger.LogInformation("Linked to external issue {ExternalIssueKey}", newIssue.ExternalIssueKey);
+				}
+				else
+				{
+					issueLogger.LogInformation("Unlinked from external issue {ExternalIssueKey}", oldIssue.ExternalIssueKey);
+				}
+			}
+
+			if (newIssue.QuarantinedByUserId != oldIssue.QuarantinedByUserId)
+			{
+				if (newIssue.QuarantinedByUserId != null)
+				{
+					issueLogger.LogInformation("Quarantined by {UserId}", newIssue.QuarantinedByUserId);
+				}
+				else
+				{
+					issueLogger.LogInformation("Quarantine cleared");
+				}
+			}
+
+
+			string oldFingerprints = oldIssue.FingerprintsDesc;
+			string newFingerprints = newIssue.FingerprintsDesc;
+			if (oldFingerprints != newFingerprints)
+			{
+				issueLogger.LogInformation("Fingerprints changed {Fingerprints}", newFingerprints);
+			}
+
+			HashSet<StreamId> oldFixStreams = new HashSet<StreamId>(oldIssue.Streams.Where(x => x.ContainsFix ?? false).Select(x => x.StreamId));
+			HashSet<StreamId> newFixStreams = new HashSet<StreamId>(newIssue.Streams.Where(x => x.ContainsFix ?? false).Select(x => x.StreamId));
+			foreach (StreamId streamId in newFixStreams.Where(x => !oldFixStreams.Contains(x)))
+			{
+				issueLogger.LogInformation("Marking stream {StreamId} as fixed", streamId);
+			}
+			foreach (StreamId streamId in oldFixStreams.Where(x => !newFixStreams.Contains(x)))
+			{
+				issueLogger.LogInformation("Marking stream {StreamId} as not fixed", streamId);
+			}
+		}
+
+		static void LogIssueSuspectChanges(ILogger issueLogger, List<IssueSuspect> oldIssueSuspects, List<IssueSuspect> newIssueSuspects)
+		{
+			HashSet<(UserId, int)> oldSuspects = new HashSet<(UserId, int)>(oldIssueSuspects.Select(x => (x.AuthorId, x.Change)));
+			HashSet<(UserId, int)> newSuspects = new HashSet<(UserId, int)>(newIssueSuspects.Select(x => (x.AuthorId, x.Change)));
+			foreach ((UserId userId, int change) in newSuspects.Where(x => !oldSuspects.Contains(x)))
+			{
+				issueLogger.LogInformation("Added suspect {UserId} for change {Change}", userId, change);
+			}
+			foreach ((UserId userId, int change) in oldSuspects.Where(x => !newSuspects.Contains(x)))
+			{
+				issueLogger.LogInformation("Removed suspect {UserId} for change {Change}", userId, change);
+			}
+
+			HashSet<UserId> oldDeclinedBy = new HashSet<UserId>(oldIssueSuspects.Where(x => x.DeclinedAt != null).Select(x => x.AuthorId));
+			HashSet<UserId> newDeclinedBy = new HashSet<UserId>(newIssueSuspects.Where(x => x.DeclinedAt != null).Select(x => x.AuthorId));
+			foreach (UserId addDeclinedBy in newDeclinedBy.Where(x => !oldDeclinedBy.Contains(x)))
+			{
+				issueLogger.LogInformation("Declined by {UserId}", addDeclinedBy);
+			}
+			foreach (UserId removeDeclinedBy in oldDeclinedBy.Where(x => !newDeclinedBy.Contains(x)))
+			{
+				issueLogger.LogInformation("Un-declined by {UserId}", removeDeclinedBy);
+			}
+		}
+
+		/// <inheritdoc/>
+		public async Task<IIssue?> GetIssueAsync(int issueId)
+		{
+			Issue issue = await _issues.Find(x => x.Id == issueId).FirstOrDefaultAsync();
+			return issue;
+		}
+
+		/// <inheritdoc/>
+		public Task<List<IIssueSuspect>> FindSuspectsAsync(IIssue issue)
+		{
+			return _issueSuspects.Find(x => x.IssueId == issue.Id).ToListAsync<IssueSuspect, IIssueSuspect>();
+		}
+
+		class ProjectedIssueId
+		{
+			[System.Diagnostics.CodeAnalysis.SuppressMessage("Style", "IDE1006:Naming Styles")]
+			public int? _id { get; set; }
+		}
+
+		/// <inheritdoc/>
+		public async Task<List<IIssue>> FindIssuesAsync(IEnumerable<int>? ids = null, UserId? userId = null, StreamId? streamId = null, int? minChange = null, int? maxChange = null, bool? resolved = null, bool? promoted = null, int? index = null, int? count = null)
+		{
+			List<Issue> results = await FilterIssuesByUserIdAsync(ids, userId, streamId, minChange, maxChange, resolved ?? false, promoted, index ?? 0, count);
+			return results.ConvertAll<IIssue>(x => x);
+		}
+
+		async Task<List<Issue>> FilterIssuesByUserIdAsync(IEnumerable<int>? ids, UserId? userId, StreamId? streamId, int? minChange, int? maxChange, bool? resolved, bool? promoted, int index, int? count)
+		{
+			if (userId == null)
+			{
+				return await FilterIssuesByStreamIdAsync(ids, streamId, minChange, maxChange, resolved, promoted, index, count);
+			}
+			else
+			{
+				FilterDefinition<IssueSuspect> filter = Builders<IssueSuspect>.Filter.Eq(x => x.AuthorId, userId);
+				if (ids != null)
+				{
+					filter &= Builders<IssueSuspect>.Filter.In(x => x.IssueId, ids);
+				}
+				if (resolved != null)
+				{
+					if (resolved.Value)
+					{
+						filter &= Builders<IssueSuspect>.Filter.Ne(x => x.ResolvedAt, null);
+					}
+					else
+					{
+						filter &= Builders<IssueSuspect>.Filter.Eq(x => x.ResolvedAt, null);
+					}
+				}
+
+				using (IAsyncCursor<ProjectedIssueId> cursor = await _issueSuspects.Aggregate().Match(filter).Group(x => x.IssueId, x => new ProjectedIssueId { _id = x.Key }).SortByDescending(x => x._id).ToCursorAsync())
+				{
+					return await PaginatedJoinAsync(cursor, (nextIds, nextIndex, nextCount) => FilterIssuesByStreamIdAsync(nextIds, streamId, minChange, maxChange, null, promoted, nextIndex, nextCount), index, count);
+				}
+			}
+		}
+
+		async Task<List<Issue>> FilterIssuesByStreamIdAsync(IEnumerable<int>? ids, StreamId? streamId, int? minChange, int? maxChange, bool? resolved, bool? promoted, int index, int? count)
+		{
+			if (streamId == null)
+			{
+				return await FilterIssuesByOtherFieldsAsync(ids, minChange, maxChange, resolved, promoted, index, count);
+			}
+			else
+			{
+				FilterDefinition<IssueSpan> filter = Builders<IssueSpan>.Filter.Eq(x => x.StreamId, streamId.Value);
+				if (ids != null)
+				{
+					filter &= Builders<IssueSpan>.Filter.In(x => x.IssueId, ids.Select<int, int?>(x => x));
+				}
+				else
+				{
+					filter &= Builders<IssueSpan>.Filter.Exists(x => x.IssueId);
+				}
+
+				if (minChange != null)
+				{
+					filter &= Builders<IssueSpan>.Filter.Not(Builders<IssueSpan>.Filter.Lt(x => x.MaxChange, minChange.Value));
+				}
+				if (maxChange != null)
+				{
+					filter &= Builders<IssueSpan>.Filter.Not(Builders<IssueSpan>.Filter.Gt(x => x.MinChange, maxChange.Value));
+				}
+
+				if (resolved != null)
+				{
+					if (resolved.Value)
+					{
+						filter &= Builders<IssueSpan>.Filter.Ne(x => x.ResolvedAt, null);
+					}
+					else
+					{
+						filter &= Builders<IssueSpan>.Filter.Eq(x => x.ResolvedAt, null);
+					}
+				}
+
+				using (IAsyncCursor<ProjectedIssueId> cursor = await _issueSpans.Aggregate().Match(filter).Group(x => x.IssueId, x => new ProjectedIssueId { _id = x.Key }).SortByDescending(x => x._id).ToCursorAsync())
+				{
+					List<Issue> results = await PaginatedJoinAsync(cursor, (nextIds, nextIndex, nextCount) => FilterIssuesByOtherFieldsAsync(nextIds, null, null, null, promoted, nextIndex, nextCount), index, count);
+					if (resolved != null)
+					{
+						for (int idx = results.Count - 1; idx >= 0; idx--)
+						{
+							Issue issue = results[idx];
+							if ((issue.ResolvedAt != null) != resolved.Value)
+							{
+								_logger.LogWarning("Issue {IssueId} has resolved state out of sync with spans", issue.Id);
+								results.RemoveAt(idx);
+							}
+						}
+					}
+					return results;
+				}
+			}
+		}
+
+		async Task<List<Issue>> FilterIssuesByOtherFieldsAsync(IEnumerable<int>? ids, int? minChange, int? maxChange, bool? resolved, bool? promoted, int index, int? count)
+		{
+			FilterDefinition<Issue> filter = FilterDefinition<Issue>.Empty;
+			if (ids != null)
+			{
+				filter &= Builders<Issue>.Filter.In(x => x.Id, ids);
+			}
+			if (resolved != null)
+			{
+				if (resolved.Value)
+				{
+					filter &= Builders<Issue>.Filter.Ne(x => x.ResolvedAt, null);
+				}
+				else
+				{
+					filter &= Builders<Issue>.Filter.Eq(x => x.ResolvedAt, null);
+				}
+			}
+			if (minChange != null)
+			{
+				filter &= Builders<Issue>.Filter.Not(Builders<Issue>.Filter.Lt(x => x.MaxSuspectChange, minChange.Value));
+			}
+			if (maxChange != null)
+			{
+				filter &= Builders<Issue>.Filter.Not(Builders<Issue>.Filter.Gt(x => x.MinSuspectChange, maxChange.Value));
+			}
+			if (promoted != null)
+			{
+				if (promoted.Value)
+				{
+					filter &= Builders<Issue>.Filter.Eq(x => x.Promoted, true);
+				}
+				else
+				{
+					filter &= Builders<Issue>.Filter.Ne(x => x.Promoted, true); // Handle the field not existing as well as being set to false.
+				}
+			}
+			return await _issues.Find(filter).SortByDescending(x => x.Id).Range(index, count).ToListAsync();
+		}
+
+		/// <summary>
+		/// Performs a client-side join of a filtered set of issues against another query
+		/// </summary>
+		/// <param name="cursor"></param>
+		/// <param name="nextStageFunc"></param>
+		/// <param name="index"></param>
+		/// <param name="count"></param>
+		/// <returns></returns>
+		static async Task<List<Issue>> PaginatedJoinAsync(IAsyncCursor<ProjectedIssueId> cursor, Func<IEnumerable<int>, int, int?, Task<List<Issue>>> nextStageFunc, int index, int? count)
+		{
+			if (count == null)
+			{
+				List<ProjectedIssueId> issueIds = await cursor.ToListAsync();
+				return await nextStageFunc(issueIds.Where(x => x._id != null).Select(x => x._id!.Value), index, null);
+			}
+			else
+			{
+				List<Issue> results = new List<Issue>();
+				while (await cursor.MoveNextAsync() && results.Count < count.Value)
+				{
+					List<Issue> nextResults = await nextStageFunc(cursor.Current.Where(x => x._id != null).Select(x => x._id!.Value), 0, count.Value - results.Count);
+					int removeCount = Math.Min(index, nextResults.Count);
+					nextResults.RemoveRange(0, removeCount);
+					index -= removeCount;
+					results.AddRange(nextResults);
+				}
+				return results;
+			}
+		}
+
+		/// <inheritdoc/>
+		public async Task<List<IIssue>> FindIssuesForChangesAsync(List<int> changes)
+		{
+			List<int> issueIds = await (await _issueSuspects.DistinctAsync(x => x.IssueId, Builders<IssueSuspect>.Filter.In(x => x.Change, changes))).ToListAsync();
+			return await _issues.Find(Builders<Issue>.Filter.In(x => x.Id, issueIds)).ToListAsync<Issue, IIssue>();
+		}
+
+		/// <inheritdoc/>
+		public async Task<IIssue?> TryUpdateIssueAsync(IIssue issue, IssueSeverity? newSeverity = null, string? newSummary = null, string? newUserSummary = null, string? newDescription = null, bool? newManuallyPromoted = null, UserId? newOwnerId = null, UserId? newNominatedById = null, bool? newAcknowledged = null, UserId? newDeclinedById = null, int? newFixChange = null, UserId? newResolvedById = null, List<ObjectId>? newExcludeSpanIds = null, DateTime? newLastSeenAt = null, string? newExternaIssueKey = null, UserId? newQuarantinedById = null)
+		{
+			Issue issueDocument = (Issue)issue;
+
+			DateTime utcNow = DateTime.UtcNow;
+
+			List<UpdateDefinition<Issue>> updates = new List<UpdateDefinition<Issue>>();
+			if (newSeverity != null)
+			{
+				updates.Add(Builders<Issue>.Update.Set(x => x.Severity, newSeverity.Value));
+			}
+			if (newSummary != null)
+			{
+				updates.Add(Builders<Issue>.Update.Set(x => x.Summary, newSummary));
+			}
+			if (newUserSummary != null)
+			{
+				if (newUserSummary.Length == 0)
+				{
+					updates.Add(Builders<Issue>.Update.Unset(x => x.UserSummary!));
+				}
+				else
+				{
+					updates.Add(Builders<Issue>.Update.Set(x => x.UserSummary, newUserSummary));
+				}
+			}
+			if (newDescription != null)
+			{
+				if (newDescription.Length == 0)
+				{
+					updates.Add(Builders<Issue>.Update.Unset(x => x.Description));
+				}
+				else
+				{
+					updates.Add(Builders<Issue>.Update.Set(x => x.Description, newDescription));
+				}
+			}
+			if (newManuallyPromoted != null)
+			{
+				updates.Add(Builders<Issue>.Update.Set(x => x.ManuallyPromoted, newManuallyPromoted.Value));
+			}
+			if (newResolvedById != null)
+			{
+				newOwnerId ??= newResolvedById;
+				newAcknowledged ??= true;
+			}
+			if (newOwnerId != null)
+			{
+				if (newOwnerId.Value == UserId.Empty)
+				{
+					updates.Add(Builders<Issue>.Update.Unset(x => x.OwnerId!));
+					updates.Add(Builders<Issue>.Update.Unset(x => x.NominatedAt!));
+					updates.Add(Builders<Issue>.Update.Unset(x => x.NominatedById!));
+				}
+				else
+				{
+					updates.Add(Builders<Issue>.Update.Set(x => x.OwnerId!, newOwnerId.Value));
+
+					updates.Add(Builders<Issue>.Update.Set(x => x.NominatedAt, DateTime.UtcNow));
+					if (newNominatedById == null)
+					{
+						updates.Add(Builders<Issue>.Update.Unset(x => x.NominatedById!));
+					}
+					else
+					{
+						updates.Add(Builders<Issue>.Update.Set(x => x.NominatedById, newNominatedById.Value));
+					}
+					newAcknowledged ??= false;
+				}
+			}
+			if (newAcknowledged != null)
+			{
+				if (newAcknowledged.Value)
+				{
+					if (issueDocument.AcknowledgedAt == null)
+					{
+						updates.Add(Builders<Issue>.Update.Set(x => x.AcknowledgedAt, utcNow));
+					}
+				}
+				else
+				{
+					if (issueDocument.AcknowledgedAt != null)
+					{
+						updates.Add(Builders<Issue>.Update.Unset(x => x.AcknowledgedAt!));
+					}
+				}
+			}
+			if (newFixChange != null)
+			{
+				if (newFixChange == 0)
+				{
+					updates.Add(Builders<Issue>.Update.Unset(x => x.FixChange!));
+				}
+				else
+				{
+					updates.Add(Builders<Issue>.Update.Set(x => x.FixChange, newFixChange));
+				}
+			}
+			if (newResolvedById != null)
+			{
+				if (newResolvedById.Value != UserId.Empty)
+				{
+					if (issueDocument.ResolvedAt == null || issueDocument.ResolvedById != newResolvedById)
+					{
+						updates.Add(Builders<Issue>.Update.Set(x => x.ResolvedAt, utcNow));
+						updates.Add(Builders<Issue>.Update.Set(x => x.ResolvedById, newResolvedById.Value));
+					}
+				}
+				else
+				{
+					if (issueDocument.ResolvedAt != null)
+					{
+						updates.Add(Builders<Issue>.Update.Unset(x => x.ResolvedAt!));
+					}
+					if (issueDocument.ResolvedById != null)
+					{
+						updates.Add(Builders<Issue>.Update.Unset(x => x.ResolvedById!));
+					}
+				}
+			}
+			if (newExcludeSpanIds != null)
+			{
+				List<ObjectId> newCombinedExcludeSpanIds = newExcludeSpanIds;
+				if (issue.ExcludeSpans != null)
+				{
+					newCombinedExcludeSpanIds = newCombinedExcludeSpanIds.Union(issue.ExcludeSpans).ToList();
+				}
+				updates.Add(Builders<Issue>.Update.Set(x => x.ExcludeSpans, newCombinedExcludeSpanIds));
+			}
+			if (newLastSeenAt != null)
+			{
+				updates.Add(Builders<Issue>.Update.Set(x => x.LastSeenAt, newLastSeenAt.Value));
+			}
+
+			if (newDeclinedById != null)
+			{
+				await _issueSuspects.UpdateManyAsync(x => x.IssueId == issue.Id && x.AuthorId == newDeclinedById.Value, Builders<IssueSuspect>.Update.Set(x => x.DeclinedAt, DateTime.UtcNow));
+			}
+			if (newQuarantinedById != null)
+			{
+				if (newQuarantinedById.Value == UserId.Empty)
+				{
+					updates.Add(Builders<Issue>.Update.Unset(x => x.QuarantinedByUserId));
+					updates.Add(Builders<Issue>.Update.Unset(x => x.QuarantineTimeUtc));
+				}
+				else
+				{
+					updates.Add(Builders<Issue>.Update.Set(x => x.QuarantinedByUserId!, newQuarantinedById.Value));
+					updates.Add(Builders<Issue>.Update.Set(x => x.QuarantineTimeUtc, DateTime.UtcNow));
+				}
+			}
+
+			if (newExternaIssueKey != null)
+			{
+				updates.Add(Builders<Issue>.Update.Set(x => x.ExternalIssueKey, newExternaIssueKey.Length == 0 ? null : newExternaIssueKey));
+			}
+
+			if (updates.Count == 0)
+			{
+				return issueDocument;
+			}
+
+			Issue? newIssue = await TryUpdateIssueAsync(issue, Builders<Issue>.Update.Combine(updates));
+			if(newIssue == null)
+			{
+				return null;
+			}
+
+			ILogger issueLogger = GetLogger(issue.Id);
+			LogIssueChanges(issueLogger, issueDocument, newIssue);
+			return newIssue;
+		}
+
+		/// <inheritdoc/>
+		public async Task<IIssue?> TryUpdateIssueDerivedDataAsync(IIssue issue, string newSummary, IssueSeverity newSeverity, List<NewIssueFingerprint> newFingerprints, List<NewIssueStream> newStreams, List<NewIssueSuspectData> newSuspects, DateTime? newResolvedAt, DateTime? newVerifiedAt, DateTime newLastSeenAt)
+		{
+			Issue issueImpl = (Issue)issue;
+
+			// Update all the suspects for this issue
+			List<IssueSuspect> oldSuspectImpls = await _issueSuspects.Find(x => x.IssueId == issue.Id).ToListAsync();
+			List<IssueSuspect> newSuspectImpls = await UpdateIssueSuspectsAsync(issue.Id, oldSuspectImpls, newSuspects, newResolvedAt);
+
+			// Find the spans for this issue
+			List<IssueSpan> newSpans = await _issueSpans.Find(x => x.IssueId == issue.Id).ToListAsync();
+
+			// Update the resolved time on any issues
+			List<ObjectId> updateSpanIds = newSpans.Where(x => x.ResolvedAt != newResolvedAt).Select(x => x.Id).ToList();
+			if (updateSpanIds.Count > 0)
+			{
+				FilterDefinition<IssueSpan> filter = Builders<IssueSpan>.Filter.In(x => x.Id, updateSpanIds);
+				await _issueSpans.UpdateManyAsync(filter, Builders<IssueSpan>.Update.Set(x => x.ResolvedAt, newResolvedAt));
+			}
+
+			// Figure out if this issue should be promoted
+			bool newPromoted;
+			if (issueImpl.ManuallyPromoted.HasValue)
+			{
+				newPromoted = issueImpl.ManuallyPromoted.Value;
+			}
+			else if (issueImpl.ManuallyPromotedDeprecated.HasValue)
+			{
+				newPromoted = issueImpl.ManuallyPromotedDeprecated.Value;
+			}
+			else
+			{
+				newPromoted = newSpans.Any(x => ((IIssueSpan)x).PromoteByDefault);
+			}
+
+			// Find the default owner
+			UserId? newDefaultOwnerId = null;
+			string? autoAssignToUser = newSpans.Select(x => x.LastFailure.Annotations?.AutoAssignToUser).Where(x => x != null).FirstOrDefault();
+			if (autoAssignToUser != null)
+			{
+				IUser? user = await _userCollection.FindUserByLoginAsync(autoAssignToUser);
+				if(user != null)
+				{
+					newDefaultOwnerId = user.Id;
+				}
+			}
+			else
+			{
+				// Figure out if we can auto-assign an owner
+				bool canAutoAssign = newPromoted || newSpans.Any(x => x.LastFailure.Annotations?.AutoAssign ?? false);
+				if (canAutoAssign && newSuspectImpls.Count > 0)
+				{
+					UserId possibleOwnerId = newSuspectImpls[0].AuthorId;
+					if (newSuspectImpls.All(x => x.AuthorId == possibleOwnerId) && newSuspectImpls.Any(x => x.DeclinedAt == null))
+					{
+						newDefaultOwnerId = possibleOwnerId;
+					}
+				}
+			}
+
+			// Get the range of suspect changes
+			int newMinSuspectChange = (newSuspects.Count > 0) ? newSuspects.Min(x => x.Change) : 0;
+			int newMaxSuspectChange = (newSuspects.Count > 0) ? newSuspects.Min(x => x.Change) : 0;
+
+			// Perform the actual update with this data
+			List<UpdateDefinition<Issue>> updates = new List<UpdateDefinition<Issue>>();
+			if (!String.Equals(issue.Summary, newSummary, StringComparison.Ordinal))
+			{
+				updates.Add(Builders<Issue>.Update.Set(x => x.Summary, newSummary));
+			}
+			if (issue.Severity != newSeverity)
+			{
+				updates.Add(Builders<Issue>.Update.Set(x => x.Severity, newSeverity));
+			}
+			if (issue.Promoted != newPromoted)
+			{
+				updates.Add(Builders<Issue>.Update.Set(x => x.Promoted, newPromoted));
+			}
+			if (issue.Fingerprints.Count != newFingerprints.Count || !newFingerprints.Zip(issue.Fingerprints).All(x => x.First.Equals(x.Second)))
+			{
+				updates.Add(Builders<Issue>.Update.Set(x => x.Fingerprints, newFingerprints.Select(x => new IssueFingerprint(x))));
+			}
+			if (issue.Streams.Count != newStreams.Count || !newStreams.Zip(issue.Streams).All(x => x.First.StreamId == x.Second.StreamId && x.First.ContainsFix == x.Second.ContainsFix))
+			{
+				updates.Add(Builders<Issue>.Update.Set(x => x.Streams, newStreams.Select(x => new IssueStream(x))));
+			}
+			if (issueImpl.MinSuspectChange != newMinSuspectChange)
+			{
+				updates.Add(Builders<Issue>.Update.Set(x => x.MinSuspectChange, newMinSuspectChange));
+			}
+			if (issueImpl.MaxSuspectChange != newMaxSuspectChange)
+			{
+				updates.Add(Builders<Issue>.Update.Set(x => x.MaxSuspectChange, newMaxSuspectChange));
+			}
+			if (issueImpl.DefaultOwnerId != newDefaultOwnerId)
+			{
+				updates.Add(Builders<Issue>.Update.Set(x => x.DefaultOwnerId, newDefaultOwnerId));
+			}
+			if (issue.ResolvedAt != newResolvedAt)
+			{
+				updates.Add(Builders<Issue>.Update.SetOrUnsetNull(x => x.ResolvedAt, newResolvedAt));
+			}
+			if (newResolvedAt == null && issue.ResolvedById != null)
+			{
+				updates.Add(Builders<Issue>.Update.Unset(x => x.ResolvedById));
+			}
+			if (issue.VerifiedAt != newVerifiedAt)
+			{
+				updates.Add(Builders<Issue>.Update.SetOrUnsetNull(x => x.VerifiedAt, newVerifiedAt));
+			}
+			if (issue.LastSeenAt != newLastSeenAt)
+			{
+				updates.Add(Builders<Issue>.Update.Set(x => x.LastSeenAt, newLastSeenAt));
+			}
+
+			Issue? newIssue = await TryUpdateIssueAsync(issue, Builders<Issue>.Update.Combine(updates));
+			if(newIssue != null)
+			{
+				ILogger issueLogger = GetLogger(issue.Id);
+				LogIssueChanges(GetLogger(issue.Id), issueImpl, newIssue);
+				LogIssueSuspectChanges(issueLogger, oldSuspectImpls, newSuspectImpls);
+				return newIssue;
+			}
+			return null;
+		}
+
+		async Task<List<IssueSuspect>> UpdateIssueSuspectsAsync(int issueId, List<IssueSuspect> oldSuspectImpls, List<NewIssueSuspectData> newSuspects, DateTime? resolvedAt)
+		{
+			List<IssueSuspect> newSuspectImpls = new List<IssueSuspect>(oldSuspectImpls);
+
+			// Find the current list of suspects
+			HashSet<(UserId, int)> curSuspectKeys = new HashSet<(UserId, int)>(oldSuspectImpls.Select(x => (x.AuthorId, x.Change)));
+			List<IssueSuspect> createSuspects = newSuspects.Where(x => !curSuspectKeys.Contains((x.AuthorId, x.Change))).Select(x => new IssueSuspect(issueId, x, resolvedAt)).ToList();
+
+			HashSet<(UserId, int)> newSuspectKeys = new HashSet<(UserId, int)>(newSuspects.Select(x => (x.AuthorId, x.Change)));
+			List<IssueSuspect> deleteSuspects = oldSuspectImpls.Where(x => !newSuspectKeys.Contains((x.AuthorId, x.Change))).ToList();
+
+			// Apply the suspect changes
+			if (createSuspects.Count > 0)
+			{
+				await _issueSuspects.InsertManyIgnoreDuplicatesAsync(createSuspects);
+				newSuspectImpls.AddRange(createSuspects);
+			}
+			if (deleteSuspects.Count > 0)
+			{
+				await _issueSuspects.DeleteManyAsync(Builders<IssueSuspect>.Filter.In(x => x.Id, deleteSuspects.Select(y => y.Id)));
+				newSuspectImpls.RemoveAll(x => !newSuspectKeys.Contains((x.AuthorId, x.Change)));
+			}
+
+			// Make sure all the remaining suspects have the correct resolved time
+			if (newSuspectImpls.Any(x => x.ResolvedAt != resolvedAt))
+			{
+				await _issueSuspects.UpdateManyAsync(Builders<IssueSuspect>.Filter.Eq(x => x.IssueId, issueId), Builders<IssueSuspect>.Update.Set(x => x.ResolvedAt, resolvedAt));
+			}
+			return newSuspectImpls;
+		}
+
+		#endregion
+
+		#region Spans
+
+		/// <inheritdoc/>
+		public async Task<IIssueSpan> AddSpanAsync(int issueId, NewIssueSpanData newSpan)
+		{
+			IssueSpan span = new IssueSpan(issueId, newSpan);
+			await _issueSpans.InsertOneAsync(span);
+			return span;
+		}
+
+		/// <inheritdoc/>
+		public async Task<IIssueSpan?> GetSpanAsync(ObjectId spanId)
+		{
+			return await _issueSpans.Find(Builders<IssueSpan>.Filter.Eq(x => x.Id, spanId)).FirstOrDefaultAsync();
+		}
+
+		/// <inheritdoc/>
+		public async Task<IIssueSpan?> TryUpdateSpanAsync(IIssueSpan span, NewIssueStepData? newLastSuccess = null, NewIssueStepData? newFailure = null, NewIssueStepData? newNextSuccess = null, List<NewIssueSpanSuspectData>? newSuspects = null, int? newIssueId = null)
+		{
+			List<UpdateDefinition<IssueSpan>> updates = new List<UpdateDefinition<IssueSpan>>();
+			if (newLastSuccess != null)
+			{
+				updates.Add(Builders<IssueSpan>.Update.Set(x => x.MinChange, newLastSuccess.Change));
+				updates.Add(Builders<IssueSpan>.Update.Set(x => x.LastSuccess, new IssueStep(span.Id, newLastSuccess)));
+			}
+			if (newFailure != null)
+			{
+				if (newFailure.Change < span.FirstFailure.Change)
+				{
+					updates.Add(Builders<IssueSpan>.Update.Set(x => x.FirstFailure, new IssueStep(span.Id, newFailure)));
+				}
+				if (newFailure.Change >= span.LastFailure.Change)
+				{
+					updates.Add(Builders<IssueSpan>.Update.Set(x => x.LastFailure, new IssueStep(span.Id, newFailure)));
+				}
+				if (newFailure.PromoteByDefault != span.PromoteByDefault && newFailure.Change >= span.LastFailure.Change)
+				{
+					updates.Add(Builders<IssueSpan>.Update.Set(x => x.PromoteByDefault, newFailure.PromoteByDefault));
+				}
+			}
+			if (newNextSuccess != null)
+			{
+				updates.Add(Builders<IssueSpan>.Update.Set(x => x.MaxChange, newNextSuccess.Change));
+				updates.Add(Builders<IssueSpan>.Update.Set(x => x.NextSuccess, new IssueStep(span.Id, newNextSuccess)));
+			}
+			if (newSuspects != null)
+			{
+				updates.Add(Builders<IssueSpan>.Update.Set(x => x.Suspects, newSuspects.ConvertAll(x => new IssueSpanSuspect(x))));
+			}
+			if (newIssueId != null)
+			{
+				updates.Add(Builders<IssueSpan>.Update.Set(x => x.IssueId, newIssueId.Value));
+			}
+
+			if (updates.Count == 0)
+			{
+				return span;
+			}
+
+			IssueSpan? newSpan = await TryUpdateSpanAsync(span, Builders<IssueSpan>.Update.Combine(updates));
+			if (newSpan != null)
+			{
+				ILogger logger = GetLogger(newSpan.IssueId);
+				if (newLastSuccess != null)
+				{
+					logger.LogInformation("Set last success for span {SpanId} to job {JobId} at CL {Change}", newSpan.Id, newLastSuccess.JobId, newLastSuccess.Change);
+				}
+				if (newNextSuccess != null)
+				{
+					logger.LogInformation("Set next success for span {SpanId} to job {JobId} at CL {Change}", newSpan.Id, newNextSuccess.JobId, newNextSuccess.Change);
+				}
+				if (newFailure != null)
+				{
+					logger.LogInformation("Added failure for span {SpanId} in job {JobId} at CL {Change}", newSpan.Id, newFailure.JobId, newFailure.Change);
+				}
+			}
+			return newSpan;
+		}
+
+		/// <inheritdoc/>
+		public async Task<List<IIssueSpan>> FindSpansAsync(int issueId)
+		{
+			return await _issueSpans.Find(x => x.IssueId == issueId).ToListAsync<IssueSpan, IIssueSpan>();
+		}
+
+		/// <inheritdoc/>
+		public Task<List<IIssueSpan>> FindSpansAsync(IEnumerable<ObjectId> spanIds)
+		{
+			return _issueSpans.Find(Builders<IssueSpan>.Filter.In(x => x.Id, spanIds)).ToListAsync<IssueSpan, IIssueSpan>();
+		}
+
+		/// <inheritdoc/>
+		public async Task<List<IIssueSpan>> FindOpenSpansAsync(StreamId streamId, TemplateRefId templateId, string nodeName, int change)
+		{
+			List<IssueSpan> spans = await _issueSpans.Find(x => x.StreamId == streamId && x.TemplateRefId == templateId && x.NodeName == nodeName && change >= x.MinChange && change <= x.MaxChange).ToListAsync();
+			return spans.ConvertAll<IIssueSpan>(x => x);
+		}
+
+		/// <inheritdoc/>
+		public Task<List<IIssueSpan>> FindSpansAsync(IEnumerable<ObjectId>? spanIds, IEnumerable<int>? issueIds, StreamId? streamId, int? minChange, int? maxChange, bool? resolved, int? index, int? count)
+		{
+			FilterDefinition<IssueSpan> filter = FilterDefinition<IssueSpan>.Empty;
+
+			if(spanIds != null)
+			{
+				filter &= Builders<IssueSpan>.Filter.In(x => x.Id, spanIds);
+			}
+
+			if (streamId != null)
+			{
+				filter &= Builders<IssueSpan>.Filter.Eq(x => x.StreamId, streamId);
+			}
+
+			if (issueIds != null)
+			{
+				filter &= Builders<IssueSpan>.Filter.In(x => x.IssueId, issueIds.Select<int, int?>(x => x));
+			}
+			if (minChange != null)
+			{
+				filter &= Builders<IssueSpan>.Filter.Not(Builders<IssueSpan>.Filter.Lt(x => x.MaxChange, minChange.Value));
+			}
+			if (maxChange != null)
+			{
+				filter &= Builders<IssueSpan>.Filter.Not(Builders<IssueSpan>.Filter.Gt(x => x.MinChange, maxChange.Value));
+			}
+			if (resolved != null)
+			{
+				if (resolved.Value)
+				{
+					filter &= Builders<IssueSpan>.Filter.Ne(x => x.ResolvedAt, null);
+				}
+				else
+				{
+					filter &= Builders<IssueSpan>.Filter.Eq(x => x.ResolvedAt, null);
+				}
+			}
+
+			return _issueSpans.Find(filter).Range(index, count).ToListAsync<IssueSpan, IIssueSpan>();
+		}
+
+		#endregion
+
+		#region Steps
+
+		/// <inheritdoc/>
+		public async Task<IIssueStep> AddStepAsync(ObjectId spanId, NewIssueStepData newStep)
+		{
+			IssueStep step = new IssueStep(spanId, newStep);
+			await _issueSteps.InsertOneAsync(step);
+			return step;
+		}
+
+		/// <inheritdoc/>
+		public Task<List<IIssueStep>> FindStepsAsync(IEnumerable<ObjectId> spanIds)
+		{
+			FilterDefinition<IssueStep> filter = Builders<IssueStep>.Filter.In(x => x.SpanId, spanIds);
+			return _issueSteps.Find(filter).ToListAsync<IssueStep, IIssueStep>();
+		}
+
+		/// <inheritdoc/>
+		public Task<List<IIssueStep>> FindStepsAsync(JobId jobId, SubResourceId? batchId, SubResourceId? stepId)
+		{
+			FilterDefinition<IssueStep> filter = Builders<IssueStep>.Filter.Eq(x => x.JobId, jobId);
+			if (batchId != null)
+			{
+				filter &= Builders<IssueStep>.Filter.Eq(x => x.BatchId, batchId.Value);
+			}
+			if (stepId != null)
+			{
+				filter &= Builders<IssueStep>.Filter.Eq(x => x.StepId, stepId.Value);
+			}
+			return _issueSteps.Find(filter).ToListAsync<IssueStep, IIssueStep>();
+		}
+
+		#endregion
+
+		/// <inheritdoc/>
+		public IAuditLogChannel<int> GetLogger(int issueId)
+		{
+			return _auditLog[issueId];
+		}
+	}
+}
