Horde: Finish re-org of Horde solution started a while ago.

diff --git a/Engine/Source/Programs/Horde/Horde.Build/Issues/IssuesController.cs b/Engine/Source/Programs/Horde/Horde.Build/Issues/IssuesController.cs
new file mode 100644
--- /dev/null
+++ b/Engine/Source/Programs/Horde/Horde.Build/Issues/IssuesController.cs
@@ -0,0 +1,700 @@
+// Copyright Epic Games, Inc. All Rights Reserved.
+
+using System;
+using System.Collections.Generic;
+using System.Linq;
+using System.Threading.Tasks;
+using Horde.Build.Acls;
+using Horde.Build.Api;
+using Horde.Build.Collections;
+using Horde.Build.Models;
+using Horde.Build.Server;
+using Horde.Build.Services;
+using Horde.Build.Utilities;
+using HordeCommon;
+using Microsoft.AspNetCore.Authorization;
+using Microsoft.AspNetCore.Mvc;
+using MongoDB.Bson;
+using Microsoft.Extensions.Logging;
+
+namespace Horde.Build.Controllers
+{
+	using JobId = ObjectId<IJob>;
+	using LogId = ObjectId<ILogFile>;
+	using StreamId = StringId<IStream>;
+	using UserId = ObjectId<IUser>;
+	using WorkflowId = StringId<WorkflowConfig>;
+
+	/// <summary>
+	/// Controller for the /api/v1/issues endpoint
+	/// </summary>
+	[Authorize]
+	[ApiController]
+	[Route("[controller]")]
+	public class IssuesController : HordeControllerBase
+	{
+		private readonly IIssueCollection _issueCollection;
+		private readonly IssueService _issueService;
+		private readonly IExternalIssueService _externalIssueService;
+		private readonly JobService _jobService;
+		private readonly StreamService _streamService;
+		private readonly IUserCollection _userCollection;
+		private readonly ILogFileService _logFileService;
+		private readonly ILogger<IssuesController> _logger;
+
+		/// <summary>
+		/// Constructor
+		/// </summary>
+		public IssuesController(ILogger<IssuesController> logger, IIssueCollection issueCollection, IssueService issueService, JobService jobService, StreamService streamService, IUserCollection userCollection, ILogFileService logFileService, IExternalIssueService externalIssueService)
+		{			
+			_issueCollection = issueCollection;
+			_issueService = issueService;
+			_jobService = jobService;
+			_streamService = streamService;
+			_userCollection = userCollection;
+			_logFileService = logFileService;
+			_externalIssueService = externalIssueService;
+			_logger = logger;
+		}
+
+		/// <summary>
+		/// Retrieve information about a specific issue
+		/// </summary>
+		/// <param name="ids">Set of issue ids to find</param>
+		/// <param name="streamId">The stream to query for</param>
+		/// <param name="minChange">The minimum changelist range to query, inclusive</param>
+		/// <param name="maxChange">The minimum changelist range to query, inclusive</param>
+		/// <param name="resolved">Whether to include resolved issues</param>
+		/// <param name="index">Starting offset of the window of results to return</param>
+		/// <param name="count">Number of results to return</param>
+		/// <param name="filter">Filter for the properties to return</param>
+		/// <returns>List of matching agents</returns>
+		[HttpGet]
+		[Route("/api/v2/issues")]
+		[ProducesResponseType(typeof(List<FindIssueResponse>), 200)]
+		public async Task<ActionResult<object>> FindIssuesV2Async([FromQuery(Name = "Id")] int[]? ids = null, [FromQuery] StreamId? streamId = null, [FromQuery] int? minChange = null, [FromQuery] int? maxChange = null, [FromQuery] bool? resolved = null, [FromQuery] int index = 0, [FromQuery] int count = 10, [FromQuery] PropertyFilter? filter = null)
+		{
+			if (ids != null && ids.Length == 0)
+			{
+				ids = null;
+			}
+
+			List<object> responses = new List<object>();
+			if (streamId != null)
+			{
+				if (!await _streamService.AuthorizeAsync(streamId.Value, AclAction.ViewStream, User, new StreamPermissionsCache()))
+				{
+					return Forbid();
+				}
+
+				IStream? stream = await _streamService.GetStreamAsync(streamId.Value);
+				if (stream == null)
+				{
+					return NotFound(streamId.Value);
+				}
+
+				List<IIssueSpan> spans = await _issueCollection.FindSpansAsync(null, ids, streamId.Value, minChange, maxChange, resolved);
+				if(spans.Count > 0)
+				{
+					// Group all the spans by their issue id
+					Dictionary<int, List<IIssueSpan>> issueIdToSpans = new Dictionary<int, List<IIssueSpan>>();
+					foreach (IIssueSpan span in spans)
+					{
+						List<IIssueSpan>? spansForIssue;
+						if (!issueIdToSpans.TryGetValue(span.IssueId, out spansForIssue))
+						{
+							spansForIssue = new List<IIssueSpan>();
+							issueIdToSpans.Add(span.IssueId, spansForIssue);
+						}
+						spansForIssue.Add(span);
+					}
+
+					// Find the matching issues
+					List<IIssue> issues = await _issueCollection.FindIssuesAsync(issueIdToSpans.Keys, index: index, count: count);
+
+					// Create the corresponding responses
+					foreach (IIssue issue in issues.OrderByDescending(x => x.Id))
+					{
+						IssueSeverity streamSeverity = IssueSeverity.Unspecified;
+
+						List<FindIssueSpanResponse> spanResponses = new List<FindIssueSpanResponse>();
+						if (issueIdToSpans.TryGetValue(issue.Id, out List<IIssueSpan>? spansForIssue))
+						{
+							// Find the current severity in the stream
+							DateTime lastStepTime = DateTime.MinValue;
+							foreach (IIssueSpan span in spansForIssue)
+							{
+								if (span.LastFailure != null && span.LastFailure.StepTime > lastStepTime)
+								{
+									lastStepTime = span.LastFailure.StepTime;
+									streamSeverity = span.LastFailure.Severity;
+								}
+							}
+
+							// Convert each issue to a response
+							foreach (IIssueSpan span in spansForIssue)
+							{
+								spanResponses.Add(new FindIssueSpanResponse(span, span.LastFailure.Annotations.WorkflowId));
+							}
+						}
+					
+						IUser? owner = null;
+						IUser? nominatedBy = null;
+						IUser? resolvedBy = null;
+						IUser? quarantinedBy = null;
+
+						if (issue.OwnerId != null)
+						{
+							owner = await _userCollection.GetCachedUserAsync(issue.OwnerId.Value);
+						}
+						
+						if (issue.NominatedById != null)
+						{
+							nominatedBy = await _userCollection.GetCachedUserAsync(issue.NominatedById.Value);
+						}
+						
+						if (issue.ResolvedById != null)
+						{
+							resolvedBy = await _userCollection.GetCachedUserAsync(issue.ResolvedById.Value);
+						}
+
+						if (issue.QuarantinedByUserId != null)
+						{
+							quarantinedBy = await _userCollection.GetCachedUserAsync(issue.QuarantinedByUserId.Value);
+						}
+
+
+						FindIssueResponse response = new FindIssueResponse(issue, owner, nominatedBy, resolvedBy, quarantinedBy, streamSeverity, spanResponses);
+						responses.Add(PropertyFilter.Apply(response, filter));
+					}
+				}
+			}
+			else
+			{
+				return BadRequest("Missing StreamId on request");
+			}
+
+			return responses;
+		}
+
+		/// <summary>
+		/// Retrieve information about a specific issue
+		/// </summary>
+		/// <param name="ids">Set of issue ids to find</param>
+		/// <param name="streamId">The stream to query for</param>
+		/// <param name="change">The changelist to query</param>
+		/// <param name="minChange">The minimum changelist range to query, inclusive</param>
+		/// <param name="maxChange">The minimum changelist range to query, inclusive</param>
+		/// <param name="jobId">Job id to filter by</param>
+		/// <param name="batchId">The batch to filter by</param>
+		/// <param name="stepId">The step to filter by</param>
+		/// <param name="labelIdx">The label within the job to filter by</param>
+		/// <param name="userId">User to filter issues for</param>
+		/// <param name="resolved">Whether to include resolved issues</param>
+		/// <param name="promoted">Whether to include promoted issues</param>
+		/// <param name="index">Starting offset of the window of results to return</param>
+		/// <param name="count">Number of results to return</param>
+		/// <param name="filter">Filter for the properties to return</param>
+		/// <returns>List of matching agents</returns>
+		[HttpGet]
+		[Route("/api/v1/issues")]
+		[ProducesResponseType(typeof(List<GetIssueResponse>), 200)]
+		public async Task<ActionResult<object>> FindIssuesAsync([FromQuery(Name = "Id")] int[]? ids = null, [FromQuery] string? streamId = null, [FromQuery] int? change = null, [FromQuery] int? minChange = null, [FromQuery] int? maxChange = null, [FromQuery] JobId? jobId = null, [FromQuery] string? batchId = null, [FromQuery] string? stepId = null, [FromQuery(Name = "label")] int? labelIdx = null, [FromQuery] string? userId = null, [FromQuery] bool? resolved = null, [FromQuery] bool? promoted = null, [FromQuery] int index = 0, [FromQuery] int count = 10, [FromQuery] PropertyFilter? filter = null)
+		{
+			if(ids != null && ids.Length == 0)
+			{
+				ids = null;
+			}
+
+			UserId? userIdValue = null;
+			if (userId != null)
+			{
+				userIdValue = new UserId(userId);
+			}
+
+			List<IIssue> issues;
+			if (jobId == null)
+			{
+				StreamId? streamIdValue = null;
+				if (streamId != null)
+				{
+					streamIdValue = new StreamId(streamId);
+				}
+
+				issues = await _issueService.Collection.FindIssuesAsync(ids, userIdValue, streamIdValue, minChange ?? change, maxChange ?? change, resolved, promoted, index, count);
+			}
+			else
+			{
+				IJob? job = await _jobService.GetJobAsync(jobId.Value);
+				if (job == null)
+				{
+					return NotFound(jobId.Value);
+				}
+				if(!await _jobService.AuthorizeAsync(job, AclAction.ViewJob, User, null))
+				{
+					return Forbid(AclAction.ViewJob, jobId.Value);
+				}
+
+				IGraph graph = await _jobService.GetGraphAsync(job);
+				issues = await _issueService.Collection.FindIssuesForJobAsync(job, graph, stepId?.ToSubResourceId(), batchId?.ToSubResourceId(), labelIdx, userIdValue, resolved, promoted, index, count);
+			}
+
+			StreamPermissionsCache permissionsCache = new StreamPermissionsCache();
+
+			List<object> responses = new List<object>();
+			foreach (IIssue issue in issues)
+			{
+				IIssueDetails details = await _issueService.GetIssueDetailsAsync(issue);
+				if (await AuthorizeIssue(details, permissionsCache))
+				{
+					bool bShowDesktopAlerts = _issueService.ShowDesktopAlertsForIssue(issue, details.Spans);
+					GetIssueResponse response = await CreateIssueResponseAsync(details, bShowDesktopAlerts);
+					responses.Add(PropertyFilter.Apply(response, filter));
+				}
+			}
+			return responses;
+		}
+
+		/// <summary>
+		/// Retrieve information about a specific issue
+		/// </summary>
+		/// <param name="issueId">Id of the issue to get information about</param>
+		/// <param name="filter">Filter for the properties to return</param>
+		/// <returns>List of matching agents</returns>
+		[HttpGet]
+		[Route("/api/v1/issues/{issueId}")]
+		[ProducesResponseType(typeof(GetIssueResponse), 200)]
+		public async Task<ActionResult<object>> GetIssueAsync(int issueId, [FromQuery] PropertyFilter? filter = null)
+		{
+			IIssueDetails? details = await _issueService.GetIssueDetailsAsync(issueId);
+			if (details == null)
+			{
+				return NotFound();
+			}
+			if (!await AuthorizeIssue(details, null))
+			{
+				return Forbid();
+			}
+
+			bool bShowDesktopAlerts = _issueService.ShowDesktopAlertsForIssue(details.Issue, details.Spans);
+			return PropertyFilter.Apply(await CreateIssueResponseAsync(details, bShowDesktopAlerts), filter);
+		}
+
+		/// <summary>
+		/// Retrieve historical information about a specific issue
+		/// </summary>
+		/// <param name="issueId">Id of the agent to get information about</param>
+		/// <param name="minTime">Minimum time for records to return</param>
+		/// <param name="maxTime">Maximum time for records to return</param>
+		/// <param name="index">Offset of the first result</param>
+		/// <param name="count">Number of records to return</param>
+		/// <returns>Information about the requested agent</returns>
+		[HttpGet]
+		[Route("/api/v1/issues/{issueId}/history")]
+		public async Task GetAgentHistoryAsync(int issueId, [FromQuery] DateTime? minTime = null, [FromQuery] DateTime? maxTime = null, [FromQuery] int index = 0, [FromQuery] int count = 50)
+		{
+			Response.ContentType = "application/json";
+			Response.StatusCode = 200;
+			await Response.StartAsync();
+			await _issueCollection.GetLogger(issueId).FindAsync(Response.BodyWriter, minTime, maxTime, index, count);
+		}
+
+		/// <summary>
+		/// Create an issue response object
+		/// </summary>
+		/// <param name="details"></param>
+		/// <param name="showDesktopAlerts"></param>
+		/// <returns></returns>
+		async Task<GetIssueResponse> CreateIssueResponseAsync(IIssueDetails details, bool showDesktopAlerts)
+		{
+			List<GetIssueAffectedStreamResponse> affectedStreams = new List<GetIssueAffectedStreamResponse>();
+			foreach (IGrouping<StreamId, IIssueSpan> streamSpans in details.Spans.GroupBy(x => x.StreamId))
+			{
+				try
+				{
+					IStream? stream = await _streamService.GetCachedStream(streamSpans.Key);
+					affectedStreams.Add(new GetIssueAffectedStreamResponse(details, stream, streamSpans));
+				}
+				catch 
+				{
+					_logger.LogError("Unable to get {StreamId} for span key", streamSpans.Key);
+				}
+			}
+			return new GetIssueResponse(details, affectedStreams, showDesktopAlerts);
+		}
+
+		/// <summary>
+		/// Retrieve events for a specific issue
+		/// </summary>
+		/// <param name="issueId">Id of the issue to get information about</param>
+		/// <param name="filter">Filter for the properties to return</param>
+		/// <returns>List of matching agents</returns>
+		[HttpGet]
+		[Route("/api/v1/issues/{issueId}/streams")]
+		[ProducesResponseType(typeof(List<GetIssueStreamResponse>), 200)]
+		public async Task<ActionResult<List<object>>> GetIssueStreamsAsync(int issueId, [FromQuery] PropertyFilter? filter = null)
+		{
+			IIssueDetails? issue = await _issueService.GetIssueDetailsAsync(issueId);
+			if (issue == null)
+			{
+				return NotFound();
+			}
+
+			StreamPermissionsCache cache = new StreamPermissionsCache();
+			if (!await AuthorizeIssue(issue, cache))
+			{
+				return Forbid();
+			}
+
+			List<object> responses = new List<object>();
+			foreach (IGrouping<StreamId, IIssueSpan> spanGroup in issue.Spans.GroupBy(x => x.StreamId))
+			{
+				if (await _streamService.AuthorizeAsync(spanGroup.Key, AclAction.ViewStream, User, cache))
+				{
+					IStream? stream = await _streamService.GetCachedStream(spanGroup.Key);
+					if (stream != null)
+					{
+						HashSet<ObjectId> spanIds = new HashSet<ObjectId>(spanGroup.Select(x => x.Id));
+						List<IIssueStep> steps = issue.Steps.Where(x => spanIds.Contains(x.SpanId)).ToList();
+						responses.Add(PropertyFilter.Apply(new GetIssueStreamResponse(stream, spanGroup.ToList(), steps), filter));
+					}
+				}
+			}
+			return responses;
+		}
+
+		/// <summary>
+		/// Retrieve events for a specific issue
+		/// </summary>
+		/// <param name="issueId">Id of the issue to get information about</param>
+		/// <param name="streamId">The stream id</param>
+		/// <param name="filter">Filter for the properties to return</param>
+		/// <returns>List of matching agents</returns>
+		[HttpGet]
+		[Route("/api/v1/issues/{issueId}/streams/{streamId}")]
+		[ProducesResponseType(typeof(List<GetIssueStreamResponse>), 200)]
+		public async Task<ActionResult<object>> GetIssueStreamAsync(int issueId, string streamId, [FromQuery] PropertyFilter? filter = null)
+		{
+			IIssueDetails? details = await _issueService.GetIssueDetailsAsync(issueId);
+			if (details == null)
+			{
+				return NotFound();
+			}
+
+			StreamId streamIdValue = new StreamId(streamId);
+			if (!await _streamService.AuthorizeAsync(streamIdValue, AclAction.ViewStream, User, null))
+			{
+				return Forbid();
+			}
+
+			IStream? stream = await _streamService.GetCachedStream(streamIdValue);
+			if (stream == null)
+			{
+				return NotFound();
+			}
+
+			List<IIssueSpan> spans = details.Spans.Where(x => x.StreamId == streamIdValue).ToList();
+			if(spans.Count == 0)
+			{
+				return NotFound();
+			}
+
+			HashSet<ObjectId> spanIds = new HashSet<ObjectId>(spans.Select(x => x.Id));
+			List<IIssueStep> steps = details.Steps.Where(x => spanIds.Contains(x.SpanId)).ToList();
+
+			return PropertyFilter.Apply(new GetIssueStreamResponse(stream, spans, steps), filter);
+		}
+
+		/// <summary>
+		/// Retrieve events for a specific issue
+		/// </summary>
+		/// <param name="issueId">Id of the issue to get information about</param>
+		/// <param name="jobId">The job id to filter for</param>
+		/// <param name="batchId">The batch to filter by</param>
+		/// <param name="stepId">The step to filter by</param>
+		/// <param name="labelIdx">The label within the job to filter by</param>
+		/// <param name="logIds">List of log ids to return issues for</param>
+		/// <param name="index">Index of the first event</param>
+		/// <param name="count">Number of events to return</param>
+		/// <param name="filter">Filter for the properties to return</param>
+		/// <returns>List of matching agents</returns>
+		[HttpGet]
+		[Route("/api/v1/issues/{issueId}/events")]
+		[ProducesResponseType(typeof(List<GetLogEventResponse>), 200)]
+		public async Task<ActionResult<List<object>>> GetIssueEventsAsync(int issueId, [FromQuery] JobId? jobId = null, [FromQuery] string? batchId = null, [FromQuery] string? stepId = null, [FromQuery(Name = "label")] int? labelIdx = null, [FromQuery] string[]? logIds = null, [FromQuery] int index = 0, [FromQuery] int count = 10, [FromQuery] PropertyFilter? filter = null)
+		{
+			HashSet<LogId> logIdValues = new HashSet<LogId>();
+			if(jobId != null)
+			{
+				IJob? job = await _jobService.GetJobAsync(jobId.Value);
+				if(job == null)
+				{
+					return NotFound();
+				}
+
+				if (stepId != null)
+				{
+					IJobStep? step;
+					if (job.TryGetStep(stepId.ToSubResourceId(), out step) && step.Outcome != JobStepOutcome.Success && step.LogId != null)
+					{
+						logIdValues.Add(step.LogId.Value);
+					}
+				}
+				else if (batchId != null)
+				{
+					IJobStepBatch? batch;
+					if (job.TryGetBatch(batchId.ToSubResourceId(), out batch))
+					{
+						logIdValues.UnionWith(batch.Steps.Where(x => x.Outcome != JobStepOutcome.Success && x.LogId != null).Select(x => x.LogId!.Value));
+					}
+				}
+				else if (labelIdx != null)
+				{
+					IGraph graph = await _jobService.GetGraphAsync(job);
+
+					HashSet<NodeRef> includedNodes = new HashSet<NodeRef>(graph.Labels[labelIdx.Value].IncludedNodes);
+
+					foreach (IJobStepBatch batch in job.Batches)
+					{
+						foreach (IJobStep step in batch.Steps)
+						{
+							NodeRef nodeRef = new NodeRef(batch.GroupIdx, step.NodeIdx);
+							if (step.Outcome != JobStepOutcome.Success && step.LogId != null && includedNodes.Contains(nodeRef))
+							{
+								logIdValues.Add(step.LogId.Value);
+							}
+						}
+					}
+				}
+				else
+				{
+					logIdValues.UnionWith(job.Batches.SelectMany(x => x.Steps).Where(x => x.Outcome != JobStepOutcome.Success && x.LogId != null).Select(x => x.LogId!.Value));
+				}
+			}
+			if(logIds != null)
+			{
+				logIdValues.UnionWith(logIds.Select(x => new LogId(x)));
+			}
+
+			List<IIssueSpan> spans = await _issueCollection.FindSpansAsync(issueId);
+			List<ILogEvent> events = await _logFileService.FindEventsForSpansAsync(spans.Select(x => x.Id), logIdValues.ToArray(), index, count);
+
+			JobPermissionsCache permissionsCache = new JobPermissionsCache();
+			Dictionary<LogId, ILogFile?> logFiles = new Dictionary<LogId, ILogFile?>();
+
+			List<object> responses = new List<object>();
+			foreach (ILogEvent logEvent in events)
+			{
+				ILogFile? logFile;
+				if (!logFiles.TryGetValue(logEvent.LogId, out logFile))
+				{
+					logFile = await _logFileService.GetLogFileAsync(logEvent.LogId);
+					logFiles[logEvent.LogId] = logFile;
+				}
+				if (logFile != null && await _jobService.AuthorizeAsync(logFile.JobId, AclAction.ViewLog, User, permissionsCache))
+				{
+					ILogEventData data = await _logFileService.GetEventDataAsync(logFile, logEvent.LineIndex, logEvent.LineCount);
+					GetLogEventResponse response = new GetLogEventResponse(logEvent, data, issueId);
+					responses.Add(PropertyFilter.Apply(response, filter));
+				}
+			}
+			return responses;
+		}
+
+		/// <summary>
+		/// Authorize the current user to see an issue
+		/// </summary>
+		/// <param name="issue">The issue to authorize</param>
+		/// <param name="permissionsCache">Cache of permissions</param>
+		/// <returns>True if the user is authorized to see the issue</returns>
+		private async Task<bool> AuthorizeIssue(IIssueDetails issue, StreamPermissionsCache? permissionsCache)
+		{
+			foreach (StreamId streamId in issue.Spans.Select(x => x.StreamId).Distinct())
+			{
+				if (await _streamService.AuthorizeAsync(streamId, AclAction.ViewStream, User, permissionsCache))
+				{
+					return true;
+				}
+			}
+			return false;
+		}
+
+		/// <summary>
+		/// Update an issue
+		/// </summary>
+		/// <param name="issueId">Id of the issue to get information about</param>
+		/// <param name="request">The update information</param>
+		/// <returns>List of matching agents</returns>
+		[HttpPut]
+		[Route("/api/v1/issues/{issueId}")]
+		public async Task<ActionResult> UpdateIssueAsync(int issueId, [FromBody] UpdateIssueRequest request)
+		{
+			UserId? newOwnerId = null;
+			if (request.OwnerId != null)
+			{
+				newOwnerId = request.OwnerId.Length == 0 ? UserId.Empty : new UserId(request.OwnerId);
+			}
+
+			UserId? newNominatedById = null;
+			if (request.NominatedById != null)
+			{
+				newNominatedById = new UserId(request.NominatedById);
+			}
+			else if (request.OwnerId != null)
+			{
+				newNominatedById = User.GetUserId();
+			}
+
+			UserId? newDeclinedById = null;
+			if (request.Declined ?? false)
+			{
+				newDeclinedById = User.GetUserId();
+			}
+
+			UserId? newResolvedById = null;
+			if (request.Resolved.HasValue)
+			{
+				newResolvedById = request.Resolved.Value ? User.GetUserId() : UserId.Empty;
+			}
+
+			UserId? newQuarantinedById = null;
+			if (request.QuarantinedById != null)
+			{
+				newQuarantinedById = request.QuarantinedById.Length > 0 ? new UserId(request.QuarantinedById) : UserId.Empty;
+			}
+
+
+			List<ObjectId>? addSpans = null;
+			if (request.AddSpans != null && request.AddSpans.Count > 0)
+			{
+				addSpans = request.AddSpans.ConvertAll(x => ObjectId.Parse(x));
+			}
+
+			List<ObjectId>? removeSpans = null;
+			if (request.RemoveSpans != null && request.RemoveSpans.Count > 0)
+			{
+				removeSpans = request.RemoveSpans.ConvertAll(x => ObjectId.Parse(x));
+			}
+
+			if (!await _issueService.UpdateIssueAsync(issueId, request.Summary, request.Description, request.Promoted, newOwnerId, newNominatedById, request.Acknowledged, newDeclinedById, request.FixChange, newResolvedById, addSpans, removeSpans, request.ExternalIssueKey, newQuarantinedById))
+			{
+				return NotFound();
+			}
+			return Ok();
+		}
+
+		// External issue tracking
+
+		/// <summary>
+		/// Get external issue information for provided keys
+		/// </summary>
+		[HttpGet]
+		[Authorize]
+		[Route("/api/v1/issues/external")]
+		public async Task<ActionResult<List<GetExternalIssueResponse>>> GetExternalIssuesAsync([FromQuery] string streamId, [FromQuery] string[] keys)
+		{
+			if (_externalIssueService == null)
+			{
+				return NotFound();
+			}
+
+			StreamId streamIdValue = new StreamId(streamId);
+			StreamPermissionsCache cache = new StreamPermissionsCache();
+
+			if (!await _streamService.AuthorizeAsync(streamIdValue, AclAction.ViewStream, User, cache))
+			{
+				return Forbid();
+			}
+
+			List<GetExternalIssueResponse> response = new List<GetExternalIssueResponse>();
+
+			if (keys.Length != 0)
+			{
+				List<IExternalIssue> issues = await _externalIssueService.GetIssuesAsync(keys);
+
+				for (int i = 0; i < issues.Count; i++)
+				{
+					response.Add(new GetExternalIssueResponse(issues[i]));
+				}
+			}
+
+			return response;
+		}
+
+		/// <summary>
+		/// Create a new external issue
+		/// </summary>
+		[HttpPost]
+		[Authorize]
+		[Route("/api/v1/issues/external")]
+		public async Task<ActionResult<CreateExternalIssueResponse>> CreateExternalIssueAsync([FromBody] CreateExternalIssueRequest issueRequest)
+		{
+			if (_externalIssueService == null)
+			{
+				return NotFound();
+			}
+
+			StreamId streamIdValue = new StreamId(issueRequest.StreamId);
+			StreamPermissionsCache cache = new StreamPermissionsCache();
+
+			if (!await _streamService.AuthorizeAsync(streamIdValue, AclAction.ViewStream, User, cache))
+			{
+				return Forbid();
+			}
+
+			IUser? user = await _userCollection.GetUserAsync(User);
+
+			if (user == null)
+			{
+				return BadRequest($"Missing user for {User.GetUserName()}");
+			}
+
+			(string? key, string? url) = await _externalIssueService.CreateIssueAsync(user, User.GetExternalIssueUser(), issueRequest.IssueId, issueRequest.Summary, issueRequest.ProjectId, issueRequest.ComponentId, issueRequest.IssueTypeId, issueRequest.Description, issueRequest.HordeIssueLink);
+
+			if (key == null)
+			{
+				throw new Exception($"Unable to create external issue");
+			}
+
+			return new CreateExternalIssueResponse(key, url);
+		}
+
+		/// <summary>
+		/// Gets external issue tracking projects associated with the provided stream
+		/// </summary>
+		/// <param name="streamId"></param>
+		/// <returns></returns>
+		[HttpGet]
+		[Authorize]
+		[Route("/api/v1/issues/external/projects")]
+		public async Task<ActionResult<List<GetExternalIssueProjectResponse>>> GetExternalIssueProjectsAsync([FromQuery] string streamId)
+		{
+			if (_externalIssueService == null)
+			{
+				return NotFound();
+			}
+
+			StreamId streamIdValue = new StreamId(streamId);
+			StreamPermissionsCache cache = new StreamPermissionsCache();
+
+			if (!await _streamService.AuthorizeAsync(streamIdValue, AclAction.ViewStream, User, cache))
+			{
+				return Forbid();
+			}
+
+			IStream? stream = await _streamService.GetStreamAsync(streamIdValue);
+
+			List<IExternalIssueProject> projects = await _externalIssueService.GetProjects(stream!);
+			List<GetExternalIssueProjectResponse> response = new List<GetExternalIssueProjectResponse>();
+
+			projects.ForEach(project =>
+			{
+				response.Add(new GetExternalIssueProjectResponse(project.Key, project.Name, project.Id, project.Components, project.IssueTypes));
+			});
+
+			return response;
+		}
+
+
+	}
+}
diff --git a/Engine/Source/Programs/Horde/Horde.Build/Jobs/ArtifactCollection.cs b/Engine/Source/Programs/Horde/Horde.Build/Jobs/ArtifactCollection.cs
new file mode 100644
--- /dev/null
+++ b/Engine/Source/Programs/Horde/Horde.Build/Jobs/ArtifactCollection.cs
@@ -0,0 +1,274 @@
+// Copyright Epic Games, Inc. All Rights Reserved.
+
+using System;
+using System.Collections.Generic;
+using System.Linq;
+using System.Threading.Tasks;
+using Horde.Build.Jobs;
+using Horde.Build.Models;
+using Horde.Build.Server;
+using Horde.Build.Storage;
+using Horde.Build.Utilities;
+using MongoDB.Bson;
+using MongoDB.Bson.Serialization.Attributes;
+using MongoDB.Driver;
+
+namespace Horde.Build.Services
+{
+	using JobId = ObjectId<IJob>;
+
+	/// <summary>
+	/// Wraps functionality for manipulating artifacts
+	/// </summary>
+	public class ArtifactCollection : IArtifactCollection
+	{
+		class Artifact : IArtifact
+		{
+			[BsonRequired, BsonId]
+			public ObjectId Id { get; set; }
+
+			[BsonRequired]
+			public JobId JobId { get; set; }
+
+			[BsonRequired]
+			public string Name { get; set; }
+
+			public SubResourceId? StepId { get; set; }
+
+			[BsonRequired]
+			public long Length { get; set; }
+
+			public string MimeType { get; set; }
+
+			[BsonRequired]
+			public int UpdateIndex { get; set; }
+
+			[BsonConstructor]
+			private Artifact()
+			{
+				Name = null!;
+				MimeType = null!;
+			}
+
+			public Artifact(JobId jobId, SubResourceId? stepId, string name, long length, string mimeType)
+			{
+				Id = ObjectId.GenerateNewId();
+				JobId = jobId;
+				StepId = stepId;
+				Name = name;
+				Length = length;
+				MimeType = mimeType;
+			}
+		}
+
+		private readonly IStorageBackend _storageBackend;
+		private readonly IMongoCollection<Artifact> _artifacts;
+
+		/// <summary>
+		/// Constructor
+		/// </summary>
+		/// <param name="mongoService">The database service</param>
+		/// <param name="storageBackend">The storage backend</param>
+		public ArtifactCollection(MongoService mongoService, IStorageBackend<ArtifactCollection> storageBackend)
+		{
+			_storageBackend = storageBackend;
+
+			// Initialize Artifacts table
+			_artifacts = mongoService.GetCollection<Artifact>("Artifacts", keys => keys.Ascending(x => x.JobId));
+		}
+
+		/// <summary>
+		/// Creates a new artifact
+		/// </summary>
+		/// <param name="jobId">Unique id of the job that owns this artifact</param>
+		/// <param name="stepId">Optional Step id</param>
+		/// <param name="name">Name of artifact</param>
+		/// <param name="mimeType">Type of artifact</param>
+		/// <param name="data">The data to write</param>
+		/// <returns>The new log file document</returns>
+		public async Task<IArtifact> CreateArtifactAsync(JobId jobId, SubResourceId? stepId, string name, string mimeType, System.IO.Stream data)
+		{
+			// upload first
+			string artifactName = ValidateName(name);
+			await _storageBackend.WriteAsync(GetPath(jobId, stepId, artifactName), data);
+
+			// then create entry
+			Artifact newArtifact = new Artifact(jobId, stepId, artifactName, data.Length, mimeType);
+			await _artifacts.InsertOneAsync(newArtifact);
+			return newArtifact;
+		}
+
+		/// <summary>
+		/// Gets all the available artifacts for a job
+		/// </summary>
+		/// <param name="jobId">Unique id of the job to query</param>
+		/// <param name="stepId">Unique id of the Step to query</param>
+		/// <param name="name">Name of the artifact</param>
+		/// <returns>List of artifact documents</returns>
+		public async Task<List<IArtifact>> GetArtifactsAsync(JobId? jobId, SubResourceId? stepId, string? name)
+		{
+			FilterDefinitionBuilder<Artifact> builder = Builders<Artifact>.Filter;
+
+			FilterDefinition<Artifact> filter = FilterDefinition<Artifact>.Empty;
+			if (jobId != null)
+			{
+				filter &= builder.Eq(x => x.JobId, jobId.Value);
+			}
+			if (stepId != null)
+			{
+				filter &= builder.Eq(x => x.StepId, stepId.Value);
+			}
+			if (name != null)
+			{
+				filter &= builder.Eq(x => x.Name, name);
+			}
+
+			return await _artifacts.Find(filter).ToListAsync<Artifact, IArtifact>();
+		}
+
+		/// <summary>
+		/// Gets a specific list of artifacts based on id
+		/// </summary>
+		/// <param name="artifactIds">The list of artifact Ids</param>
+		/// <returns>List of artifact documents</returns>
+		public async Task<List<IArtifact>> GetArtifactsAsync(IEnumerable<ObjectId> artifactIds)
+		{
+			FilterDefinitionBuilder<Artifact> builder = Builders<Artifact>.Filter;
+
+			FilterDefinition<Artifact> filter = FilterDefinition<Artifact>.Empty;
+			filter &= builder.In(x => x.Id, artifactIds);
+			
+			return await _artifacts.Find(filter).ToListAsync<Artifact, IArtifact>();
+		}
+
+		/// <summary>
+		/// Gets an artifact by ID
+		/// </summary>
+		/// <param name="artifactId">Unique id of the artifact</param>
+		/// <returns>The artifact document</returns>
+		public async Task<IArtifact?> GetArtifactAsync(ObjectId artifactId)
+		{
+			return await _artifacts.Find<Artifact>(x => x.Id == artifactId).FirstOrDefaultAsync();
+		}
+
+		/// <summary>
+		/// Attempts to update a single artifact document, if the update counter is valid
+		/// </summary>
+		/// <param name="current">The current artifact document</param>
+		/// <param name="update">The update definition</param>
+		/// <returns>True if the document was updated, false if another writer updated the document first</returns>
+		private async Task<bool> TryUpdateArtifactAsync(Artifact current, UpdateDefinition<Artifact> update)
+		{
+			UpdateResult result = await _artifacts.UpdateOneAsync<Artifact>(x => x.Id == current.Id && x.UpdateIndex == current.UpdateIndex, update.Set(x => x.UpdateIndex, current.UpdateIndex + 1));
+			return result.ModifiedCount == 1;
+		}
+
+		/// <summary>
+		/// Updates an artifact
+		/// </summary>
+		/// <param name="artifact">The artifact</param>
+		/// <param name="newMimeType">New mime type</param>
+		/// <param name="newData">New data</param>
+		/// <returns>Async task</returns>
+		public async Task<bool> UpdateArtifactAsync(IArtifact? artifact, string newMimeType, System.IO.Stream newData)
+		{
+			while (artifact != null)
+			{
+				UpdateDefinitionBuilder<Artifact> updateBuilder = Builders<Artifact>.Update;
+
+				List<UpdateDefinition<Artifact>> updates = new List<UpdateDefinition<Artifact>>();
+				updates.Add(updateBuilder.Set(x => x.MimeType, newMimeType));
+				updates.Add(updateBuilder.Set(x => x.Length, newData.Length));
+
+				// re-upload the data to external
+				string artifactName = ValidateName(artifact.Name);
+				await _storageBackend.WriteAsync(GetPath(artifact.JobId, artifact.StepId, artifactName), newData);
+
+				if (await TryUpdateArtifactAsync((Artifact)artifact, updateBuilder.Combine(updates)))
+				{
+					return true;
+				}
+
+				artifact = await GetArtifactAsync(artifact.Id);
+			}
+			return false;
+		}
+
+		/// <summary>
+		/// gets artifact data
+		/// </summary>
+		/// <param name="artifact">The artifact</param>
+		/// <returns>The chunk data</returns>
+		public async Task<System.IO.Stream> OpenArtifactReadStreamAsync(IArtifact artifact)
+		{
+			System.IO.Stream? stream = await _storageBackend.ReadAsync(GetPath(artifact.JobId, artifact.StepId, artifact.Name));
+			if (stream == null)
+			{
+				throw new Exception($"Unable to get artifact {artifact.Id}");
+			}
+			return stream;
+		}
+
+		/// <summary>
+		/// Get the path for an artifact
+		/// </summary>
+		/// <param name="jobId"></param>
+		/// <param name="stepId"></param>
+		/// <param name="name"></param>
+		/// <returns></returns>
+		private static string GetPath(JobId jobId, SubResourceId? stepId, string name)
+		{
+			if (stepId == null)
+			{
+				return $"{jobId}/{name}";
+			}
+			else
+			{
+				return $"{jobId}/{stepId.Value}/{name}";
+			}
+		}
+
+		/// <summary>
+		/// Checks that the name given is valid for an artifact
+		/// </summary>
+		/// <param name="name"></param>
+		private static string ValidateName(string name)
+		{
+			string newName = name.Replace('\\', '/');
+			if (newName.Length == 0)
+			{
+				throw new ArgumentException("Artifact name is empty");
+			}
+			else if (newName.StartsWith('/'))
+			{
+				throw new ArgumentException($"Artifact has an absolute path ({newName})");
+			}
+			else if (newName.EndsWith('/'))
+			{
+				throw new ArgumentException($"Artifact does not have a file name ({newName})");
+			}
+			else if (newName.Contains("//", StringComparison.Ordinal))
+			{
+				throw new ArgumentException($"Artifact name contains an invalid directory ({newName})");
+			}
+
+			string invalidChars = ":<>|\"";
+			for (int idx = 0; idx < newName.Length; idx++)
+			{
+				int minIdx = idx;
+				for (; idx < newName.Length && newName[idx] != '/'; idx++)
+				{
+					if (invalidChars.Contains(newName[idx], StringComparison.Ordinal))
+					{
+						throw new ArgumentException($"Invalid character in artifact name ({newName})");
+					}
+				}
+				if ((idx == minIdx + 1 && newName[minIdx] == '.') || (idx == minIdx + 2 && newName[minIdx] == '.' && newName[minIdx + 1] == '.'))
+				{
+					throw new ArgumentException($"Artifact may not contain symbolic directory names ({newName})");
+				}
+			}
+			return newName;
+		}
+	}
+}
\ No newline at end of file
diff --git a/Engine/Source/Programs/Horde/Horde.Build/Jobs/Artifacts/ArtifactsController.cs b/Engine/Source/Programs/Horde/Horde.Build/Jobs/Artifacts/ArtifactsController.cs
new file mode 100644
--- /dev/null
+++ b/Engine/Source/Programs/Horde/Horde.Build/Jobs/Artifacts/ArtifactsController.cs
@@ -0,0 +1,443 @@
+// Copyright Epic Games, Inc. All Rights Reserved.
+
+using System;
+using System.Collections.Generic;
+using System.IdentityModel.Tokens.Jwt;
+using System.IO;
+using System.IO.Compression;
+using System.Linq;
+using System.Net.Mime;
+using System.Security.Claims;
+using System.Threading.Tasks;
+using Horde.Build.Acls;
+using Horde.Build.Api;
+using Horde.Build.Jobs;
+using Horde.Build.Models;
+using Horde.Build.Server;
+using Horde.Build.Services;
+using Horde.Build.Utilities;
+using Microsoft.AspNetCore.Authorization;
+using Microsoft.AspNetCore.Http;
+using Microsoft.AspNetCore.Mvc;
+using Microsoft.IdentityModel.Tokens;
+using MongoDB.Bson;
+
+namespace Horde.Build.Controllers
+{
+	using JobId = ObjectId<IJob>;
+
+	/// <summary>
+	/// Controller for the /api/artifacts endpoint
+	/// </summary>
+	[ApiController]
+	[Route("[controller]")]
+	public class ArtifactsController : ControllerBase
+	{
+		/// <summary>
+		/// Instance of the database service
+		/// </summary>
+		private readonly MongoService _mongoService;
+
+		/// <summary>
+		/// Instance of the artifact collection
+		/// </summary>
+		private readonly IArtifactCollection _artifactCollection;
+
+		/// <summary>
+		/// Instance of the ACL service
+		/// </summary>
+		private readonly AclService _aclService;
+
+		/// <summary>
+		/// Instance of the Job service
+		/// </summary>
+		private readonly JobService _jobService;
+
+		/// <summary>
+		/// Constructor
+		/// </summary>
+		public ArtifactsController(MongoService mongoService, IArtifactCollection artifactCollection, AclService aclService, JobService jobService)
+		{
+			_mongoService = mongoService;
+			_artifactCollection = artifactCollection;
+			_aclService = aclService;
+			_jobService = jobService;
+		}
+
+		/// <summary>
+		/// Creates an artifact
+		/// </summary>
+		/// <param name="jobId">BatchId</param>
+		/// <param name="stepId">StepId</param>
+		/// <param name="file">The file contents</param>
+		/// <returns>Http result code</returns>
+		[HttpPost]
+		[Authorize]
+		[Route("/api/v1/artifacts")]
+		public async Task<ActionResult<CreateArtifactResponse>> CreateArtifact([FromQuery] JobId jobId, [FromQuery]string? stepId, IFormFile file)
+		{
+			IJob? job = await _jobService.GetJobAsync(jobId);
+			if(job == null)
+			{
+				return NotFound();
+			}
+
+			if (!await _jobService.AuthorizeAsync(job, AclAction.UploadArtifact, User, null))
+			{
+				return Forbid();
+			}
+
+			IJobStep? step = null;
+			if(stepId != null)
+			{
+				foreach(IJobStepBatch batch in job.Batches)
+				{
+					if(batch.TryGetStep(stepId.ToSubResourceId(), out step))
+					{
+						break;
+					}
+				}
+				if(step == null)
+				{
+					// if the step doesn't exist in any of the batches, not found
+					return NotFound();
+				}
+			}
+
+			IArtifact newArtifact = await _artifactCollection.CreateArtifactAsync(job.Id, step?.Id, file.FileName, file.ContentType ?? "horde-mime/unknown", file.OpenReadStream());
+			return new CreateArtifactResponse(newArtifact.Id.ToString());
+		}
+
+		/// <summary>
+		/// Updates an artifact
+		/// </summary>
+		/// <param name="artifactId">JobId</param>
+		/// <param name="file">The file contents</param>
+		/// <returns>Http result code</returns>
+		[HttpPut]
+		[Authorize]
+		[Route("/api/v1/artifacts/{artifactId}")]
+		public async Task<ActionResult<CreateArtifactResponse>> UpdateArtifact(string artifactId, IFormFile file)
+		{
+			IArtifact? artifact = await _artifactCollection.GetArtifactAsync(artifactId.ToObjectId());
+			if (artifact == null)
+			{
+				return NotFound();
+			}
+			if (!await _jobService.AuthorizeAsync(artifact.JobId, AclAction.UploadArtifact, User, null))
+			{
+				return Forbid();
+			}
+
+			await _artifactCollection.UpdateArtifactAsync(artifact, file.ContentType ?? "horde-mime/unknown", file.OpenReadStream());
+			return Ok();
+		}
+
+		/// <summary>
+		/// Query artifacts for a job step
+		/// </summary>
+		/// <param name="jobId">Optional JobId to filter by</param>
+		/// <param name="stepId">Optional StepId to filter by</param>
+		/// <param name="code">Whether to generate a direct download code</param>
+		/// <param name="filter">Filter for the properties to return</param>
+		/// <returns>Information about all the artifacts</returns>
+		[HttpGet]
+		[Authorize]
+		[Route("/api/v1/artifacts")]
+		[ProducesResponseType(typeof(List<GetArtifactResponse>), 200)]
+		public async Task<ActionResult<List<object>>> GetArtifacts([FromQuery] JobId jobId, [FromQuery] string? stepId = null, [FromQuery] bool code = false, [FromQuery] PropertyFilter? filter = null)
+		{
+			if (!await _jobService.AuthorizeAsync(jobId, AclAction.DownloadArtifact, User, null))
+			{
+				return Forbid();
+			}
+
+			string? downloadCode = code ? (string?)GetDirectDownloadCodeForJob(jobId) : null;
+
+			List<IArtifact> artifacts = await _artifactCollection.GetArtifactsAsync(jobId, stepId?.ToSubResourceId(), null);
+			return artifacts.ConvertAll(x => new GetArtifactResponse(x, downloadCode).ApplyFilter(filter));
+		}
+
+		/// <summary>
+		/// Gets the claim required to download artifacts for a particular job
+		/// </summary>
+		/// <param name="jobId">The job id</param>
+		/// <returns>The required claim</returns>
+		static Claim GetDirectDownloadClaim(JobId jobId)
+		{
+			return new Claim(HordeClaimTypes.JobArtifacts, jobId.ToString());
+		}
+
+		/// <summary>
+		/// Get a download code for the artifacts of a job
+		/// </summary>
+		/// <param name="jobId">The job id</param>
+		/// <returns>The download code</returns>
+		string GetDirectDownloadCodeForJob(JobId jobId)
+		{
+			Claim downloadClaim = GetDirectDownloadClaim(jobId);
+			return _aclService.IssueBearerToken(new[] { downloadClaim }, TimeSpan.FromHours(4.0));
+		}
+
+		/// <summary>
+		/// Retrieve metadata about a specific artifact
+		/// </summary>
+		/// <param name="artifactId">Id of the artifact to get information about</param>
+		/// <param name="code">Whether to generate a direct download code</param>
+		/// <param name="filter">Filter for the properties to return</param>
+		/// <returns>Information about the requested project</returns>
+		[HttpGet]
+		[Authorize]
+		[Route("/api/v1/artifacts/{artifactId}")]
+		[ProducesResponseType(typeof(GetArtifactResponse), 200)]
+		public async Task<ActionResult<object>> GetArtifact(string artifactId, bool code = false, [FromQuery] PropertyFilter? filter = null)
+		{
+			IArtifact? artifact = await _artifactCollection.GetArtifactAsync(artifactId.ToObjectId());
+			if (artifact == null)
+			{
+				return NotFound();
+			}
+			if (!await _jobService.AuthorizeAsync(artifact.JobId, AclAction.DownloadArtifact, User, null))
+			{
+				return Forbid();
+			}
+
+			string? downloadCode = code? (string?)GetDirectDownloadCodeForJob(artifact.JobId) : null;
+			return new GetArtifactResponse(artifact, downloadCode).ApplyFilter(filter);
+		}
+
+		/// <summary>
+		/// Retrieve raw data for an artifact
+		/// </summary>
+		/// <param name="artifactId">Id of the artifact to get information about</param>
+		/// <returns>Raw artifact data</returns>
+		[HttpGet]
+		[Authorize]
+		[Route("/api/v1/artifacts/{artifactId}/data")]
+		public async Task<ActionResult> GetArtifactData(string artifactId)
+		{
+			IArtifact? artifact = await _artifactCollection.GetArtifactAsync(artifactId.ToObjectId());
+			if (artifact == null)
+			{
+				return NotFound();
+			}
+			if (!await _jobService.AuthorizeAsync(artifact.JobId, AclAction.DownloadArtifact, User, null))
+			{
+				return Forbid();
+			}
+
+			// Fun, filestream result automatically closes the stream!
+			return new FileStreamResult(await _artifactCollection.OpenArtifactReadStreamAsync(artifact), artifact.MimeType);
+		}
+		
+		/// <summary>
+		/// Retrieve raw data for an artifact by filename
+		/// </summary>
+		/// <param name="jobId">Unique id for the job</param>
+		/// <param name="stepId">Unique id for the step</param>
+		/// <param name="filename">Filename of artifact from step</param>
+		/// <returns>Raw artifact data</returns>
+		[HttpGet]
+		[Route("/api/v1/jobs/{jobId}/steps/{stepId}/artifacts/{filename}/data")]
+		public async Task<ActionResult<object>> GetArtifactDataByFilename(JobId jobId, string stepId, string filename)
+		{
+			SubResourceId stepIdValue = stepId.ToSubResourceId();
+
+			if (!await _jobService.AuthorizeAsync(jobId, AclAction.DownloadArtifact, User, null))
+			{
+				return Forbid();
+			}
+			
+			List<IArtifact> artifacts = await _artifactCollection.GetArtifactsAsync(jobId, stepIdValue, filename);
+			if (artifacts.Count == 0)
+			{
+				return NotFound();
+			}
+
+			IArtifact artifact = artifacts[0];
+			return new FileStreamResult(await _artifactCollection.OpenArtifactReadStreamAsync(artifact), artifact.MimeType);
+		}
+
+		/// <summary>
+		/// Class to return a file stream without the "content-disposition: attachment" header
+		/// </summary>
+		class InlineFileStreamResult : FileStreamResult
+		{
+			/// <summary>
+			/// The suggested download filename
+			/// </summary>
+			readonly string _fileName;
+
+			/// <summary>
+			/// Constructor
+			/// </summary>
+			/// <param name="stream"></param>
+			/// <param name="mimeType"></param>
+			/// <param name="fileName"></param>
+			public InlineFileStreamResult(System.IO.Stream stream, string mimeType, string fileName)
+				: base(stream, mimeType)
+			{
+				_fileName = fileName;
+			}
+
+			/// <inheritdoc/>
+			public override Task ExecuteResultAsync(ActionContext context)
+			{
+				ContentDisposition contentDisposition = new ContentDisposition();
+				contentDisposition.Inline = true;
+				contentDisposition.FileName = _fileName;
+				context.HttpContext.Response.Headers.Add("Content-Disposition", contentDisposition.ToString());
+				
+				return base.ExecuteResultAsync(context);
+			}
+		}
+
+		/// <summary>
+		/// Retrieve raw data for an artifact
+		/// </summary>
+		/// <param name="artifactId">Id of the artifact to get information about</param>
+		/// <param name="code">The authorization code for this resource</param>
+		/// <returns>Raw artifact data</returns>
+		[HttpGet]
+		[AllowAnonymous]
+		[Route("/api/v1/artifacts/{artifactId}/download")]
+		public async Task<ActionResult> DownloadArtifact(string artifactId, [FromQuery] string code)
+		{
+			TokenValidationParameters parameters = new TokenValidationParameters();
+			parameters.ValidateAudience = false;
+			parameters.RequireExpirationTime = true;
+			parameters.ValidateLifetime = true;
+			parameters.ValidIssuer = _mongoService.JwtIssuer;
+			parameters.ValidateIssuer = true;
+			parameters.ValidateIssuerSigningKey = true;
+			parameters.IssuerSigningKey = _mongoService.JwtSigningKey;
+
+			JwtSecurityTokenHandler handler = new JwtSecurityTokenHandler();
+			ClaimsPrincipal principal = handler.ValidateToken(code, parameters, out _);
+
+			IArtifact? artifact = await _artifactCollection.GetArtifactAsync(artifactId.ToObjectId());
+			if (artifact == null)
+			{
+				return NotFound();
+			}
+
+			Claim directDownloadClaim = GetDirectDownloadClaim(artifact.JobId);
+			if (!principal.HasClaim(directDownloadClaim.Type, directDownloadClaim.Value))
+			{
+				return Forbid();
+			}
+
+			return new InlineFileStreamResult(await _artifactCollection.OpenArtifactReadStreamAsync(artifact), artifact.MimeType, Path.GetFileName(artifact.Name));
+		}
+
+		/// <summary>
+		/// Returns a zip archive of many artifacts
+		/// </summary>
+		/// <param name="artifactZipRequest">Artifact request params</param>
+		/// <returns>Zip of many artifacts</returns>
+		[HttpPost]
+		[Authorize]
+		[Route("/api/v1/artifacts/zip")]
+		public async Task<ActionResult> ZipArtifacts(GetArtifactZipRequest artifactZipRequest)
+		{
+			if (artifactZipRequest.JobId == null)
+			{
+				return BadRequest("Must specify a JobId");
+			}
+
+			IJob? job = await _jobService.GetJobAsync(new JobId(artifactZipRequest.JobId!));
+			if (job == null)
+			{
+				return NotFound();
+			}
+			if (!await _jobService.AuthorizeAsync(job, AclAction.DownloadArtifact, User, null))
+			{
+				return Forbid();
+			}
+
+			List<IArtifact> artifacts = await _artifactCollection.GetArtifactsAsync(job.Id, artifactZipRequest.StepId?.ToSubResourceId(), null);
+
+			Dictionary<ObjectId, IArtifact> idToArtifact = artifacts.ToDictionary(x => x.Id, x => x);
+
+			List<IArtifact> zipArtifacts;
+			if (artifactZipRequest.ArtifactIds == null)
+			{
+				zipArtifacts = artifacts;
+			}
+			else
+			{
+				zipArtifacts = new List<IArtifact>();
+				foreach (string artifactId in artifactZipRequest.ArtifactIds)
+				{
+					IArtifact? artifact;
+					if (idToArtifact.TryGetValue(artifactId.ToObjectId(), out artifact))
+					{
+						zipArtifacts.Add(artifact);
+					}
+					else
+					{
+						return NotFound();
+					}
+				}
+			}
+
+			IGraph graph = await _jobService.GetGraphAsync(job);
+
+			return new CustomFileCallbackResult("Artifacts.zip", "application/octet-stream", false, async (outputStream, context) =>
+			{
+				// Make an unseekable MemoryStream for the ZipArchive. We have to do this because the ZipEntry stream falls back to a synchronous write to it's own stream wrappers.
+				using (CustomBufferStream zipOutputStream = new CustomBufferStream())
+				{
+					// Keep the stream open after dispose so we can write the EOF bits.
+					using (ZipArchive zipArchive = new ZipArchive(zipOutputStream, ZipArchiveMode.Create, true))
+					{
+						foreach (IArtifact artifact in zipArtifacts)
+						{
+							await using (System.IO.Stream artifactStream = await _artifactCollection.OpenArtifactReadStreamAsync(artifact))
+							{
+								// tack on the step name into the directory if it exists
+								string stepName = String.Empty;
+								if (artifact.StepId.HasValue)
+								{
+									foreach (IJobStepBatch batch in job.Batches)
+									{
+										IJobStep? step;
+										if (batch.TryGetStep(artifact.StepId.Value, out step))
+										{
+											stepName = graph.Groups[batch.GroupIdx].Nodes[step.NodeIdx].Name;
+											break;
+										}
+									}
+								}
+
+								ZipArchiveEntry zipEntry = zipArchive.CreateEntry(artifact.Name);
+								using (System.IO.Stream entryStream = zipEntry.Open())
+								{
+									byte[] buffer = new byte[4096];
+									int totalBytesRead = 0;
+									while (totalBytesRead < artifact.Length)
+									{
+										int bytesRead = await artifactStream.ReadAsync(buffer, 0, buffer.Length);
+
+										// Write bytes to the entry stream.  Also advances the MemStream pos.
+										await entryStream.WriteAsync(buffer, 0, bytesRead);
+										// Dump what we have to the output stream
+										await outputStream.WriteAsync(zipOutputStream.GetBuffer(), 0, (int)zipOutputStream.Position);
+
+										// Reset everything.
+										zipOutputStream.Position = 0;
+										zipOutputStream.SetLength(0);
+										totalBytesRead += bytesRead;
+									}
+								}
+							}
+						}
+					}
+					// Write out the EOF stuff
+					zipOutputStream.Position = 0;
+					await zipOutputStream.CopyToAsync(outputStream);
+				}
+			});
+		}
+	}
+}
diff --git a/Engine/Source/Programs/Horde/Horde.Build/Jobs/Artifacts/IArtifact.cs b/Engine/Source/Programs/Horde/Horde.Build/Jobs/Artifacts/IArtifact.cs
new file mode 100644
--- /dev/null
+++ b/Engine/Source/Programs/Horde/Horde.Build/Jobs/Artifacts/IArtifact.cs
@@ -0,0 +1,45 @@
+// Copyright Epic Games, Inc. All Rights Reserved.
+
+using Horde.Build.Utilities;
+using MongoDB.Bson;
+
+namespace Horde.Build.Models
+{
+	using JobId = ObjectId<IJob>;
+
+	/// <summary>
+	/// Information about an artifact
+	/// </summary>
+	public interface IArtifact
+	{
+		/// <summary>
+		/// Identifier for the Artifact. Randomly generated.
+		/// </summary>
+		public ObjectId Id { get; }
+
+		/// <summary>
+		/// Unique id of the job containing this artifact
+		/// </summary>
+		public JobId JobId { get; }
+
+		/// <summary>
+		/// Unique id of the job containing this artifact
+		/// </summary>
+		public string Name { get; }
+
+		/// <summary>
+		/// Unique id of the step containing this artifact
+		/// </summary>
+		public SubResourceId? StepId { get; }
+
+		/// <summary>
+		/// Total size of the file
+		/// </summary>
+		public long Length { get; }
+
+		/// <summary>
+		/// Type of artifact
+		/// </summary>
+		public string MimeType { get; }
+	}
+}
diff --git a/Engine/Source/Programs/Horde/Horde.Build/Jobs/Artifacts/IArtifactCollection.cs b/Engine/Source/Programs/Horde/Horde.Build/Jobs/Artifacts/IArtifactCollection.cs
new file mode 100644
--- /dev/null
+++ b/Engine/Source/Programs/Horde/Horde.Build/Jobs/Artifacts/IArtifactCollection.cs
@@ -0,0 +1,68 @@
+// Copyright Epic Games, Inc. All Rights Reserved.
+
+using System.Collections.Generic;
+using System.Threading.Tasks;
+using Horde.Build.Models;
+using Horde.Build.Utilities;
+using MongoDB.Bson;
+
+namespace Horde.Build.Jobs
+{
+	using JobId = ObjectId<IJob>;
+
+	/// <summary>
+	/// Interface for a collection of artifacts
+	/// </summary>
+	public interface IArtifactCollection
+	{
+		/// <summary>
+		/// Creates a new artifact
+		/// </summary>
+		/// <param name="jobId">Unique id of the job that owns this artifact</param>
+		/// <param name="stepId">Optional Step id</param>
+		/// <param name="name">Name of artifact</param>
+		/// <param name="mimeType">Type of artifact</param>
+		/// <param name="data">The data to write</param>
+		/// <returns>The new log file document</returns>
+		Task<IArtifact> CreateArtifactAsync(JobId jobId, SubResourceId? stepId, string name, string mimeType, System.IO.Stream data);
+
+		/// <summary>
+		/// Gets all the available artifacts for a job
+		/// </summary>
+		/// <param name="jobId">Unique id of the job to query</param>
+		/// <param name="stepId">Unique id of the Step to query</param>
+		/// <param name="name">Name of the artifact</param>
+		/// <returns>List of artifact documents</returns>
+		Task<List<IArtifact>> GetArtifactsAsync(JobId? jobId, SubResourceId? stepId, string? name);
+
+		/// <summary>
+		/// Gets a specific list of artifacts based on id
+		/// </summary>
+		/// <param name="artifactIds">The list of artifact Ids</param>
+		/// <returns>List of artifact documents</returns>
+		Task<List<IArtifact>> GetArtifactsAsync(IEnumerable<ObjectId> artifactIds);
+
+		/// <summary>
+		/// Gets an artifact by ID
+		/// </summary>
+		/// <param name="artifactId">Unique id of the artifact</param>
+		/// <returns>The artifact document</returns>
+		Task<IArtifact?> GetArtifactAsync(ObjectId artifactId);
+
+		/// <summary>
+		/// Updates an artifact
+		/// </summary>
+		/// <param name="artifact">The artifact</param>
+		/// <param name="newMimeType">New mime type</param>
+		/// <param name="newData">New data</param>
+		/// <returns>Async task</returns>
+		Task<bool> UpdateArtifactAsync(IArtifact? artifact, string newMimeType, System.IO.Stream newData);
+
+		/// <summary>
+		/// gets artifact data
+		/// </summary>
+		/// <param name="artifact">The artifact</param>
+		/// <returns>The chunk data</returns>
+		Task<System.IO.Stream> OpenArtifactReadStreamAsync(IArtifact artifact);
+	}
+}
diff --git a/Engine/Source/Programs/Horde/Horde.Build/Jobs/ArtifactsController.Messages.cs b/Engine/Source/Programs/Horde/Horde.Build/Jobs/ArtifactsController.Messages.cs
new file mode 100644
--- /dev/null
+++ b/Engine/Source/Programs/Horde/Horde.Build/Jobs/ArtifactsController.Messages.cs
@@ -0,0 +1,118 @@
+// Copyright Epic Games, Inc. All Rights Reserved.
+
+using System.Collections.Generic;
+using Horde.Build.Models;
+
+namespace Horde.Build.Api
+{
+	/// <summary>
+	/// Response from creating an artifact
+	/// </summary>
+	public class CreateArtifactResponse
+	{
+		/// <summary>
+		/// Unique id for this artifact
+		/// </summary>
+		public string Id { get; set; }
+
+		/// <summary>
+		/// Constructor
+		/// </summary>
+		/// <param name="id">The artifact file id</param>
+		public CreateArtifactResponse(string id)
+		{
+			Id = id;
+		}
+	}
+
+	/// <summary>
+	/// Request to get a zip of artifacts
+	/// </summary>
+	public class GetArtifactZipRequest
+	{
+		/// <summary>
+		/// Job id to get all artifacts for
+		/// </summary>
+		public string? JobId { get; set; }
+
+		/// <summary>
+		/// Step id to filter by
+		/// </summary>
+		public string? StepId { get; set; }
+
+		/// <summary>
+		/// Further filter by a list of artifact ids
+		/// </summary>
+		public List<string>? ArtifactIds { get; set; }
+
+		/// <summary>
+		/// Constructor
+		/// </summary>
+		/// <param name="jobId">Job id to get all artifacts for</param>
+		/// <param name="stepId">step to filter by</param>
+		/// <param name="artifactIds">The artifact ids.  Returns all artifacts for a job </param>
+		public GetArtifactZipRequest(string? jobId, string? stepId, List<string>? artifactIds)
+		{
+			JobId = jobId;
+			StepId = stepId;
+			ArtifactIds = artifactIds;
+		}
+	}
+
+	/// <summary>
+	/// Response describing an artifact
+	/// </summary>
+	public class GetArtifactResponse
+	{
+		/// <summary>
+		/// Unique id of the artifact
+		/// </summary>
+		public string Id { get; set; }
+
+		/// <summary>
+		/// Name of the artifact
+		/// </summary>
+		public string Name { get; set; }
+
+		/// <summary>
+		/// Unique id of the job for this artifact
+		/// </summary>
+		public string JobId { get; set; }
+
+		/// <summary>
+		/// Optional unique id of the step for this artifact
+		/// </summary>
+		public string? StepId { get; set; }
+
+		/// <summary>
+		/// Type of artifact
+		/// </summary>
+		public string MimeType { get; set; }
+
+		/// <summary>
+		/// Length of artifact
+		/// </summary>
+		public long Length { get; set; }
+
+		/// <summary>
+		/// Short-lived code that can be used to download this artifact through direct download links in the browser
+		/// </summary>
+		public string? Code { get; set; }
+
+		/// <summary>
+		/// Constructor
+		/// </summary>
+		/// <param name="artifact">The artifact to construct from</param>
+		/// <param name="code">The direct download code</param>
+		public GetArtifactResponse(IArtifact artifact, string? code)
+		{
+			Id = artifact.Id.ToString();
+			Name = artifact.Name;
+			JobId = artifact.JobId.ToString();
+			StepId = artifact.StepId.ToString();
+			MimeType = artifact.MimeType;
+			Length = artifact.Length;
+			Code = code;
+		}
+	}
+}
diff --git a/Engine/Source/Programs/Horde/Horde.Build/Jobs/GraphCollection.cs b/Engine/Source/Programs/Horde/Horde.Build/Jobs/GraphCollection.cs
new file mode 100644
--- /dev/null
+++ b/Engine/Source/Programs/Horde/Horde.Build/Jobs/GraphCollection.cs
@@ -0,0 +1,476 @@
+// Copyright Epic Games, Inc. All Rights Reserved.
+
+using System;
+using System.Collections.Concurrent;
+using System.Collections.Generic;
+using System.Diagnostics;
+using System.Linq;
+using System.Threading;
+using System.Threading.Tasks;
+using EpicGames.Core;
+using Horde.Build.Models;
+using Horde.Build.Server;
+using Horde.Build.Utilities;
+using HordeCommon;
+using MongoDB.Bson;
+using MongoDB.Bson.Serialization.Attributes;
+using MongoDB.Bson.Serialization.Options;
+using MongoDB.Driver;
+
+namespace Horde.Build.Collections.Impl
+{
+	/// <summary>
+	/// Collection of graph documents
+	/// </summary>
+	public class GraphCollection : IGraphCollection
+	{
+		/// <summary>
+		/// Represents a node in the graph
+		/// </summary>
+		[DebuggerDisplay("{Name}")]
+		class Node : INode
+		{
+			[BsonRequired]
+			public string Name { get; set; }
+
+			[BsonIgnoreIfNull]
+			public NodeRef[] InputDependencies { get; set; }
+
+			[BsonIgnoreIfNull]
+			public NodeRef[] OrderDependencies { get; set; }
+
+			public Priority Priority { get; set; }
+			public bool AllowRetry { get; set; } = true;
+			public bool RunEarly { get; set; }
+			public bool Warnings { get; set; } = true;
+
+			[BsonIgnoreIfNull]
+			public Dictionary<string, string>? Credentials { get; set; }
+
+			[BsonIgnoreIfNull]
+			public Dictionary<string, string>? Properties { get; set; }
+
+			[BsonIgnoreIfNull]
+			[BsonDictionaryOptions(DictionaryRepresentation.ArrayOfDocuments)]
+			public NodeAnnotations? Annotations { get; set; }
+
+			IReadOnlyDictionary<string, string>? INode.Credentials => Credentials;
+			IReadOnlyDictionary<string, string>? INode.Properties => Properties;
+			IReadOnlyNodeAnnotations INode.Annotations => Annotations ?? NodeAnnotations.Empty;
+
+			[BsonConstructor]
+			private Node()
+			{
+				Name = null!;
+				InputDependencies = null!;
+				OrderDependencies = null!;
+			}
+
+			public Node(string name, NodeRef[] inputDependencies, NodeRef[] orderDependencies, Priority priority, bool allowRetry, bool runEarly, bool warnings, Dictionary<string, string>? credentials, Dictionary<string, string>? properties, IReadOnlyNodeAnnotations? annotations)
+			{
+				Name = name;
+				InputDependencies = inputDependencies;
+				OrderDependencies = orderDependencies;
+				Priority = priority;
+				AllowRetry = allowRetry;
+				RunEarly = runEarly;
+				Warnings = warnings;
+				Credentials = credentials;
+				Properties = properties;
+				if (annotations != null && annotations.Count > 0)
+				{
+					Annotations = new NodeAnnotations(annotations);
+				}
+			}
+		}
+
+		class NodeGroup : INodeGroup
+		{
+			public string AgentType { get; set; }
+			public List<Node> Nodes { get; set; }
+
+			IReadOnlyList<INode> INodeGroup.Nodes => Nodes;
+
+			/// <summary>
+			/// Private constructor for BSON serializer
+			/// </summary>
+			[BsonConstructor]
+			private NodeGroup()
+			{
+				AgentType = null!;
+				Nodes = null!;
+			}
+
+			/// <summary>
+			/// Constructor
+			/// </summary>
+			/// <param name="agentType">The type of agent to execute this group</param>
+			/// <param name="nodes">Nodes to execute</param>
+			public NodeGroup(string agentType, List<Node> nodes)
+			{
+				AgentType = agentType;
+				Nodes = nodes;
+			}
+		}
+
+		class Aggregate : IAggregate
+		{
+			public string Name { get; set; }
+			public List<NodeRef> Nodes { get; set; }
+
+			IReadOnlyList<NodeRef> IAggregate.Nodes => Nodes;
+
+			private Aggregate()
+			{
+				Name = null!;
+				Nodes = new List<NodeRef>();
+			}
+
+			public Aggregate(string name, List<NodeRef> nodes)
+			{
+				Name = name;
+				Nodes = nodes;
+			}
+		}
+
+		class Label : ILabel
+		{
+			[BsonIgnoreIfNull]
+			public string? Name { get; set; }
+
+			[BsonIgnoreIfNull]
+			public string? Category { get; set; }
+
+			[BsonIgnoreIfNull]
+			public string? DashboardName { get; set; }
+
+			[BsonIgnoreIfNull]
+			public string? DashboardCategory { get; set; }
+
+			[BsonIgnoreIfNull]
+			public string? UgsName { get; set; }
+
+			[BsonIgnoreIfNull]
+			public string? UgsProject { get; set; }
+
+			[BsonIgnoreIfNull]
+			public LabelChange Change { get; set; }
+
+			public List<NodeRef> RequiredNodes { get; set; }
+			public List<NodeRef> IncludedNodes { get; set; }
+
+			string? ILabel.DashboardName => DashboardName ?? Name;
+			string? ILabel.DashboardCategory => DashboardCategory ?? Category;
+
+			private Label()
+			{
+				RequiredNodes = new List<NodeRef>();
+				IncludedNodes = new List<NodeRef>();
+			}
+
+			public Label(string? dashboardName, string? dashboardCategory, string? ugsName, string? ugsProject, LabelChange change, List<NodeRef> requiredNodes, List<NodeRef> includedNodes)
+			{
+				DashboardName = dashboardName;
+				DashboardCategory = dashboardCategory;
+				UgsName = ugsName;
+				UgsProject = ugsProject;
+				Change = change;
+				RequiredNodes = requiredNodes;
+				IncludedNodes = includedNodes;
+			}
+		}
+
+		class GraphDocument : IGraph
+		{
+			public static GraphDocument Empty { get; } = new GraphDocument(new List<NodeGroup>(), new List<Aggregate>(), new List<Label>());
+
+			[BsonRequired, BsonId]
+			public ContentHash Id { get; private set; } = ContentHash.Empty;
+
+			public int Schema { get; set; }
+			public List<NodeGroup> Groups { get; private set; } = new List<NodeGroup>();
+			public List<Aggregate> Aggregates { get; private set; } = new List<Aggregate>();
+			public List<Label> Labels { get; private set; } = new List<Label>();
+
+			[BsonIgnore]
+			IReadOnlyDictionary<string, NodeRef>? _cachedNodeNameToRef;
+
+			IReadOnlyList<INodeGroup> IGraph.Groups => Groups;
+			IReadOnlyList<IAggregate> IGraph.Aggregates => Aggregates;
+			IReadOnlyList<ILabel> IGraph.Labels => Labels;
+
+			[BsonConstructor]
+			private GraphDocument()
+			{
+			}
+
+			public GraphDocument(List<NodeGroup> groups, List<Aggregate> aggregates, List<Label> labels)
+			{
+				Groups = groups;
+				Aggregates = aggregates;
+				Labels = labels;
+				Id = ContentHash.SHA1(BsonExtensionMethods.ToBson(this));
+			}
+
+			public GraphDocument(GraphDocument baseGraph, List<NewGroup>? newGroupRequests, List<NewAggregate>? newAggregateRequests, List<NewLabel>? newLabelRequests)
+			{
+				Dictionary<string, NodeRef> nodeNameToRef = new Dictionary<string, NodeRef>(baseGraph.GetNodeNameToRef(), StringComparer.OrdinalIgnoreCase);
+
+				// Update the new list of groups
+				List<NodeGroup> newGroups = new List<NodeGroup>(baseGraph.Groups);
+				if (newGroupRequests != null)
+				{
+					foreach (NewGroup newGroupRequest in newGroupRequests)
+					{
+						List<Node> nodes = new List<Node>();
+						foreach (NewNode newNodeRequest in newGroupRequest.Nodes)
+						{
+							int nodeIdx = nodes.Count;
+
+							Priority priority = newNodeRequest.Priority ?? Priority.Normal;
+							bool bAllowRetry = newNodeRequest.AllowRetry ?? true;
+							bool bRunEarly = newNodeRequest.RunEarly ?? false;
+							bool bWarnings = newNodeRequest.Warnings ?? true;
+
+							NodeRef[] inputDependencies = (newNodeRequest.InputDependencies == null) ? Array.Empty<NodeRef>() : newNodeRequest.InputDependencies.Select(x => nodeNameToRef[x]).ToArray();
+							NodeRef[] orderDependencies = (newNodeRequest.OrderDependencies == null) ? Array.Empty<NodeRef>() : newNodeRequest.OrderDependencies.Select(x => nodeNameToRef[x]).ToArray();
+							orderDependencies = orderDependencies.Union(inputDependencies).ToArray();
+							nodes.Add(new Node(newNodeRequest.Name, inputDependencies, orderDependencies, priority, bAllowRetry, bRunEarly, bWarnings, newNodeRequest.Credentials, newNodeRequest.Properties, newNodeRequest.Annotations));
+
+							nodeNameToRef.Add(newNodeRequest.Name, new NodeRef(newGroups.Count, nodeIdx));
+						}
+						newGroups.Add(new NodeGroup(newGroupRequest.AgentType, nodes));
+					}
+				}
+
+				// Update the list of aggregates
+				List<Aggregate> newAggregates = new List<Aggregate>(baseGraph.Aggregates);
+				if (newAggregateRequests != null)
+				{
+					foreach (NewAggregate newAggregateRequest in newAggregateRequests)
+					{
+						List<NodeRef> nodes = newAggregateRequest.Nodes.ConvertAll(x => nodeNameToRef[x]);
+						newAggregates.Add(new Aggregate(newAggregateRequest.Name, nodes));
+					}
+				}
+
+				// Update the list of labels
+				List<Label> newLabels = new List<Label>(baseGraph.Labels);
+				if (newLabelRequests != null)
+				{
+					foreach (NewLabel newLabelRequest in newLabelRequests)
+					{
+						List<NodeRef> requiredNodes = newLabelRequest.RequiredNodes.ConvertAll(x => nodeNameToRef[x]);
+						List<NodeRef> includedNodes = newLabelRequest.IncludedNodes.ConvertAll(x => nodeNameToRef[x]);
+						newLabels.Add(new Label(newLabelRequest.DashboardName, newLabelRequest.DashboardCategory, newLabelRequest.UgsName, newLabelRequest.UgsProject, newLabelRequest.Change, requiredNodes, includedNodes));
+					}
+				}
+
+				// Create the new arrays
+				Groups = newGroups;
+				Aggregates = newAggregates;
+				Labels = newLabels;
+
+				// Create the new graph, and save the generated node lookup into it
+				_cachedNodeNameToRef = nodeNameToRef;
+
+				// Compute the hash
+				Id = ContentHash.SHA1(BsonExtensionMethods.ToBson(this));
+			}
+
+			public IReadOnlyDictionary<string, NodeRef> GetNodeNameToRef()
+			{
+				if (_cachedNodeNameToRef == null)
+				{
+					Dictionary<string, NodeRef> nodeNameToRef = new Dictionary<string, NodeRef>(StringComparer.OrdinalIgnoreCase);
+					for (int groupIdx = 0; groupIdx < Groups.Count; groupIdx++)
+					{
+						List<Node> nodes = Groups[groupIdx].Nodes;
+						for (int nodeIdx = 0; nodeIdx < nodes.Count; nodeIdx++)
+						{
+							Node node = nodes[nodeIdx];
+							nodeNameToRef[node.Name] = new NodeRef(groupIdx, nodeIdx);
+						}
+					}
+					_cachedNodeNameToRef = nodeNameToRef;
+				}
+				return _cachedNodeNameToRef;
+			}
+		}
+
+		/// <summary>
+		/// Stores information about a cached graph
+		/// </summary>
+		class CachedGraph
+		{
+			/// <summary>
+			/// Time at which the graph was last accessed
+			/// </summary>
+			public long LastAccessTime => _lastAccessTimePrivate;
+
+			/// <summary>
+			/// Backing value for <see cref="LastAccessTime"/>
+			/// </summary>
+			private long _lastAccessTimePrivate;
+
+			/// <summary>
+			/// The graph instance
+			/// </summary>
+			public IGraph Graph { get; }
+
+			/// <summary>
+			/// Constructor
+			/// </summary>
+			/// <param name="graph">The graph to store</param>
+			public CachedGraph(IGraph graph)
+			{
+				_lastAccessTimePrivate = Stopwatch.GetTimestamp();
+				Graph = graph;
+			}
+
+			/// <summary>
+			/// Update the last access time
+			/// </summary>
+			public void Touch()
+			{
+				for (; ; )
+				{
+					long time = Stopwatch.GetTimestamp();
+					long lastAccessTimeCopy = _lastAccessTimePrivate;
+					if (time < lastAccessTimeCopy || Interlocked.CompareExchange(ref _lastAccessTimePrivate, time, lastAccessTimeCopy) == lastAccessTimeCopy)
+					{
+						break;
+					}
+				}
+			}
+		}
+
+		/// <summary>
+		/// The jobs collection
+		/// </summary>
+		readonly IMongoCollection<GraphDocument> _graphs;
+
+		/// <summary>
+		/// Maximum number of graphs to keep in the cache
+		/// </summary>
+		const int MaxGraphs = 1000;
+
+		/// <summary>
+		/// Cache of recently accessed graphs
+		/// </summary>
+		readonly ConcurrentDictionary<ContentHash, CachedGraph> _cachedGraphs = new ConcurrentDictionary<ContentHash, CachedGraph>();
+
+		/// <summary>
+		/// Constructor
+		/// </summary>
+		/// <param name="mongoService">The database service singleton</param>
+		public GraphCollection(MongoService mongoService)
+		{
+			_graphs = mongoService.GetCollection<GraphDocument>("Graphs");
+		}
+
+		/// <summary>
+		/// Adds a new graph document
+		/// </summary>
+		/// <param name="graph">The graph to add</param>
+		/// <returns>Async task</returns>
+		async Task AddAsync(GraphDocument graph)
+		{
+			if (!await _graphs.Find(x => x.Id == graph.Id).AnyAsync())
+			{
+				try
+				{
+					await _graphs.InsertOneAsync(graph);
+				}
+				catch (MongoWriteException ex)
+				{
+					if (ex.WriteError.Category != ServerErrorCategory.DuplicateKey)
+					{
+						throw;
+					}
+				}
+			}
+		}
+
+		/// <inheritdoc/>
+		public async Task<IGraph> AddAsync(ITemplate template)
+		{
+			Node node = new Node(IJob.SetupNodeName, Array.Empty<NodeRef>(), Array.Empty<NodeRef>(), Priority.High, true, false, true, null, null, null);
+			NodeGroup group = new NodeGroup(template.InitialAgentType ?? "Win64", new List<Node> { node });
+
+			GraphDocument graph = new GraphDocument(new List<NodeGroup> { group }, new List<Aggregate>(), new List<Label>());
+			await AddAsync(graph);
+			return graph;
+		}
+
+		/// <inheritdoc/>
+		public async Task<IGraph> AppendAsync(IGraph? baseGraph, List<NewGroup>? newGroupRequests, List<NewAggregate>? newAggregateRequests, List<NewLabel>? newLabelRequests)
+		{
+			GraphDocument graph = new GraphDocument((GraphDocument?)baseGraph ?? GraphDocument.Empty, newGroupRequests, newAggregateRequests, newLabelRequests);
+			await AddAsync(graph);
+			return graph;
+		}
+
+		/// <inheritdoc/>
+		public async Task<IGraph> GetAsync(ContentHash? hash)
+		{
+			// Special case for an empty graph request
+			if (hash == null || hash == ContentHash.Empty || hash == GraphDocument.Empty.Id)
+			{
+				return GraphDocument.Empty;
+			}
+
+			// Try to read the graph from the cache
+			CachedGraph? cachedGraph;
+			if (_cachedGraphs.TryGetValue(hash, out cachedGraph))
+			{
+				// Update the last access time
+				cachedGraph.Touch();
+			}
+			else
+			{
+				// Trim the cache
+				while (_cachedGraphs.Count > MaxGraphs)
+				{
+					ContentHash? removeHash = _cachedGraphs.OrderBy(x => x.Value.LastAccessTime).Select(x => x.Key).FirstOrDefault();
+					if (removeHash == null || removeHash == ContentHash.Empty)
+					{
+						break;
+					}
+					_cachedGraphs.TryRemove(removeHash, out _);
+				}
+
+				// Create the new entry
+				cachedGraph = new CachedGraph(await _graphs.Find<GraphDocument>(x => x.Id == hash).FirstAsync());
+				_cachedGraphs.TryAdd(hash, cachedGraph);
+			}
+			return cachedGraph.Graph;
+		}
+
+		/// <inheritdoc/>
+		public async Task<List<IGraph>> FindAllAsync(ContentHash[]? hashes, int? index, int? count)
+		{
+			FilterDefinitionBuilder<GraphDocument> filterBuilder = Builders<GraphDocument>.Filter;
+
+			FilterDefinition<GraphDocument> filter = filterBuilder.Empty;
+			if (hashes != null)
+			{
+				filter &= filterBuilder.In(x => x.Id, hashes);
+			}
+
+			List<GraphDocument> results;
+			IFindFluent<GraphDocument, GraphDocument> search = _graphs.Find(filter);
+			if (index != null)
+			{
+				search = search.Skip(index.Value);
+			}
+			if (count != null)
+			{
+				search = search.Limit(count.Value);
+			}
+
+			results = await search.ToListAsync();
+			return results.ConvertAll<IGraph>(x => x);
+		}
+	}
+}
diff --git a/Engine/Source/Programs/Horde/Horde.Build/Jobs/Graphs/GraphMessages.cs b/Engine/Source/Programs/Horde/Horde.Build/Jobs/Graphs/GraphMessages.cs
new file mode 100644
--- /dev/null
+++ b/Engine/Source/Programs/Horde/Horde.Build/Jobs/Graphs/GraphMessages.cs
@@ -0,0 +1,252 @@
+// Copyright Epic Games, Inc. All Rights Reserved.
+
+using System.Collections.Generic;
+using System.Linq;
+using EpicGames.Core;
+using Horde.Build.Models;
+using HordeCommon;
+
+namespace Horde.Build.Api
+{
+	/// <summary>
+	/// Information required to create a node
+	/// </summary>
+	public class GetNodeResponse
+	{
+		/// <summary>
+		/// The name of this node 
+		/// </summary>
+		public string Name { get; set; }
+
+		/// <summary>
+		/// Indices of nodes which must have succeeded for this node to run
+		/// </summary>
+		public List<string> InputDependencies { get; set; }
+
+		/// <summary>
+		/// Indices of nodes which must have completed for this node to run
+		/// </summary>
+		public List<string> OrderDependencies { get; set; }
+
+		/// <summary>
+		/// The priority of this node
+		/// </summary>
+		public Priority Priority { get; set; }
+
+		/// <summary>
+		/// Whether this node can be retried
+		/// </summary>
+		public bool AllowRetry { get; set; }
+
+		/// <summary>
+		/// This node can start running early, before dependencies of other nodes in the same group are complete
+		/// </summary>
+		public bool RunEarly { get; set; }
+
+		/// <summary>
+		/// Whether to include warnings in diagnostic output
+		/// </summary>
+		public bool Warnings { get; set; }
+
+		/// <summary>
+		/// Average time to execute this node based on historical trends
+		/// </summary>
+		public float? AverageDuration { get; set; }
+
+		/// <summary>
+		/// Credentials required for this node to run. This dictionary maps from environment variable names to a credential property in the format 'CredentialName.PropertyName'.
+		/// </summary>
+		public IReadOnlyDictionary<string, string>? Credentials { get; set; }
+
+		/// <summary>
+		/// Properties for this node
+		/// </summary>
+		public IReadOnlyDictionary<string, string>? Properties { get; set; }
+
+		/// <summary>
+		/// Annotations for this node
+		/// </summary>
+		public IReadOnlyDictionary<string, string>? Annotations { get; set; }
+
+		/// <summary>
+		/// Constructor
+		/// </summary>
+		/// <param name="node">The node to construct from</param>
+		/// <param name="groups">The groups in the graph</param>
+		public GetNodeResponse(INode node, IReadOnlyList<INodeGroup> groups)
+		{
+			Name = node.Name;
+			InputDependencies = new List<string>(node.InputDependencies.Select(x => groups[x.GroupIdx].Nodes[x.NodeIdx].Name));
+			OrderDependencies = new List<string>(node.OrderDependencies.Select(x => groups[x.GroupIdx].Nodes[x.NodeIdx].Name));
+			Priority = node.Priority;
+			AllowRetry = node.AllowRetry;
+			RunEarly = node.RunEarly;
+			Warnings = node.Warnings;
+			Credentials = node.Credentials;
+			Properties = node.Properties;
+			Annotations = node.Annotations;
+		}
+	}
+
+	/// <summary>
+	/// Information about a group of nodes
+	/// </summary>
+	public class GetGroupResponse
+	{
+		/// <summary>
+		/// The type of agent to execute this group
+		/// </summary>
+		public string AgentType { get; set; }
+
+		/// <summary>
+		/// Nodes in the group
+		/// </summary>
+		public List<GetNodeResponse> Nodes { get; set; }
+
+		/// <summary>
+		/// Constructor
+		/// </summary>
+		/// <param name="group">The group to construct from</param>
+		/// <param name="groups">Other groups in this graph</param>
+		public GetGroupResponse(INodeGroup group, IReadOnlyList<INodeGroup> groups)
+		{
+			AgentType = group.AgentType;
+			Nodes = group.Nodes.ConvertAll(x => new GetNodeResponse(x, groups));
+		}
+	}
+
+	/// <summary>
+	/// Information about an aggregate
+	/// </summary>
+	public class GetAggregateResponse
+	{
+		/// <summary>
+		/// Name of the aggregate
+		/// </summary>
+		public string Name { get; set; }
+
+		/// <summary>
+		/// Nodes which must be part of the job for the aggregate to be shown
+		/// </summary>
+		public List<string> Nodes { get; set; }
+
+		/// <summary>
+		/// Constructor
+		/// </summary>
+		/// <param name="aggregate">The aggregate to construct from</param>
+		/// <param name="groups">List of groups in this graph</param>
+		public GetAggregateResponse(IAggregate aggregate, IReadOnlyList<INodeGroup> groups)
+		{
+			Name = aggregate.Name;
+			Nodes = aggregate.Nodes.ConvertAll(x => x.ToNode(groups).Name);
+		}
+	}
+
+	/// <summary>
+	/// Information about a label
+	/// </summary>
+	public class GetLabelResponse
+	{
+		/// <summary>
+		/// Category of the aggregate
+		/// </summary>
+		public string? Category => DashboardCategory;
+
+		/// <summary>
+		/// Label for this aggregate
+		/// </summary>
+		public string? Name => DashboardName;
+
+		/// <summary>
+		/// Name to show for this label on the dashboard
+		/// </summary>
+		public string? DashboardName { get; set; }
+
+		/// <summary>
+		/// Category to show this label in on the dashboard
+		/// </summary>
+		public string? DashboardCategory { get; set; }
+
+		/// <summary>
+		/// Name to show for this label in UGS
+		/// </summary>
+		public string? UgsName { get; set; }
+
+		/// <summary>
+		/// Project to display this label for in UGS
+		/// </summary>
+		public string? UgsProject { get; set; }
+
+		/// <summary>
+		/// Nodes which must be part of the job for the aggregate to be shown
+		/// </summary>
+		public List<string> RequiredNodes { get; set; }
+
+		/// <summary>
+		/// Nodes to include in the status of this aggregate, if present in the job
+		/// </summary>
+		public List<string> IncludedNodes { get; set; }
+
+		/// <summary>
+		/// Constructor
+		/// </summary>
+		/// <param name="label">The label to construct from</param>
+		/// <param name="groups">List of groups in this graph</param>
+		public GetLabelResponse(ILabel label, IReadOnlyList<INodeGroup> groups)
+		{
+			DashboardName = label.DashboardName;
+			DashboardCategory = label.DashboardCategory;
+			UgsName = label.UgsName;
+			UgsProject = label.UgsProject;
+			RequiredNodes = label.RequiredNodes.ConvertAll(x => x.ToNode(groups).Name);
+			IncludedNodes = label.IncludedNodes.ConvertAll(x => x.ToNode(groups).Name);
+		}
+	}
+
+	/// <summary>
+	/// Information about a graph
+	/// </summary>
+	public class GetGraphResponse
+	{
+		/// <summary>
+		/// The hash of the graph
+		/// </summary>
+		public string Hash { get; set; }
+
+		/// <summary>
+		/// Array of nodes for this job
+		/// </summary>
+		public List<GetGroupResponse>? Groups { get; set; }
+
+		/// <summary>
+		/// List of aggregates
+		/// </summary>
+		public List<GetAggregateResponse>? Aggregates { get; set; }
+
+		/// <summary>
+		/// List of labels for the graph
+		/// </summary>
+		public List<GetLabelResponse>? Labels { get; set; }
+
+		/// <summary>
+		/// Constructor
+		/// </summary>
+		/// <param name="graph">The graph to construct from</param>
+		public GetGraphResponse(IGraph graph)
+		{
+			Hash = graph.Id.ToString();
+			if (graph.Groups.Count > 0)
+			{
+				Groups = graph.Groups.ConvertAll(x => new GetGroupResponse(x, graph.Groups));
+			}
+			if (graph.Aggregates.Count > 0)
+			{
+				Aggregates = graph.Aggregates.ConvertAll(x => new GetAggregateResponse(x, graph.Groups));
+			}
+			if (graph.Labels.Count > 0)
+			{
+				Labels = graph.Labels.ConvertAll(x => new GetLabelResponse(x, graph.Groups));
+			}
+		}
+	}
+}
