Horde: Finish re-org of Horde solution started a while ago.

diff --git a/Engine/Source/Programs/Horde/Horde.Build/Jobs/JobCollection.cs b/Engine/Source/Programs/Horde/Horde.Build/Jobs/JobCollection.cs
new file mode 100644
--- /dev/null
+++ b/Engine/Source/Programs/Horde/Horde.Build/Jobs/JobCollection.cs
@@ -0,0 +1,1740 @@
+// Copyright Epic Games, Inc. All Rights Reserved.
+
+using System;
+using System.Collections.Generic;
+using System.Diagnostics.CodeAnalysis;
+using System.IO;
+using System.Linq;
+using System.Threading.Tasks;
+using EpicGames.Core;
+using Horde.Build.Acls;
+using Horde.Build.Api;
+using Horde.Build.Models;
+using Horde.Build.Server;
+using Horde.Build.Services;
+using Horde.Build.Utilities;
+using HordeCommon;
+using Microsoft.Extensions.Logging;
+using MongoDB.Bson;
+using MongoDB.Bson.IO;
+using MongoDB.Bson.Serialization;
+using MongoDB.Bson.Serialization.Attributes;
+using MongoDB.Driver;
+using OpenTracing;
+using OpenTracing.Util;
+
+namespace Horde.Build.Collections.Impl
+{
+	using JobId = ObjectId<IJob>;
+	using LeaseId = ObjectId<ILease>;
+	using LogId = ObjectId<ILogFile>;
+	using PoolId = StringId<IPool>;
+	using SessionId = ObjectId<ISession>;
+	using StreamId = StringId<IStream>;
+	using TemplateRefId = StringId<TemplateRef>;
+	using UserId = ObjectId<IUser>;
+
+	/// <summary>
+	/// Wrapper around the jobs collection in a mongo DB
+	/// </summary>
+	public class JobCollection : IJobCollection
+	{
+		/// <summary>
+		/// Embedded jobstep document
+		/// </summary>
+		[BsonIgnoreExtraElements]
+		class JobStepDocument : IJobStep
+		{
+			[BsonRequired]
+			public SubResourceId Id { get; set; }
+
+			[BsonRequired]
+			public int NodeIdx { get; set; }
+
+			[BsonRequired]
+			public JobStepState State { get; set; } = JobStepState.Waiting;
+
+			public JobStepOutcome Outcome { get; set; } = JobStepOutcome.Success;
+
+			public JobStepError Error { get; set; } = JobStepError.None;
+
+			[BsonIgnoreIfNull]
+			public LogId? LogId { get; set; }
+
+			[BsonIgnoreIfNull]
+			public ObjectId? NotificationTriggerId { get; set; }
+
+			[BsonIgnoreIfNull]
+			public DateTimeOffset? ReadyTime { get; set; }
+
+			[BsonIgnoreIfNull]
+			public DateTimeOffset? StartTime { get; set; }
+
+			[BsonIgnoreIfNull]
+			public DateTimeOffset? FinishTime { get; set; }
+
+			[BsonIgnoreIfNull]
+			public Priority? Priority { get; set; }
+
+			[BsonIgnoreIfDefault, BsonDefaultValue(false)]
+			public bool Retry { get; set; }
+
+			public UserId? RetriedByUserId { get; set; }
+
+			[BsonElement("RetryByUser")]
+			public string? RetriedByUserDeprecated { get; set; }
+
+			public bool AbortRequested { get; set; } = false;
+
+			public UserId? AbortedByUserId { get; set; }
+			
+			[BsonElement("AbortByUser")]
+			public string? AbortedByUserDeprecated { get; set; }
+
+			[BsonIgnoreIfNull]
+			public List<Report>? Reports { get; set; }
+			IReadOnlyList<IReport>? IJobStep.Reports => Reports;
+
+			[BsonIgnoreIfNull]
+			public Dictionary<string, string>? Properties { get; set; }
+
+			DateTime? IJobStep.ReadyTimeUtc => ReadyTime?.UtcDateTime;
+			DateTime? IJobStep.StartTimeUtc => StartTime?.UtcDateTime;
+			DateTime? IJobStep.FinishTimeUtc => FinishTime?.UtcDateTime;
+
+			[BsonConstructor]
+			private JobStepDocument()
+			{
+			}
+
+			public JobStepDocument(SubResourceId id, int nodeIdx)
+			{
+				Id = id;
+				NodeIdx = nodeIdx;
+			}
+		}
+
+		class JobStepBatchDocument : IJobStepBatch
+		{
+			[BsonRequired]
+			public SubResourceId Id { get; set; }
+
+			public LogId? LogId { get; set; }
+
+			[BsonRequired]
+			public int GroupIdx { get; set; }
+
+			[BsonRequired]
+			public JobStepBatchState State { get; set; }
+
+			[BsonIgnoreIfDefault, BsonDefaultValue(JobStepBatchError.None)]
+			public JobStepBatchError Error { get; set; }
+
+			public List<JobStepDocument> Steps { get; set; } = new List<JobStepDocument>();
+
+			[BsonIgnoreIfNull]
+			public PoolId? PoolId { get; set; }
+
+			[BsonIgnoreIfNull]
+			public AgentId? AgentId { get; set; }
+
+			[BsonIgnoreIfNull]
+			public SessionId? SessionId { get; set; }
+
+			[BsonIgnoreIfNull]
+			public LeaseId? LeaseId { get; set; }
+
+			public int SchedulePriority { get; set; }
+
+			[BsonIgnoreIfNull]
+			public DateTimeOffset? ReadyTime { get; set; }
+
+			[BsonIgnoreIfNull]
+			public DateTimeOffset? StartTime { get; set; }
+
+			[BsonIgnoreIfNull]
+			public DateTimeOffset? FinishTime { get; set; }
+
+			IReadOnlyList<IJobStep> IJobStepBatch.Steps => Steps;
+			DateTime? IJobStepBatch.ReadyTimeUtc => ReadyTime?.UtcDateTime;
+			DateTime? IJobStepBatch.StartTimeUtc => StartTime?.UtcDateTime;
+			DateTime? IJobStepBatch.FinishTimeUtc => FinishTime?.UtcDateTime;
+
+			[BsonConstructor]
+			private JobStepBatchDocument()
+			{
+			}
+
+			public JobStepBatchDocument(SubResourceId id, int groupIdx)
+			{
+				Id = id;
+				GroupIdx = groupIdx;
+			}
+		}
+
+		class ChainedJobDocument : IChainedJob
+		{
+			public string Target { get; set; }
+			public TemplateRefId TemplateRefId { get; set; }
+			public JobId? JobId { get; set; }
+
+			[BsonConstructor]
+			private ChainedJobDocument()
+			{
+				Target = String.Empty;
+			}
+
+			public ChainedJobDocument(ChainedJobTemplate trigger)
+			{
+				Target = trigger.Trigger;
+				TemplateRefId = trigger.TemplateRefId;
+			}
+		}
+
+		class LabelNotificationDocument
+		{
+			public int _labelIdx;
+			public ObjectId _triggerId;
+		}
+
+		class JobDocument : IJob
+		{
+			[BsonRequired, BsonId]
+			public JobId Id { get; set; }
+
+			public StreamId StreamId { get; set; }
+			public TemplateRefId TemplateId { get; set; }
+			public ContentHash? TemplateHash { get; set; }
+			public ContentHash GraphHash { get; set; }
+
+			[BsonIgnoreIfNull]
+			public UserId? StartedByUserId { get; set; }
+
+			[BsonIgnoreIfNull, BsonElement("StartedByUser")]
+			public string? StartedByUserDeprecated { get; set; }
+
+			[BsonIgnoreIfNull]
+			public UserId? AbortedByUserId { get; set; }
+
+			[BsonIgnoreIfNull, BsonElement("AbortedByUser")]
+			public string? AbortedByUserDeprecated { get; set; }
+
+			[BsonRequired]
+			public string Name { get; set; }
+
+			public int Change { get; set; }
+			public int CodeChange { get; set; }
+			public int PreflightChange { get; set; }
+			public int ClonedPreflightChange { get; set; }
+			public Priority Priority { get; set; }
+
+			[BsonIgnoreIfDefault]
+			public bool AutoSubmit { get; set; }
+
+			[BsonIgnoreIfNull]
+			public int? AutoSubmitChange { get; set; }
+
+			[BsonIgnoreIfNull]
+			public string? AutoSubmitMessage { get; set; }
+
+			public bool UpdateIssues { get; set; }
+
+			public bool PromoteIssuesByDefault { get; set; }
+
+			[BsonIgnoreIfNull]
+			public DateTimeOffset? CreateTime { get; set; }
+
+			[BsonIgnoreIfNull]
+			public DateTime? CreateTimeUtc { get; set; }
+
+			public int SchedulePriority { get; set; }
+			public List<JobStepBatchDocument> Batches { get; set; } = new List<JobStepBatchDocument>();
+			public List<Report>? Reports { get; set; }
+			public List<string> Arguments { get; set; } = new List<string>();
+			public List<int> ReferencedByIssues { get; set; } = new List<int>();
+			public ObjectId? NotificationTriggerId { get; set; }
+			public bool ShowUgsBadges { get; set; }
+			public bool ShowUgsAlerts { get; set; }
+			public string? NotificationChannel { get; set; }
+			public string? NotificationChannelFilter { get; set; }
+			public List<LabelNotificationDocument> _labelNotifications = new List<LabelNotificationDocument>();
+			public List<ChainedJobDocument> ChainedJobs { get; set; } = new List<ChainedJobDocument>();
+
+			[BsonIgnoreIfNull]
+			public List<NodeRef>? RetriedNodes { get; set; }
+			public SubResourceId NextSubResourceId { get; set; }
+
+			[BsonIgnoreIfNull]
+			public DateTimeOffset? UpdateTime { get; set; }
+
+			[BsonIgnoreIfNull]
+			public DateTime? UpdateTimeUtc { get; set; }
+
+			public Acl? Acl { get; set; }
+
+			[BsonRequired]
+			public int UpdateIndex { get; set; }
+
+			DateTime IJob.CreateTimeUtc => CreateTimeUtc ?? CreateTime?.UtcDateTime ?? DateTime.UnixEpoch;
+			DateTime IJob.UpdateTimeUtc => UpdateTimeUtc ?? UpdateTime?.UtcDateTime ?? DateTime.UnixEpoch;
+			IReadOnlyList<IJobStepBatch> IJob.Batches => (IReadOnlyList<JobStepBatchDocument>)Batches;
+			IReadOnlyList<IReport>? IJob.Reports => Reports;
+			IReadOnlyList<string> IJob.Arguments => Arguments;
+			IReadOnlyList<int> IJob.Issues => ReferencedByIssues;
+			IReadOnlyDictionary<int, ObjectId> IJob.LabelIdxToTriggerId => _labelNotifications.ToDictionary(x => x._labelIdx, x => x._triggerId);
+			IReadOnlyList<IChainedJob> IJob.ChainedJobs => ChainedJobs;
+
+			[BsonConstructor]
+			private JobDocument()
+			{
+				Name = null!;
+				GraphHash = null!;
+			}
+
+			public JobDocument(JobId id, StreamId streamId, TemplateRefId templateId, ContentHash templateHash, ContentHash graphHash, string name, int change, int codeChange, int preflightChange, int clonedPreflightChange, UserId? startedByUserId, Priority? priority, bool? autoSubmit, bool? updateIssues, bool? promoteIssuesByDefault, DateTime createTimeUtc, List<ChainedJobDocument> chainedJobs, bool showUgsBadges, bool showUgsAlerts, string? notificationChannel, string? notificationChannelFilter, List<string>? arguments)
+			{
+				Id = id;
+				StreamId = streamId;
+				TemplateId = templateId;
+				TemplateHash = templateHash;
+				GraphHash = graphHash;
+				Name = name;
+				Change = change;
+				CodeChange = codeChange;
+				PreflightChange = preflightChange;
+				ClonedPreflightChange = clonedPreflightChange;
+				StartedByUserId = startedByUserId;
+				Priority = priority ?? HordeCommon.Priority.Normal;
+				AutoSubmit = autoSubmit ?? false;
+				UpdateIssues = updateIssues ?? (startedByUserId == null && preflightChange == 0);
+				PromoteIssuesByDefault = promoteIssuesByDefault ?? false;
+				CreateTimeUtc = createTimeUtc;
+				ChainedJobs = chainedJobs;
+				ShowUgsBadges = showUgsBadges;
+				ShowUgsAlerts = showUgsAlerts;
+				NotificationChannel = notificationChannel;
+				NotificationChannelFilter = notificationChannelFilter;
+				Arguments = arguments ?? Arguments;
+				NextSubResourceId = SubResourceId.Random();
+				UpdateTimeUtc = createTimeUtc;
+			}
+		}
+
+		/// <summary>
+		/// Projection of a job definition to just include permissions info
+		/// </summary>
+		[SuppressMessage("Design", "CA1812: Class is never instantiated")]
+		class JobPermissions : IJobPermissions
+		{
+			public static ProjectionDefinition<JobDocument> Projection { get; } = Builders<JobDocument>.Projection.Include(x => x.Acl).Include(x => x.StreamId);
+
+			public Acl? Acl { get; set; }
+			public StreamId StreamId { get; set; }
+		}
+
+		/// <summary>
+		/// Maximum number of times a step can be retried (after the original run)
+		/// </summary>
+		const int MaxRetries = 2;
+
+		readonly IMongoCollection<JobDocument> _jobs;
+		readonly MongoIndex<JobDocument> _createTimeIndex;
+		readonly MongoIndex<JobDocument> _updateTimeIndex;
+		readonly IClock _clock;
+		readonly ILogger<JobCollection> _logger;
+
+		/// <summary>
+		/// Constructor
+		/// </summary>
+		/// <param name="mongoService">The database service singleton</param>
+		/// <param name="clock"></param>
+		/// <param name="logger">The logger instance</param>
+		public JobCollection(MongoService mongoService, IClock clock, ILogger<JobCollection> logger)
+		{
+			_clock = clock;
+			_logger = logger;
+
+			List<MongoIndex<JobDocument>> indexes = new List<MongoIndex<JobDocument>>();
+			indexes.Add(keys => keys.Ascending(x => x.StreamId));
+			indexes.Add(keys => keys.Ascending(x => x.Change));
+			indexes.Add(keys => keys.Ascending(x => x.PreflightChange));
+			indexes.Add(_createTimeIndex = MongoIndex.Create<JobDocument>(keys => keys.Descending(x => x.CreateTimeUtc)));
+			indexes.Add(_updateTimeIndex = MongoIndex.Create<JobDocument>(keys => keys.Descending(x => x.UpdateTimeUtc)));
+			indexes.Add(keys => keys.Ascending(x => x.Name));
+			indexes.Add(keys => keys.Ascending(x => x.StartedByUserId));
+			indexes.Add(keys => keys.Ascending(x => x.TemplateId));
+			indexes.Add(keys => keys.Descending(x => x.SchedulePriority));
+			_jobs = mongoService.GetCollection<JobDocument>("Jobs", indexes);
+		}
+
+		static Task PostLoadAsync(JobDocument job)
+		{
+			if (job.GraphHash == ContentHash.Empty)
+			{
+				job.Batches.Clear();
+			}
+			return Task.CompletedTask;
+		}
+
+		static JobDocument Clone(JobDocument job)
+		{
+			using (MemoryStream stream = new MemoryStream())
+			{
+				using (BsonBinaryWriter writer = new BsonBinaryWriter(stream))
+				{
+					BsonSerializer.Serialize(writer, job);
+				}
+				return BsonSerializer.Deserialize<JobDocument>(stream.ToArray());
+			}
+		}
+
+		/// <inheritdoc/>
+		public async Task<IJob> AddAsync(JobId jobId, StreamId streamId, TemplateRefId templateRefId, ContentHash templateHash, IGraph graph, string name, int change, int codeChange, int? preflightChange, int? clonedPreflightChange, UserId? startedByUserId, Priority? priority, bool? autoSubmit, bool? updateIssues, bool? promoteIssuesByDefault, List<ChainedJobTemplate>? chainedJobs, bool showUgsBadges, bool showUgsAlerts, string? notificationChannel, string? notificationChannelFilter, List<string>? arguments)
+		{
+			List<ChainedJobDocument> jobTriggers = new List<ChainedJobDocument>();
+			if (chainedJobs == null)
+			{
+				jobTriggers = new List<ChainedJobDocument>();
+			}
+			else
+			{
+				jobTriggers = chainedJobs.ConvertAll(x => new ChainedJobDocument(x));
+			}
+
+			JobDocument newJob = new JobDocument(jobId, streamId, templateRefId, templateHash, graph.Id, name, change, codeChange, preflightChange ?? 0, clonedPreflightChange ?? 0, startedByUserId, priority, autoSubmit, updateIssues, promoteIssuesByDefault, DateTime.UtcNow, jobTriggers, showUgsBadges, showUgsAlerts, notificationChannel, notificationChannelFilter, arguments);
+			CreateBatches(newJob, graph, _logger);
+
+			await _jobs.InsertOneAsync(newJob);
+
+			return newJob;
+		}
+
+		/// <inheritdoc/>
+		public async Task<IJob?> GetAsync(JobId jobId)
+		{
+			JobDocument? job = await _jobs.Find<JobDocument>(x => x.Id == jobId).FirstOrDefaultAsync();
+			if (job != null)
+			{
+				await PostLoadAsync(job);
+			}
+			return job;
+		}
+
+		/// <inheritdoc/>
+		public async Task<bool> RemoveAsync(IJob job)
+		{
+			DeleteResult result = await _jobs.DeleteOneAsync(x => x.Id == job.Id && x.UpdateIndex == job.UpdateIndex);
+			return result.DeletedCount > 0;
+		}
+
+		/// <inheritdoc/>
+		public async Task RemoveStreamAsync(StreamId streamId)
+		{
+			await _jobs.DeleteManyAsync(x => x.StreamId == streamId);
+		}
+
+		/// <inheritdoc/>
+		public async Task<IJobPermissions?> GetPermissionsAsync(JobId jobId)
+		{
+			return await _jobs.Find<JobDocument>(x => x.Id == jobId).Project<JobPermissions>(JobPermissions.Projection).FirstOrDefaultAsync();
+		}
+
+		/// <inheritdoc/>
+		public async Task<List<IJob>> FindAsync(JobId[]? jobIds, StreamId? streamId, string? name, TemplateRefId[]? templates, int? minChange, int? maxChange, int? preflightChange, bool? preflightOnly, UserId ? preflightStartedByUser, UserId? startedByUser, DateTimeOffset? minCreateTime, DateTimeOffset? maxCreateTime, DateTimeOffset? modifiedBefore, DateTimeOffset? modifiedAfter, int? index, int? count, bool consistentRead, string? indexHint, bool? excludeUserJobs)
+		{
+			FilterDefinitionBuilder<JobDocument> filterBuilder = Builders<JobDocument>.Filter;
+
+			FilterDefinition<JobDocument> filter = filterBuilder.Empty;
+			if (jobIds != null && jobIds.Length > 0)
+			{
+				filter &= filterBuilder.In(x => x.Id, jobIds);
+			}
+			if (streamId != null)
+			{
+				filter &= filterBuilder.Eq(x => x.StreamId, streamId.Value);
+			}
+			if (name != null)
+			{
+				if (name.StartsWith("$", StringComparison.InvariantCulture))
+				{
+					BsonRegularExpression regex = new BsonRegularExpression(name.Substring(1), "i");
+					filter &= filterBuilder.Regex(x => x.Name, regex);
+				}
+				else
+				{
+					filter &= filterBuilder.Eq(x => x.Name, name);
+				}				
+			}
+			if (templates != null)
+			{
+				filter &= filterBuilder.In(x => x.TemplateId, templates);
+			}
+			if (minChange != null)
+			{
+				filter &= filterBuilder.Gte(x => x.Change, minChange);
+			}
+			if (maxChange != null)
+			{
+				filter &= filterBuilder.Lte(x => x.Change, maxChange);
+			}
+			if (preflightChange != null)
+			{
+				filter &= filterBuilder.Eq(x => x.PreflightChange, preflightChange);
+			}
+			if (preflightOnly != null && preflightOnly.Value)
+			{
+				filter &= filterBuilder.Ne(x => x.PreflightChange, 0);
+			}
+			if (excludeUserJobs != null && excludeUserJobs.Value)
+			{
+				filter &= filterBuilder.Eq(x => x.StartedByUserId, null);
+			}
+			else
+			{
+				if (preflightStartedByUser != null)
+				{
+					filter &= filterBuilder.Or(filterBuilder.Eq(x => x.PreflightChange, 0), filterBuilder.Eq(x => x.StartedByUserId, preflightStartedByUser));
+				}
+				if (startedByUser != null)
+				{
+					filter &= filterBuilder.Eq(x => x.StartedByUserId, startedByUser);
+				}
+			}
+			if (minCreateTime != null)
+			{
+				filter &= filterBuilder.Gte(x => x.CreateTimeUtc!, minCreateTime.Value.UtcDateTime);
+			}
+			if (maxCreateTime != null)
+			{
+				filter &= filterBuilder.Lte(x => x.CreateTimeUtc!, maxCreateTime.Value.UtcDateTime);
+			}
+			if (modifiedBefore != null)
+			{
+				filter &= filterBuilder.Lte(x => x.UpdateTimeUtc!, modifiedBefore.Value.UtcDateTime);
+			}
+			if (modifiedAfter != null)
+			{
+				filter &= filterBuilder.Gte(x => x.UpdateTimeUtc!, modifiedAfter.Value.UtcDateTime);
+			}
+
+			List<JobDocument> results;
+			using (IScope scope = GlobalTracer.Instance.BuildSpan("Jobs.Find").StartActive())
+			{
+				IMongoCollection<JobDocument> collection = consistentRead ? _jobs : _jobs.WithReadPreference(ReadPreference.SecondaryPreferred);
+				results = await collection.FindWithHint(filter, indexHint, x => x.SortByDescending(x => x.CreateTimeUtc!).Range(index, count).ToListAsync());
+			}
+			foreach (JobDocument result in results)
+			{
+				await PostLoadAsync(result);
+			}
+			return results.ConvertAll<JobDocument, IJob>(x => x);
+		}
+
+		/// <inheritdoc/>
+		public async Task<List<IJob>> FindLatestByStreamWithTemplatesAsync(StreamId streamId, TemplateRefId[] templates, UserId? preflightStartedByUser, DateTimeOffset? maxCreateTime, DateTimeOffset? modifiedAfter, int? index, int? count, bool consistentRead)
+		{
+			string indexHint = _createTimeIndex.Name;
+			if (modifiedAfter != null)
+			{
+				indexHint = _updateTimeIndex.Name;
+			}
+			
+			// This find call uses an index hint. Modifying the parameter passed to FindAsync can affect execution time a lot as the query planner is forced to use the specified index.
+			// Casting to interface to benefit from default parameter values
+			return await (this as IJobCollection).FindAsync(
+				streamId: streamId, templates: templates, preflightStartedByUser: preflightStartedByUser, modifiedAfter: modifiedAfter, maxCreateTime: maxCreateTime,
+				index: index, count: count, indexHint: indexHint, consistentRead: consistentRead);
+		}
+
+		/// <inheritdoc/>
+		public async Task<IJob?> TryUpdateJobAsync(IJob inJob, IGraph graph, string? name, Priority? priority, bool? autoSubmit, int? autoSubmitChange, string? autoSubmitMessage, UserId? abortedByUserId, ObjectId? notificationTriggerId, List<Report>? reports, List<string>? arguments, KeyValuePair<int, ObjectId>? labelIdxToTriggerId, KeyValuePair<TemplateRefId, JobId>? jobTrigger)
+		{
+			// Create the update 
+			UpdateDefinitionBuilder<JobDocument> updateBuilder = Builders<JobDocument>.Update;
+			List<UpdateDefinition<JobDocument>> updates = new List<UpdateDefinition<JobDocument>>();
+
+			// Flag for whether to update batches
+			bool bUpdateBatches = false;
+
+			// Build the update list
+			JobDocument jobDocument = Clone((JobDocument)inJob);
+			if (name != null)
+			{
+				jobDocument.Name = name;
+				updates.Add(updateBuilder.Set(x => x.Name, jobDocument.Name));
+			}
+			if (priority != null)
+			{
+				jobDocument.Priority = priority.Value;
+				updates.Add(updateBuilder.Set(x => x.Priority, jobDocument.Priority));
+			}
+			if (autoSubmit != null)
+			{
+				jobDocument.AutoSubmit = autoSubmit.Value;
+				updates.Add(updateBuilder.Set(x => x.AutoSubmit, jobDocument.AutoSubmit));
+			}
+			if (autoSubmitChange != null)
+			{
+				jobDocument.AutoSubmitChange = autoSubmitChange.Value;
+				updates.Add(updateBuilder.Set(x => x.AutoSubmitChange, jobDocument.AutoSubmitChange));
+			}
+			if (autoSubmitMessage != null)
+			{
+				jobDocument.AutoSubmitMessage = (autoSubmitMessage.Length == 0)? null : autoSubmitMessage;
+				updates.Add(updateBuilder.SetOrUnsetNullRef(x => x.AutoSubmitMessage, jobDocument.AutoSubmitMessage));
+			}
+			if (abortedByUserId != null && jobDocument.AbortedByUserId == null)
+			{
+				jobDocument.AbortedByUserId = abortedByUserId;
+				updates.Add(updateBuilder.Set(x => x.AbortedByUserId, jobDocument.AbortedByUserId));
+				bUpdateBatches = true;
+			}
+			if (notificationTriggerId != null)
+			{
+				jobDocument.NotificationTriggerId = notificationTriggerId.Value;
+				updates.Add(updateBuilder.Set(x => x.NotificationTriggerId, notificationTriggerId));
+			}
+			if (labelIdxToTriggerId != null)
+			{
+				if (jobDocument._labelNotifications.Any(x => x._labelIdx == labelIdxToTriggerId.Value.Key))
+				{
+					throw new ArgumentException("Cannot update label trigger that already exists");
+				}
+				jobDocument._labelNotifications.Add(new LabelNotificationDocument { _labelIdx = labelIdxToTriggerId.Value.Key, _triggerId = labelIdxToTriggerId.Value.Value });
+				updates.Add(updateBuilder.Set(x => x._labelNotifications, jobDocument._labelNotifications));
+			}
+			if (jobTrigger != null)
+			{
+				for (int idx = 0; idx < jobDocument.ChainedJobs.Count; idx++)
+				{
+					ChainedJobDocument jobTriggerDocument = jobDocument.ChainedJobs[idx];
+					if (jobTriggerDocument.TemplateRefId == jobTrigger.Value.Key)
+					{
+						int localIdx = idx;
+						jobTriggerDocument.JobId = jobTrigger.Value.Value;
+						updates.Add(updateBuilder.Set(x => x.ChainedJobs[localIdx].JobId, jobTrigger.Value.Value));
+					}
+				}
+			}
+			if (reports != null)
+			{
+				jobDocument.Reports ??= new List<Report>();
+				jobDocument.Reports.RemoveAll(x => reports.Any(y => y.Name == x.Name));
+				jobDocument.Reports.AddRange(reports);
+				updates.Add(updateBuilder.Set(x => x.Reports, jobDocument.Reports));
+			}
+			if (arguments != null)
+			{
+				HashSet<string> modifiedArguments = new HashSet<string>(jobDocument.Arguments);
+				modifiedArguments.SymmetricExceptWith(arguments);
+
+				foreach (string modifiedArgument in modifiedArguments)
+				{
+					if (modifiedArgument.StartsWith(IJob.TargetArgumentPrefix, StringComparison.OrdinalIgnoreCase))
+					{
+						bUpdateBatches = true;
+					}
+				}
+
+				jobDocument.Arguments = arguments.ToList();
+				updates.Add(updateBuilder.Set(x => x.Arguments, jobDocument.Arguments));
+			}
+
+			// Update the batches
+			if (bUpdateBatches)
+			{
+				UpdateBatches(jobDocument, graph, updates, _logger);
+			}
+
+			// Update the new list of job steps
+			return await TryUpdateAsync(jobDocument, updateBuilder.Combine(updates));
+		}
+
+		/// <inheritdoc/>
+		public async Task<IJob?> TryUpdateBatchAsync(IJob job, IGraph graph, SubResourceId batchId, LogId? newLogId, JobStepBatchState? newState, JobStepBatchError? newError)
+		{
+			JobDocument jobDocument = Clone((JobDocument)job);
+
+			// Find the index of the appropriate batch
+			int batchIdx = jobDocument.Batches.FindIndex(x => x.Id == batchId);
+			if (batchIdx == -1)
+			{
+				return null;
+			}
+
+			// If we're marking the batch as complete and there are still steps to run (eg. because the agent crashed), we need to mark all the steps as complete first
+			JobStepBatchDocument batch = jobDocument.Batches[batchIdx];
+
+			// Create the update 
+			UpdateDefinitionBuilder<JobDocument> updateBuilder = Builders<JobDocument>.Update;
+			List<UpdateDefinition<JobDocument>> updates = new List<UpdateDefinition<JobDocument>>();
+
+			// Update the batch
+			if (newLogId != null)
+			{
+				batch.LogId = newLogId.Value;
+				updates.Add(updateBuilder.Set(x => x.Batches[batchIdx].LogId, batch.LogId));
+			}
+			if (newState != null)
+			{
+				batch.State = newState.Value;
+				updates.Add(updateBuilder.Set(x => x.Batches[batchIdx].State, batch.State));
+
+				if (batch.StartTime == null && newState >= JobStepBatchState.Starting)
+				{
+					batch.StartTime = _clock.UtcNow;
+					updates.Add(updateBuilder.Set(x => x.Batches[batchIdx].StartTime, batch.StartTime));
+				}
+				if (newState == JobStepBatchState.Complete)
+				{
+					batch.FinishTime = _clock.UtcNow;
+					updates.Add(updateBuilder.Set(x => x.Batches[batchIdx].FinishTime, batch.FinishTime));
+				}
+			}
+			if (newError != null)
+			{
+				batch.Error = newError.Value;
+				updates.Add(updateBuilder.Set(x => x.Batches[batchIdx].Error, batch.Error));
+			}
+
+			// If the batch is being marked as incomplete, see if we can reschedule any of the work
+			if (newError == JobStepBatchError.Incomplete)
+			{
+				// Get the new list of retried nodes
+				List<NodeRef> retriedNodes = jobDocument.RetriedNodes ?? new List<NodeRef>();
+
+				// Check if there are any steps that need to be run again
+				bool bUpdateState = false;
+				foreach (JobStepDocument step in batch.Steps)
+				{
+					if (step.State == JobStepState.Ready || step.State == JobStepState.Waiting)
+					{
+						if (CanRetryNode(jobDocument, batch.GroupIdx, step.NodeIdx))
+						{
+							retriedNodes.Add(new NodeRef(batch.GroupIdx, step.NodeIdx));
+						}
+						else
+						{
+							step.State = JobStepState.Skipped;
+						}
+						bUpdateState = true;
+					}
+				}
+
+				// Update the steps
+				if (bUpdateState)
+				{
+					updates.Clear();
+					UpdateBatches(jobDocument, graph, updates, _logger);
+
+					jobDocument.RetriedNodes = retriedNodes;
+					updates.Add(updateBuilder.Set(x => x.RetriedNodes, jobDocument.RetriedNodes));
+				}
+			}
+
+			// Update the new list of job steps
+			return await TryUpdateAsync(jobDocument, updates);
+		}
+
+		/// <inheritdoc/>
+		public Task<IJob?> TryUpdateStepAsync(IJob job, IGraph graph, SubResourceId batchId, SubResourceId stepId, JobStepState newState, JobStepOutcome newOutcome, JobStepError? newError, bool? newAbortRequested, UserId? newAbortByUserId, LogId? newLogId, ObjectId? newNotificationTriggerId, UserId? newRetryByUserId, Priority? newPriority, List<Report>? newReports, Dictionary<string, string?>? newProperties)
+		{
+			JobDocument jobDocument = Clone((JobDocument)job);
+
+			// Create the update 
+			UpdateDefinitionBuilder<JobDocument> updateBuilder = Builders<JobDocument>.Update;
+			List<UpdateDefinition<JobDocument>> updates = new List<UpdateDefinition<JobDocument>>();
+
+			// Update the appropriate batch
+			bool bRefreshBatches = false;
+			bool bRefreshDependentJobSteps = false;
+			for (int loopBatchIdx = 0; loopBatchIdx < jobDocument.Batches.Count; loopBatchIdx++)
+			{
+				int batchIdx = loopBatchIdx; // For lambda capture
+				JobStepBatchDocument batch = jobDocument.Batches[batchIdx];
+				if (batch.Id == batchId)
+				{
+					for (int loopStepIdx = 0; loopStepIdx < batch.Steps.Count; loopStepIdx++)
+					{
+						int stepIdx = loopStepIdx; // For lambda capture
+						JobStepDocument step = batch.Steps[stepIdx];
+						if (step.Id == stepId)
+						{
+
+							// Update the request abort status
+							if (newAbortRequested != null && step.AbortRequested == false)
+							{
+								step.AbortRequested = newAbortRequested.Value;
+								updates.Add(updateBuilder.Set(x => x.Batches[batchIdx].Steps[stepIdx].AbortRequested, step.AbortRequested));
+
+								// If the step is pending, and not running on an agent, set to aborted
+								if (step.IsPending() && step.State != JobStepState.Running)
+								{
+									newState = JobStepState.Aborted;
+									newOutcome = JobStepOutcome.Failure;
+								}
+
+								bRefreshDependentJobSteps = true;
+							}
+
+							// Update the user that requested the abort
+							if (newAbortByUserId != null && step.AbortedByUserId == null)
+							{
+								step.AbortedByUserId = newAbortByUserId;
+								updates.Add(updateBuilder.Set(x => x.Batches[batchIdx].Steps[stepIdx].AbortedByUserId, step.AbortedByUserId));
+
+								bRefreshDependentJobSteps = true;
+							}
+
+							// Update the state
+							if (newState != JobStepState.Unspecified && step.State != newState)
+							{
+								step.State = newState;
+								updates.Add(updateBuilder.Set(x => x.Batches[batchIdx].Steps[stepIdx].State, step.State));
+
+								if (step.State == JobStepState.Running)
+								{
+									step.StartTime = _clock.UtcNow;
+									updates.Add(updateBuilder.Set(x => x.Batches[batchIdx].Steps[stepIdx].StartTime, step.StartTime));
+								}
+								else if (step.State == JobStepState.Completed || step.State == JobStepState.Aborted)
+								{
+									step.FinishTime = _clock.UtcNow;
+									updates.Add(updateBuilder.Set(x => x.Batches[batchIdx].Steps[stepIdx].FinishTime, step.FinishTime));
+								}
+
+								bRefreshDependentJobSteps = true;
+							}
+
+							// Update the job outcome
+							if (newOutcome != JobStepOutcome.Unspecified && step.Outcome != newOutcome)
+							{
+								step.Outcome = newOutcome;
+								updates.Add(updateBuilder.Set(x => x.Batches[batchIdx].Steps[stepIdx].Outcome, step.Outcome));
+
+								bRefreshDependentJobSteps = true;
+							}
+
+							// Update the job step error
+							if (newError != null)
+							{
+								step.Error = newError.Value;
+								updates.Add(updateBuilder.Set(x => x.Batches[batchIdx].Steps[stepIdx].Error, step.Error));
+							}
+
+							// Update the log id
+							if (newLogId != null && step.LogId != newLogId.Value)
+							{
+								step.LogId = newLogId.Value;
+								updates.Add(updateBuilder.Set(x => x.Batches[batchIdx].Steps[stepIdx].LogId, step.LogId));
+
+								bRefreshDependentJobSteps = true;
+							}
+
+							// Update the notification trigger id
+							if (newNotificationTriggerId != null && step.NotificationTriggerId == null)
+							{
+								step.NotificationTriggerId = newNotificationTriggerId.Value;
+								updates.Add(updateBuilder.Set(x => x.Batches[batchIdx].Steps[stepIdx].NotificationTriggerId, step.NotificationTriggerId));
+							}
+
+							// Update the retry flag
+							if (newRetryByUserId != null && step.RetriedByUserId == null)
+							{
+								step.RetriedByUserId = newRetryByUserId;
+								updates.Add(updateBuilder.Set(x => x.Batches[batchIdx].Steps[stepIdx].RetriedByUserId, step.RetriedByUserId));
+
+								bRefreshBatches = true;
+							}
+
+							// Update the priority
+							if (newPriority != null && newPriority.Value != step.Priority)
+							{
+								step.Priority = newPriority.Value;
+								updates.Add(updateBuilder.Set(x => x.Batches[batchIdx].Steps[stepIdx].Priority, step.Priority));
+
+								bRefreshBatches = true;
+							}
+
+							// Add any new reports
+							if (newReports != null)
+							{
+								step.Reports ??= new List<Report>();
+								step.Reports.RemoveAll(x => newReports.Any(y => y.Name == x.Name));
+								step.Reports.AddRange(newReports);
+								updates.Add(updateBuilder.Set(x => x.Batches[batchIdx].Steps[stepIdx].Reports, step.Reports));
+							}
+
+							// Apply any property updates
+							if (newProperties != null)
+							{
+								if (step.Properties == null)
+								{
+									step.Properties = new Dictionary<string, string>(StringComparer.Ordinal);
+								}
+
+								foreach (KeyValuePair<string, string?> pair in newProperties)
+								{
+									if (pair.Value == null)
+									{
+										step.Properties.Remove(pair.Key);
+									}
+									else
+									{
+										step.Properties[pair.Key] = pair.Value;
+									}
+								}
+
+								if (step.Properties.Count == 0)
+								{
+									step.Properties = null;
+									updates.Add(updateBuilder.Unset(x => x.Batches[batchIdx].Steps[stepIdx].Properties));
+								}
+								else
+								{
+									foreach (KeyValuePair<string, string?> pair in newProperties)
+									{
+										if (pair.Value == null)
+										{
+											updates.Add(updateBuilder.Unset(x => x.Batches[batchIdx].Steps[stepIdx].Properties![pair.Key]));
+										}
+										else
+										{
+											updates.Add(updateBuilder.Set(x => x.Batches[batchIdx].Steps[stepIdx].Properties![pair.Key], pair.Value));
+										}
+									}
+								}
+							}
+							break;
+						}
+					}
+					break;
+				}
+			}
+
+			// Update the batches
+			if (bRefreshBatches)
+			{
+				updates.Clear(); // UpdateBatches will update the entire batches list. We need to remove all the individual step updates to avoid an exception.
+				UpdateBatches(jobDocument, graph, updates, _logger);
+			}
+
+			// Update the state of dependent jobsteps
+			if (bRefreshDependentJobSteps)
+			{
+				RefreshDependentJobSteps(jobDocument, graph, updates, _logger);
+				RefreshJobPriority(jobDocument, updates);
+			}
+
+			// Update the new list of job steps
+			return TryUpdateAsync(jobDocument, updates);
+		}
+
+		Task<IJob?> TryUpdateAsync(JobDocument job, List<UpdateDefinition<JobDocument>> updates)
+		{
+			if (updates.Count == 0)
+			{
+				return Task.FromResult<IJob?>(job);
+			}
+			else
+			{
+				return TryUpdateAsync(job, Builders<JobDocument>.Update.Combine(updates));
+			}
+		}
+
+		async Task<IJob?> TryUpdateAsync(JobDocument job, UpdateDefinition<JobDocument> update)
+		{
+			int newUpdateIndex = job.UpdateIndex + 1;
+			update = update.Set(x => x.UpdateIndex, newUpdateIndex);
+
+			DateTime newUpdateTimeUtc = DateTime.UtcNow;
+			update = update.Set(x => x.UpdateTimeUtc, newUpdateTimeUtc);
+
+			UpdateResult result = await _jobs.UpdateOneAsync<JobDocument>(x => x.Id == job.Id && x.UpdateIndex == job.UpdateIndex, update);
+			if (result.ModifiedCount > 0)
+			{
+				job.UpdateIndex = newUpdateIndex;
+				job.UpdateTimeUtc = newUpdateTimeUtc;
+				return job;
+			}
+			return null;
+		}
+
+		/// <inheritdoc/>
+		public Task<IJob?> TryRemoveFromDispatchQueueAsync(IJob job)
+		{
+			JobDocument jobDocument = Clone((JobDocument)job);
+
+			// Create the update 
+			UpdateDefinitionBuilder<JobDocument> updateBuilder = Builders<JobDocument>.Update;
+			List<UpdateDefinition<JobDocument>> updates = new List<UpdateDefinition<JobDocument>>();
+
+			jobDocument.SchedulePriority = 0;
+			updates.Add(updateBuilder.Set(x => x.SchedulePriority, jobDocument.SchedulePriority));
+
+			// Update the new list of job steps
+			return TryUpdateAsync(jobDocument, updates);
+		}
+
+		/// <inheritdoc/>
+		public Task<IJob?> TryUpdateGraphAsync(IJob job, IGraph newGraph)
+		{
+			JobDocument jobDocument = (JobDocument)job;
+
+			// Create the update 
+			UpdateDefinitionBuilder<JobDocument> updateBuilder = Builders<JobDocument>.Update;
+			List<UpdateDefinition<JobDocument>> updates = new List<UpdateDefinition<JobDocument>>();
+
+			jobDocument.GraphHash = newGraph.Id;
+			updates.Add(updateBuilder.Set(x => x.GraphHash, job.GraphHash));
+
+			UpdateBatches(jobDocument, newGraph, updates, _logger);
+
+			// Update the new list of job steps
+			return TryUpdateAsync(jobDocument, updates);
+		}
+
+		/// <inheritdoc/>
+		public async Task AddIssueToJobAsync(JobId jobId, int issueId)
+		{
+			FilterDefinition<JobDocument> jobFilter = Builders<JobDocument>.Filter.Eq(x => x.Id, jobId);
+			UpdateDefinition<JobDocument> jobUpdate = Builders<JobDocument>.Update.AddToSet(x => x.ReferencedByIssues, issueId).Inc(x => x.UpdateIndex, 1).Max(x => x.UpdateTimeUtc, DateTime.UtcNow);
+			await _jobs.UpdateOneAsync(jobFilter, jobUpdate);
+		}
+
+		/// <inheritdoc/>
+		public async Task<List<IJob>> GetDispatchQueueAsync()
+		{
+			List<JobDocument> newJobs = await _jobs.Find(x => x.SchedulePriority > 0).SortByDescending(x => x.SchedulePriority).ThenBy(x => x.CreateTimeUtc).ToListAsync();
+			foreach (JobDocument result in newJobs)
+			{
+				await PostLoadAsync(result);
+			}
+			return newJobs.ConvertAll<JobDocument, IJob>(x => x);
+		}
+
+		/// <summary>
+		/// Marks a job as skipped
+		/// </summary>
+		/// <param name="job">The job to update</param>
+		/// <param name="graph">Graph for the job</param>
+		/// <param name="reason">Reason for this batch being failed</param>
+		/// <returns>Updated version of the job</returns>
+		public async Task<IJob?> SkipAllBatchesAsync(IJob? job, IGraph graph, JobStepBatchError reason)
+		{
+			while (job != null)
+			{
+				JobDocument jobDocument = Clone((JobDocument)job);
+
+				for (int batchIdx = 0; batchIdx < jobDocument.Batches.Count; batchIdx++)
+				{
+					JobStepBatchDocument batch = jobDocument.Batches[batchIdx];
+					if (batch.State == JobStepBatchState.Ready || batch.State == JobStepBatchState.Waiting)
+					{
+						batch.State = JobStepBatchState.Complete;
+						batch.Error = reason;
+						batch.FinishTime = DateTimeOffset.UtcNow;
+
+						for (int stepIdx = 0; stepIdx < batch.Steps.Count; stepIdx++)
+						{
+							JobStepDocument step = batch.Steps[stepIdx];
+							if (step.State == JobStepState.Ready || step.State == JobStepState.Waiting)
+							{
+								step.State = JobStepState.Completed;
+								step.Outcome = JobStepOutcome.Failure;
+							}
+						}
+					}
+				}
+
+				List<UpdateDefinition<JobDocument>> updates = new List<UpdateDefinition<JobDocument>>();
+				UpdateBatches(jobDocument, graph, updates, _logger);
+
+				IJob? newJob = await TryUpdateAsync(jobDocument, updates);
+				if (newJob != null)
+				{
+					return newJob;
+				}
+
+				job = await GetAsync(job.Id);
+			}
+			return job;
+		}
+
+		/// <inheritdoc/>
+		public async Task<IJob?> SkipBatchAsync(IJob? job, SubResourceId batchId, IGraph graph, JobStepBatchError reason)
+		{
+			while (job != null)
+			{
+				JobDocument jobDocument = (JobDocument)job;
+
+				JobStepBatchDocument? batch = jobDocument.Batches.FirstOrDefault(x => x.Id == batchId);
+				if (batch == null)
+				{
+					return job;
+				}
+
+				batch.State = JobStepBatchState.Complete;
+				batch.Error = reason;
+				batch.FinishTime = DateTimeOffset.UtcNow;
+
+				foreach (JobStepDocument step in batch.Steps)
+				{
+					if (step.State != JobStepState.Skipped)
+					{
+						step.State = JobStepState.Skipped;
+						step.Outcome = JobStepOutcome.Failure;
+					}
+				}
+
+				List<UpdateDefinition<JobDocument>> updates = new List<UpdateDefinition<JobDocument>>();
+				UpdateBatches(jobDocument, graph, updates, _logger);
+
+				IJob? newJob = await TryUpdateAsync(jobDocument, updates);
+				if(newJob != null)
+				{
+					return newJob;
+				}
+
+				job = await GetAsync(job.Id);
+			}
+			return job;
+		}
+
+		/// <inheritdoc/>
+		public async Task<IJob?> TryAssignLeaseAsync(IJob job, int batchIdx, PoolId poolId, AgentId agentId, SessionId sessionId, LeaseId leaseId, LogId logId)
+		{
+			// Try to update the job with this agent id
+			UpdateDefinitionBuilder<JobDocument> updateBuilder = Builders<JobDocument>.Update;
+			List<UpdateDefinition<JobDocument>> updates = new List<UpdateDefinition<JobDocument>>();
+
+			updates.Add(updateBuilder.Set(x => x.Batches[batchIdx].PoolId, poolId));
+			updates.Add(updateBuilder.Set(x => x.Batches[batchIdx].AgentId, agentId));
+			updates.Add(updateBuilder.Set(x => x.Batches[batchIdx].SessionId, sessionId));
+			updates.Add(updateBuilder.Set(x => x.Batches[batchIdx].LeaseId, leaseId));
+			updates.Add(updateBuilder.Set(x => x.Batches[batchIdx].LogId, logId));
+
+			// Extra logs for catching why session IDs sometimes doesn't match in prod. Resulting in PermissionDenied errors.
+			if (batchIdx < job.Batches.Count && job.Batches[batchIdx].SessionId != null)
+			{
+				string currentSessionId = job.Batches[batchIdx].SessionId!.Value.ToString();
+				_logger.LogError("Attempt to replace current session ID {CurrSessionId} with {NewSessionId} for batch {JobId}:{BatchId}", currentSessionId, sessionId.ToString(), job.Id.ToString(), job.Batches[batchIdx].Id);
+				return null;
+			}
+
+			JobDocument jobDocument = Clone((JobDocument)job);
+			if (await TryUpdateAsync(jobDocument, updates) == null)
+			{
+				return null;
+			}
+
+			JobStepBatchDocument batch = jobDocument.Batches[batchIdx];
+			batch.AgentId = agentId;
+			batch.SessionId = sessionId;
+			batch.LeaseId = leaseId;
+			batch.LogId = logId;
+
+			return jobDocument;
+		}
+
+		/// <inheritdoc/>
+		public async Task<IJob?> TryCancelLeaseAsync(IJob job, int batchIdx)
+		{
+			_logger.LogDebug("Cancelling lease {LeaseId} for agent {AgentId}", job.Batches[batchIdx].LeaseId, job.Batches[batchIdx].AgentId);
+
+			JobDocument jobDocument = Clone((JobDocument)job);
+
+			UpdateDefinition<JobDocument> update = Builders<JobDocument>.Update.Unset(x => x.Batches[batchIdx].AgentId).Unset(x => x.Batches[batchIdx].SessionId).Unset(x => x.Batches[batchIdx].LeaseId);
+			if (await TryUpdateAsync(jobDocument, update) == null)
+			{
+				return null;
+			}
+
+			JobStepBatchDocument batch = jobDocument.Batches[batchIdx];
+			batch.AgentId = null;
+			batch.SessionId = null;
+			batch.LeaseId = null;
+
+			return jobDocument;
+		}
+
+		/// <inheritdoc/>
+		public Task<IJob?> TryFailBatchAsync(IJob job, int batchIdx, IGraph graph, JobStepBatchError error)
+		{
+			JobDocument jobDocument = Clone((JobDocument)job);
+			JobStepBatchDocument batch = jobDocument.Batches[batchIdx];
+			_logger.LogDebug("Failing batch {JobId}:{BatchId} with error {Error}", job.Id, batch.Id, error);
+
+			UpdateDefinitionBuilder<JobDocument> updateBuilder = new UpdateDefinitionBuilder<JobDocument>();
+			List<UpdateDefinition<JobDocument>> updates = new List<UpdateDefinition<JobDocument>>();
+
+			if (batch.State != JobStepBatchState.Complete)
+			{
+				batch.State = JobStepBatchState.Complete;
+				updates.Add(updateBuilder.Set(x => x.Batches[batchIdx].State, batch.State));
+
+				batch.Error = error;
+				updates.Add(updateBuilder.Set(x => x.Batches[batchIdx].Error, batch.Error));
+
+				batch.FinishTime = DateTimeOffset.UtcNow;
+				updates.Add(updateBuilder.Set(x => x.Batches[batchIdx].FinishTime, batch.FinishTime));
+			}
+
+			for (int stepIdx = 0; stepIdx < batch.Steps.Count; stepIdx++)
+			{
+				JobStepDocument step = batch.Steps[stepIdx];
+				if (step.State == JobStepState.Running)
+				{
+					int stepIdxCopy = stepIdx;
+
+					step.State = JobStepState.Aborted;
+					updates.Add(updateBuilder.Set(x => x.Batches[batchIdx].Steps[stepIdxCopy].State, step.State));
+
+					step.Outcome = JobStepOutcome.Failure;
+					updates.Add(updateBuilder.Set(x => x.Batches[batchIdx].Steps[stepIdxCopy].Outcome, step.Outcome));
+
+					step.FinishTime = _clock.UtcNow;
+					updates.Add(updateBuilder.Set(x => x.Batches[batchIdx].Steps[stepIdxCopy].FinishTime, step.FinishTime));
+				}
+				else if (step.State == JobStepState.Ready)
+				{
+					int stepIdxCopy = stepIdx;
+
+					step.State = JobStepState.Skipped;
+					updates.Add(updateBuilder.Set(x => x.Batches[batchIdx].Steps[stepIdxCopy].State, step.State));
+
+					step.Outcome = JobStepOutcome.Failure;
+					updates.Add(updateBuilder.Set(x => x.Batches[batchIdx].Steps[stepIdxCopy].Outcome, step.Outcome));
+				}
+			}
+
+			RefreshDependentJobSteps(jobDocument, graph, updates, _logger);
+			RefreshJobPriority(jobDocument, updates);
+
+			return TryUpdateAsync(jobDocument, updateBuilder.Combine(updates));
+		}
+
+		/// <summary>
+		/// Populate the list of batches for a job
+		/// </summary>
+		/// <param name="job">The job to update</param>
+		/// <param name="graph">The graph definition for this job</param>
+		/// <param name="logger">Logger for output messages</param>
+		private static void CreateBatches(JobDocument job, IGraph graph, ILogger logger)
+		{
+			UpdateBatches(job, graph, new List<UpdateDefinition<JobDocument>>(), logger);
+		}
+
+		/// <summary>
+		/// Updates the list of batches for a job
+		/// </summary>
+		/// <param name="job">Job to update</param>
+		/// <param name="graph">Graph definition for this job</param>
+		/// <param name="updates">List of updates for the job</param>
+		/// <param name="logger">Logger for updates</param>
+		static void UpdateBatches(JobDocument job, IGraph graph, List<UpdateDefinition<JobDocument>> updates, ILogger logger)
+		{
+			UpdateDefinitionBuilder<JobDocument> updateBuilder = Builders<JobDocument>.Update;
+
+			// Update the list of batches
+			CreateOrUpdateBatches(job, graph);
+			updates.Add(updateBuilder.Set(x => x.Batches, job.Batches));
+			updates.Add(updateBuilder.Set(x => x.NextSubResourceId, job.NextSubResourceId));
+
+			// Update all the dependencies
+			RefreshDependentJobSteps(job, graph, new List<UpdateDefinition<JobDocument>>(), logger);
+			RefreshJobPriority(job, updates);
+		}
+
+		/// <summary>
+		/// Update the jobsteps for the given node graph 
+		/// </summary>
+		/// <param name="job">The job to update</param>
+		/// <param name="graph">The graph for this job</param>
+		private static void CreateOrUpdateBatches(JobDocument job, IGraph graph)
+		{
+			// Find the priorities of each node, incorporating all the per-step overrides
+			Dictionary<INode, Priority> nodePriorities = new Dictionary<INode, Priority>();
+			foreach (INodeGroup group in graph.Groups)
+			{
+				foreach (INode node in group.Nodes)
+				{
+					nodePriorities[node] = node.Priority;
+				}
+			}
+			foreach (JobStepBatchDocument batch in job.Batches)
+			{
+				INodeGroup group = graph.Groups[batch.GroupIdx];
+				foreach (JobStepDocument step in batch.Steps)
+				{
+					if (step.Priority != null)
+					{
+						INode node = group.Nodes[step.NodeIdx];
+						nodePriorities[node] = step.Priority.Value;
+					}
+				}
+			}
+
+			// Remove any steps and batches that haven't started yet
+			foreach (JobStepBatchDocument batch in job.Batches)
+			{
+				batch.Steps.RemoveAll(x => x.State == JobStepState.Waiting || x.State == JobStepState.Ready);
+			}
+
+			// Remove any skipped nodes whose skipped state is no longer valid
+			HashSet<INode> failedNodes = new HashSet<INode>();
+			foreach (JobStepBatchDocument batch in job.Batches)
+			{
+				INodeGroup group = graph.Groups[batch.GroupIdx];
+				foreach (JobStepDocument step in batch.Steps)
+				{
+					INode node = group.Nodes[step.NodeIdx];
+					if (step.Retry)
+					{
+						failedNodes.Remove(node);
+					}
+					else if (step.State == JobStepState.Skipped && (node.InputDependencies.Any(x => failedNodes.Contains(graph.GetNode(x))) || !CanRetryNode(job, batch.GroupIdx, step.NodeIdx)))
+					{
+						failedNodes.Add(node);
+					}
+					else if (step.Outcome == JobStepOutcome.Failure)
+					{
+						failedNodes.Add(node);
+					}
+					else
+					{
+						failedNodes.Remove(node);
+					}
+				}
+				batch.Steps.RemoveAll(x => x.State == JobStepState.Skipped && !failedNodes.Contains(group.Nodes[x.NodeIdx]));
+			}
+
+			// Remove any batches which are now empty
+			job.Batches.RemoveAll(x => x.Steps.Count == 0 && x.Error == JobStepBatchError.None);
+
+			// Find all the targets in this job
+			HashSet<string> targets = new HashSet<string>(StringComparer.OrdinalIgnoreCase);
+			if (job.AbortedByUserId == null)
+			{
+				foreach (string argument in job.Arguments)
+				{
+					if (argument.StartsWith(IJob.TargetArgumentPrefix, StringComparison.OrdinalIgnoreCase))
+					{
+						targets.UnionWith(argument.Substring(IJob.TargetArgumentPrefix.Length).Split(';'));
+					}
+				}
+				targets.Add(IJob.SetupNodeName);
+			}
+
+			// Add all the referenced aggregates
+			HashSet<INode> newNodesToExecute = new HashSet<INode>();
+			foreach (IAggregate aggregate in graph.Aggregates)
+			{
+				if (targets.Contains(aggregate.Name))
+				{
+					newNodesToExecute.UnionWith(aggregate.Nodes.Select(x => graph.GetNode(x)));
+				}
+			}
+
+			// Add any individual nodes
+			foreach (INode node in graph.Groups.SelectMany(x => x.Nodes))
+			{
+				if (targets.Contains(node.Name))
+				{
+					newNodesToExecute.Add(node);
+				}
+			}
+
+			// Also add any dependencies of these nodes
+			for (int groupIdx = graph.Groups.Count - 1; groupIdx >= 0; groupIdx--)
+			{
+				INodeGroup group = graph.Groups[groupIdx];
+				for (int nodeIdx = group.Nodes.Count - 1; nodeIdx >= 0; nodeIdx--)
+				{
+					INode node = group.Nodes[nodeIdx];
+					if (newNodesToExecute.Contains(node))
+					{
+						foreach (NodeRef dependency in node.InputDependencies)
+						{
+							newNodesToExecute.Add(graph.GetNode(dependency));
+						}
+					}
+				}
+			}
+
+			// Cancel any batches which are still running but are no longer required
+			foreach (JobStepBatchDocument batch in job.Batches)
+			{
+				if (batch.State == JobStepBatchState.Starting || batch.State == JobStepBatchState.Running)
+				{
+					INodeGroup group = graph.Groups[batch.GroupIdx];
+					if (!batch.Steps.Any(x => newNodesToExecute.Contains(group.Nodes[x.NodeIdx])))
+					{
+						batch.Error = JobStepBatchError.Cancelled;
+					}
+				}
+			}
+
+			// Remove all the nodes which have already succeeded
+			foreach (JobStepBatchDocument batch in job.Batches)
+			{
+				foreach (IJobStep step in batch.Steps)
+				{
+					if ((step.State == JobStepState.Running && step.RetriedByUserId == null) || (step.State == JobStepState.Completed && step.RetriedByUserId == null) || (step.State == JobStepState.Aborted && step.RetriedByUserId == null) || step.State == JobStepState.Skipped)
+					{
+						newNodesToExecute.Remove(graph.Groups[batch.GroupIdx].Nodes[step.NodeIdx]);
+					}
+				}
+			}
+
+			// Re-add all the nodes that have input dependencies in the same group.
+			for (int groupIdx = graph.Groups.Count - 1; groupIdx >= 0; groupIdx--)
+			{
+				INodeGroup group = graph.Groups[groupIdx];
+				for (int nodeIdx = group.Nodes.Count - 1; nodeIdx >= 0; nodeIdx--)
+				{
+					INode node = group.Nodes[nodeIdx];
+					if (newNodesToExecute.Contains(node))
+					{
+						foreach (NodeRef dependency in node.InputDependencies)
+						{
+							if (dependency.GroupIdx == groupIdx)
+							{
+								newNodesToExecute.Add(group.Nodes[dependency.NodeIdx]);
+							}
+						}
+					}
+				}
+			}
+
+			// Build a list of nodes which are currently set to be executed
+			HashSet<INode> existingNodesToExecute = new HashSet<INode>();
+			foreach (JobStepBatchDocument batch in job.Batches)
+			{
+				INodeGroup group = graph.Groups[batch.GroupIdx];
+				foreach (IJobStep step in batch.Steps)
+				{
+					INode node = group.Nodes[step.NodeIdx];
+					existingNodesToExecute.Add(node);
+				}
+			}
+
+			// Figure out the existing batch for each group
+			JobStepBatchDocument?[] appendToBatches = new JobStepBatchDocument?[graph.Groups.Count];
+			foreach (JobStepBatchDocument batch in job.Batches)
+			{
+				if (batch.CanBeAppendedTo())
+				{
+					INodeGroup group = graph.Groups[batch.GroupIdx];
+					INode firstNode = group.Nodes[batch.Steps[0].NodeIdx];
+					appendToBatches[batch.GroupIdx] = batch;
+				}
+			}
+
+			// Invalidate all the entries for groups where we're too late to append new entries (ie. we need to execute an earlier node that wasn't executed previously)
+			for (int groupIdx = 0; groupIdx < graph.Groups.Count; groupIdx++)
+			{
+				INodeGroup group = graph.Groups[groupIdx];
+				for (int nodeIdx = 0; nodeIdx < group.Nodes.Count; nodeIdx++)
+				{
+					INode node = group.Nodes[nodeIdx];
+					if (newNodesToExecute.Contains(node) && !existingNodesToExecute.Contains(node))
+					{
+						IJobStepBatch? batch = appendToBatches[groupIdx];
+						if (batch != null)
+						{
+							IJobStep lastStep = batch.Steps[batch.Steps.Count - 1];
+							if (nodeIdx < lastStep.NodeIdx)
+							{
+								appendToBatches[groupIdx] = null;
+							}
+						}
+					}
+				}
+			}
+
+			// Create all the new jobsteps
+			for (int groupIdx = 0; groupIdx < graph.Groups.Count; groupIdx++)
+			{
+				INodeGroup group = graph.Groups[groupIdx];
+				for (int nodeIdx = 0; nodeIdx < group.Nodes.Count; nodeIdx++)
+				{
+					INode node = group.Nodes[nodeIdx];
+					if (newNodesToExecute.Contains(node))
+					{
+						JobStepBatchDocument? batch = appendToBatches[groupIdx];
+						if (batch == null)
+						{
+							job.NextSubResourceId = job.NextSubResourceId.Next();
+
+							batch = new JobStepBatchDocument(job.NextSubResourceId, groupIdx);
+							job.Batches.Add(batch);
+
+							appendToBatches[groupIdx] = batch;
+						}
+						if (batch.Steps.Count == 0 || nodeIdx > batch.Steps[^1].NodeIdx)
+						{
+							job.NextSubResourceId = job.NextSubResourceId.Next();
+
+							JobStepDocument step = new JobStepDocument(job.NextSubResourceId, nodeIdx);
+							batch.Steps.Add(step);
+						}
+					}
+				}
+			}
+
+			// Find the priority of each node, propagating dependencies from dependent nodes
+			for (int groupIdx = 0; groupIdx < graph.Groups.Count; groupIdx++)
+			{
+				INodeGroup group = graph.Groups[groupIdx];
+				for (int nodeIdx = 0; nodeIdx < group.Nodes.Count; nodeIdx++)
+				{
+					INode node = group.Nodes[nodeIdx];
+					Priority nodePriority = nodePriorities[node];
+
+					foreach (NodeRef dependencyRef in node.OrderDependencies)
+					{
+						INode dependency = graph.Groups[dependencyRef.GroupIdx].Nodes[dependencyRef.NodeIdx];
+						if (nodePriorities[node] > nodePriority)
+						{
+							nodePriorities[dependency] = nodePriority;
+						}
+					}
+				}
+			}
+			foreach (JobStepBatchDocument batch in job.Batches)
+			{
+				if (batch.Steps.Count > 0)
+				{
+					INodeGroup group = graph.Groups[batch.GroupIdx];
+					Priority nodePriority = batch.Steps.Max(x => nodePriorities[group.Nodes[x.NodeIdx]]);
+					batch.SchedulePriority = ((int)job.Priority * 10) + (int)nodePriority + 1; // Reserve '0' for none.
+				}
+			}
+
+			// Check we're not running a node which doesn't allow retries more than once
+			Dictionary<INode, int> nodeExecutionCount = new Dictionary<INode, int>();
+			foreach (IJobStepBatch batch in job.Batches)
+			{
+				INodeGroup group = graph.Groups[batch.GroupIdx];
+				foreach (IJobStep step in batch.Steps)
+				{
+					INode node = group.Nodes[step.NodeIdx];
+
+					int count;
+					nodeExecutionCount.TryGetValue(node, out count);
+
+					if (!node.AllowRetry && count > 0)
+					{
+						throw new RetryNotAllowedException(node.Name);
+					}
+
+					nodeExecutionCount[node] = count + 1;
+				}
+			}
+		}
+
+		/// <summary>
+		/// Tests whether a node can be retried again
+		/// </summary>
+		static bool CanRetryNode(JobDocument job, int groupIdx, int nodeIdx)
+		{
+			return job.RetriedNodes == null || job.RetriedNodes.Count(x => x.GroupIdx == groupIdx && x.NodeIdx == nodeIdx) < MaxRetries;
+		}
+
+		/// <summary>
+		/// Gets the scheduling priority of this job
+		/// </summary>
+		/// <param name="job">Job to consider</param>
+		public static int GetSchedulePriority(IJob job)
+		{
+			int newSchedulePriority = 0;
+			foreach (IJobStepBatch batch in job.Batches)
+			{
+				if (batch.State == JobStepBatchState.Ready)
+				{
+					newSchedulePriority = Math.Max(batch.SchedulePriority, newSchedulePriority);
+				}
+			}
+			return newSchedulePriority;
+		}
+
+		/// <summary>
+		/// Update the state of any jobsteps that are dependent on other jobsteps (eg. transition them from waiting to ready based on other steps completing)
+		/// </summary>
+		/// <param name="job">The job to update</param>
+		/// <param name="graph">The graph for this job</param>
+		/// <param name="updates">List of updates to the job</param>
+		/// <param name="logger">Logger instance</param>
+		static void RefreshDependentJobSteps(JobDocument job, IGraph graph, List<UpdateDefinition<JobDocument>> updates, ILogger logger)
+		{
+			// Update the batches
+			UpdateDefinitionBuilder<JobDocument> updateBuilder = Builders<JobDocument>.Update;
+			if (job.Batches != null)
+			{
+				Dictionary<INode, IJobStep> stepForNode = new Dictionary<INode, IJobStep>();
+				for (int loopBatchIdx = 0; loopBatchIdx < job.Batches.Count; loopBatchIdx++)
+				{
+					int batchIdx = loopBatchIdx; // For lambda capture
+					JobStepBatchDocument batch = job.Batches[batchIdx];
+
+					for (int loopStepIdx = 0; loopStepIdx < batch.Steps.Count; loopStepIdx++)
+					{
+						int stepIdx = loopStepIdx; // For lambda capture
+						JobStepDocument step = batch.Steps[stepIdx];
+
+						JobStepState newState = step.State;
+						JobStepOutcome newOutcome = step.Outcome;
+
+						INode node = graph.Groups[batch.GroupIdx].Nodes[step.NodeIdx];
+						if (newState == JobStepState.Waiting)
+						{
+							List<IJobStep> steps = GetDependentSteps(graph, node, stepForNode);
+							if (steps.Any(x => x.AbortRequested || x.IsFailedOrSkipped()))
+							{
+								newState = JobStepState.Skipped;
+								newOutcome = JobStepOutcome.Failure;
+							}
+							else if (!steps.Any(x => !x.AbortRequested && x.IsPending()))
+							{
+								logger.LogDebug("Transitioning job {JobId}, batch {BatchId}, step {StepId} to ready state ({Dependencies})", job.Id, batch.Id, step.Id, String.Join(", ", steps.Select(x => x.Id.ToString())));
+								newState = JobStepState.Ready;
+							}
+						}
+
+						if (newState != step.State)
+						{
+							step.State = newState;
+							updates.Add(updateBuilder.Set(x => x.Batches[batchIdx].Steps[stepIdx].State, newState));
+						}
+
+						if (newOutcome != step.Outcome)
+						{
+							step.Outcome = newOutcome;
+							updates.Add(updateBuilder.Set(x => x.Batches[batchIdx].Steps[stepIdx].Outcome, newOutcome));
+						}
+
+						stepForNode[node] = step;
+					}
+
+					if (batch.State == JobStepBatchState.Waiting || batch.State == JobStepBatchState.Ready)
+					{
+						DateTime? newReadyTime;
+						JobStepBatchState newState = GetBatchState(job, graph, batch, stepForNode, out newReadyTime);
+						if (batch.State != newState)
+						{
+							batch.State = newState;
+							updates.Add(updateBuilder.Set(x => x.Batches[batchIdx].State, batch.State));
+						}
+						if (batch.ReadyTime != newReadyTime)
+						{
+							batch.ReadyTime = newReadyTime;
+							updates.Add(updateBuilder.Set(x => x.Batches[batchIdx].ReadyTime, batch.ReadyTime));
+						}
+					}
+				}
+			}
+		}
+
+		/// <summary>
+		/// Updates the schedule priority of a job
+		/// </summary>
+		/// <param name="job"></param>
+		/// <param name="updates"></param>
+		static void RefreshJobPriority(JobDocument job, List<UpdateDefinition<JobDocument>> updates)
+		{
+			// Update the weighted priority for the job
+			int newSchedulePriority = GetSchedulePriority(job);
+			if (job.SchedulePriority != newSchedulePriority)
+			{
+				job.SchedulePriority = newSchedulePriority;
+				updates.Add(Builders<JobDocument>.Update.Set(x => x.SchedulePriority, newSchedulePriority));
+			}
+		}
+
+		/// <summary>
+		/// Gets the steps that a node depends on
+		/// </summary>
+		/// <param name="graph">The graph for this job</param>
+		/// <param name="node">The node to test</param>
+		/// <param name="stepForNode">Map of node to step</param>
+		/// <returns></returns>
+		static List<IJobStep> GetDependentSteps(IGraph graph, INode node, Dictionary<INode, IJobStep> stepForNode)
+		{
+			List<IJobStep> steps = new List<IJobStep>();
+			foreach (NodeRef orderDependencyRef in node.OrderDependencies)
+			{
+				IJobStep? step;
+				if (stepForNode.TryGetValue(graph.GetNode(orderDependencyRef), out step))
+				{
+					steps.Add(step);
+				}
+			}
+			return steps;
+		}
+
+		/// <summary>
+		/// Gets the new state for a batch
+		/// </summary>
+		/// <param name="job">The job being executed</param>
+		/// <param name="graph">Graph for the job</param>
+		/// <param name="batch">List of nodes in the job</param>
+		/// <param name="stepForNode">Array mapping each index to the appropriate step for that node</param>
+		/// <param name="outReadyTimeUtc">Receives the time at which the batch was ready to execute</param>
+		/// <returns>True if the batch is ready, false otherwise</returns>
+		static JobStepBatchState GetBatchState(IJob job, IGraph graph, IJobStepBatch batch, Dictionary<INode, IJobStep> stepForNode, out DateTime? outReadyTimeUtc)
+		{
+			// Check if the batch is already complete
+			if (batch.Steps.All(x => x.State == JobStepState.Skipped || x.State == JobStepState.Completed || x.State == JobStepState.Aborted))
+			{
+				outReadyTimeUtc = batch.ReadyTimeUtc;
+				return JobStepBatchState.Complete;
+			}
+
+			// Get the dependencies for this batch to start. Some steps may be "after" dependencies that are optional parts of the graph.
+			List<INode> nodeDependencies = batch.GetStartDependencies(graph.Groups).ToList();
+
+			// Check if we're still waiting on anything
+			DateTime readyTimeUtc = job.CreateTimeUtc;
+			foreach (INode nodeDependency in nodeDependencies)
+			{
+				IJobStep? stepDependency;
+				if (stepForNode.TryGetValue(nodeDependency, out stepDependency))
+				{
+					if (stepDependency.State != JobStepState.Completed && stepDependency.State != JobStepState.Skipped && stepDependency.State != JobStepState.Aborted)
+					{
+						outReadyTimeUtc = null;
+						return JobStepBatchState.Waiting;
+					}
+
+					if (stepDependency.FinishTimeUtc != null && stepDependency.FinishTimeUtc.Value > readyTimeUtc)
+					{
+						readyTimeUtc = stepDependency.FinishTimeUtc.Value;
+					}
+				}
+			}
+
+			// Otherwise return the ready state
+			outReadyTimeUtc = readyTimeUtc;
+			return JobStepBatchState.Ready;
+		}
+
+		/// <inheritdoc/>
+		public async Task UpgradeDocumentsAsync()
+		{
+			IAsyncCursor<JobDocument> cursor = await _jobs.Find(Builders<JobDocument>.Filter.Eq(x => x.UpdateTimeUtc, null)).ToCursorAsync();
+
+			int numUpdated = 0;
+			while (await cursor.MoveNextAsync())
+			{
+				foreach (JobDocument document in cursor.Current)
+				{
+					UpdateDefinition<JobDocument> update = Builders<JobDocument>.Update.Set(x => x.CreateTimeUtc, ((IJob)document).CreateTimeUtc).Set(x => x.UpdateTimeUtc, ((IJob)document).UpdateTimeUtc);
+					await _jobs.UpdateOneAsync(Builders<JobDocument>.Filter.Eq(x => x.Id, document.Id), update);
+					numUpdated++;
+				}
+				_logger.LogInformation("Updated {NumDocuments} documents", numUpdated);
+			}
+		}
+	}
+}
