Horde: Finish re-org of Horde solution started a while ago.

diff --git a/Engine/Source/Programs/Horde/Horde.Build/Jobs/JobService.cs b/Engine/Source/Programs/Horde/Horde.Build/Jobs/JobService.cs
new file mode 100644
--- /dev/null
+++ b/Engine/Source/Programs/Horde/Horde.Build/Jobs/JobService.cs
@@ -0,0 +1,1358 @@
+// Copyright Epic Games, Inc. All Rights Reserved.
+
+using System;
+using System.Collections.Generic;
+using System.Globalization;
+using System.Linq;
+using System.Security.Claims;
+using System.Text.RegularExpressions;
+using System.Threading.Tasks;
+using EpicGames.Core;
+using Horde.Build.Acls;
+using Horde.Build.Api;
+using Horde.Build.Collections;
+using Horde.Build.Models;
+using Horde.Build.Tasks.Impl;
+using Horde.Build.Utilities;
+using HordeCommon;
+using Microsoft.Extensions.Logging;
+using MongoDB.Bson;
+using MongoDB.Driver;
+using OpenTracing;
+using OpenTracing.Util;
+
+namespace Horde.Build.Services
+{
+	using JobId = ObjectId<IJob>;
+	using LeaseId = ObjectId<ILease>;
+	using LogId = ObjectId<ILogFile>;
+	using SessionId = ObjectId<ISession>;
+	using StreamId = StringId<IStream>;
+	using TemplateRefId = StringId<TemplateRef>;
+	using UserId = ObjectId<IUser>;
+
+	/// <summary>
+	/// Exception thrown when attempting to retry executing a node that does not allow retries
+	/// </summary>
+	public class RetryNotAllowedException : Exception
+	{
+		/// <summary>
+		/// The node name
+		/// </summary>
+		public string NodeName { get; private set; }
+
+		/// <summary>
+		/// Constructor
+		/// </summary>
+		/// <param name="nodeName">Name of the node that does not permit retries</param>
+		public RetryNotAllowedException(string nodeName)
+			: base($"Node '{nodeName}' does not permit retries")
+		{
+			NodeName = nodeName;
+		}
+	}
+
+	/// <summary>
+	/// Cache of information about job ACLs
+	/// </summary>
+	public class JobPermissionsCache : StreamPermissionsCache
+	{
+		/// <summary>
+		/// Map of job id to permissions for that job
+		/// </summary>
+		public Dictionary<JobId, IJobPermissions?> Jobs { get; set; } = new Dictionary<JobId, IJobPermissions?>();
+	}
+
+	/// <summary>
+	/// Wraps funtionality for manipulating jobs, jobsteps, and jobstep runs
+	/// </summary>
+	public class JobService
+	{
+		readonly IJobCollection _jobs;
+		readonly IGraphCollection _graphs;
+		readonly IAgentCollection _agents;
+		readonly IJobStepRefCollection _jobStepRefs;
+		readonly IJobTimingCollection _jobTimings;
+		readonly IUserCollection _userCollection;
+		readonly INotificationTriggerCollection _triggerCollection;
+		readonly JobTaskSource _jobTaskSource;
+		readonly StreamService _streamService;
+		readonly ITemplateCollection _templateCollection;
+		readonly IssueService? _issueService;
+		readonly IPerforceService _perforceService;
+		readonly ILogger _logger;
+
+		/// <summary>
+		/// Delegate for label update events
+		/// </summary>
+		public delegate void LabelUpdateEvent(IJob job, IReadOnlyList<(LabelState, LabelOutcome)> oldLabelStates, IReadOnlyList<(LabelState, LabelOutcome)> newLabelStates);
+
+		/// <summary>
+		/// Event triggered when a label state changes
+		/// </summary>
+		public event LabelUpdateEvent? OnLabelUpdate;
+
+		/// <summary>
+		/// Delegate for job step complete events
+		/// </summary>
+		public delegate void JobStepCompleteEvent(IJob job, IGraph graph, SubResourceId batchId, SubResourceId stepId);
+
+		/// <summary>
+		/// Event triggered when a job step completes
+		/// </summary>
+		public event JobStepCompleteEvent? OnJobStepComplete;
+		
+		/// <summary>
+		/// Event triggered when a job is scheduled on the underlying JobTaskSource
+		///
+		/// Exposed and duplicated here to avoid dependency on the exact JobTaskSource
+		/// </summary>
+		public event JobTaskSource.JobScheduleEvent? OnJobScheduled;
+
+		/// <summary>
+		/// Constructor
+		/// </summary>
+		/// <param name="jobs">The jobs collection</param>
+		/// <param name="graphs">The graphs collection</param>
+		/// <param name="agents">The agents collection</param>
+		/// <param name="jobStepRefs">The jobsteprefs collection</param>
+		/// <param name="jobTimings">The job timing document collection</param>
+		/// <param name="userCollection">User profiles</param>
+		/// <param name="triggerCollection">Trigger collection</param>
+		/// <param name="jobTaskSource">The queue service</param>
+		/// <param name="streamService">The stream service</param>
+		/// <param name="templateCollection">The template service</param>
+		/// <param name="issueService">The issue service</param>
+		/// <param name="perforceService">The perforce service</param>
+		/// <param name="logger">Log output</param>
+		public JobService(IJobCollection jobs, IGraphCollection graphs, IAgentCollection agents, IJobStepRefCollection jobStepRefs, IJobTimingCollection jobTimings, IUserCollection userCollection, INotificationTriggerCollection triggerCollection, JobTaskSource jobTaskSource, StreamService streamService, ITemplateCollection templateCollection, IssueService issueService, IPerforceService perforceService, ILogger<JobService> logger)
+		{
+			_jobs = jobs;
+			_graphs = graphs;
+			_agents = agents;
+			_jobStepRefs = jobStepRefs;
+			_jobTimings = jobTimings;
+			_userCollection = userCollection;
+			_triggerCollection = triggerCollection;
+			_jobTaskSource = jobTaskSource;
+			_streamService = streamService;
+			_templateCollection = templateCollection;
+			_issueService = issueService;
+			_perforceService = perforceService;
+			_logger = logger;
+
+			_jobTaskSource.OnJobScheduled += (pool, numAgentsOnline, job, graph, batchId) =>
+			{
+				OnJobScheduled?.Invoke(pool, numAgentsOnline, job, graph, batchId);
+			};
+		}
+
+#pragma warning disable IDE0060
+		static bool ShouldClonePreflightChange(StreamId streamId)
+		{
+			return false; //			return StreamId == new StreamId("ue5-main");
+		}
+#pragma warning restore IDE0060
+
+		/// <summary>
+		/// Creates a new job
+		/// </summary>
+		/// <param name="jobId">A requested job id</param>
+		/// <param name="stream">The stream that this job belongs to</param>
+		/// <param name="templateRefId">Name of the template ref</param>
+		/// <param name="templateHash">Template for this job</param>
+		/// <param name="graph">The graph for the new job</param>
+		/// <param name="name">Name of the job</param>
+		/// <param name="change">The change to build</param>
+		/// <param name="codeChange">The corresponding code changelist</param>
+		/// <param name="preflightChange">Optional changelist to preflight</param>
+		/// <param name="clonedPreflightChange">Duplicated preflight change</param>
+		/// <param name="startedByUserId">Id of the user that started the job</param>
+		/// <param name="priority">Priority of the job</param>
+		/// <param name="autoSubmit">Whether to automatically submit the preflighted change on completion</param>
+		/// <param name="updateIssues">Whether to update issues when this job completes</param>
+		/// <param name="promoteIssuesByDefault">Whether to promote issues from this job by default</param>
+		/// <param name="jobTriggers">List of downstream job triggers</param>
+		/// <param name="showUgsBadges">Whether to show badges in UGS for this job</param>
+		/// <param name="showUgsAlerts">Whether to show alerts in UGS for this job</param>
+		/// <param name="notificationChannel">Notification Channel for this job</param>
+		/// <param name="notificationChannelFilter">Notification Channel filter for this job</param>
+		/// <param name="arguments">Arguments for the job</param>
+		/// <returns>Unique id representing the job</returns>
+		public async Task<IJob> CreateJobAsync(JobId? jobId, IStream stream, TemplateRefId templateRefId, ContentHash templateHash, IGraph graph, string name, int change, int codeChange, int? preflightChange, int? clonedPreflightChange, UserId? startedByUserId, Priority? priority, bool? autoSubmit, bool? updateIssues, bool? promoteIssuesByDefault, List<ChainedJobTemplate>? jobTriggers, bool showUgsBadges, bool showUgsAlerts, string? notificationChannel, string? notificationChannelFilter, IReadOnlyList<string> arguments)
+		{
+			using IScope traceScope = GlobalTracer.Instance.BuildSpan("JobService.CreateJobAsync").StartActive();
+			traceScope.Span.SetTag("JobId", jobId);
+			traceScope.Span.SetTag("Stream", stream.Name);
+			traceScope.Span.SetTag("TemplateRefId", templateRefId);
+			traceScope.Span.SetTag("TemplateHash", templateHash);
+			traceScope.Span.SetTag("GraphId", graph.Id);
+			traceScope.Span.SetTag("Name", name);
+			traceScope.Span.SetTag("Change", change);
+			traceScope.Span.SetTag("CodeChange", codeChange);
+			traceScope.Span.SetTag("PreflightChange", preflightChange);
+			traceScope.Span.SetTag("ClonedPreflightChange", clonedPreflightChange);
+			traceScope.Span.SetTag("StartedByUserId", startedByUserId.ToString());
+			traceScope.Span.SetTag("Priority", priority.ToString());
+			traceScope.Span.SetTag("ShowUgsBadges", showUgsBadges);
+			traceScope.Span.SetTag("NotificationChannel", notificationChannel ?? "null");
+			traceScope.Span.SetTag("NotificationChannelFilter", notificationChannelFilter ?? "null");
+			traceScope.Span.SetTag("Arguments", String.Join(',', arguments));
+
+			if (autoSubmit != null)
+			{
+				traceScope.Span.SetTag("AutoSubmit", autoSubmit.Value);
+			}
+			if (updateIssues != null)
+			{
+				traceScope.Span.SetTag("UpdateIssues", updateIssues.Value);
+			}
+			if (jobTriggers != null)
+			{
+				traceScope.Span.SetTag("JobTriggers.Count", jobTriggers.Count);
+			}
+
+			JobId jobIdValue = jobId ?? Horde.Build.Utilities.ObjectId<IJob>.GenerateNewId();
+			using IDisposable scope = _logger.BeginScope("CreateJobAsync({JobId})", jobIdValue);
+
+			if (preflightChange != null && ShouldClonePreflightChange(stream.Id))
+			{
+				clonedPreflightChange = await CloneShelvedChangeAsync(stream.ClusterName, clonedPreflightChange ?? preflightChange.Value);
+			}
+
+			_logger.LogInformation("Creating job at CL {Change}, code CL {CodeChange}, preflight CL {PreflightChange}, cloned CL {ClonedPreflightChange}", change, codeChange, preflightChange, clonedPreflightChange);
+
+			Dictionary<string, string> properties = new Dictionary<string, string>();
+			properties["Change"] = change.ToString(CultureInfo.InvariantCulture);
+			properties["CodeChange"] = codeChange.ToString(CultureInfo.InvariantCulture);
+			properties["PreflightChange"] = preflightChange?.ToString(CultureInfo.InvariantCulture) ?? String.Empty;
+			properties["ClonedPreflightChange"] = clonedPreflightChange?.ToString(CultureInfo.InvariantCulture) ?? String.Empty;
+			properties["StreamId"] = stream.Id.ToString();
+			properties["TemplateId"] = templateRefId.ToString();
+			properties["JobId"] = jobIdValue.ToString();
+
+			List<string> expandedArguments = new List<string>();
+			if (arguments != null)
+			{
+				foreach (string argument in arguments)
+				{
+					string expandedArgument = StringUtils.ExpandProperties(argument, properties);
+					expandedArguments.Add(expandedArgument);
+				}
+			}
+
+			name = StringUtils.ExpandProperties(name, properties);
+
+			IJob newJob = await _jobs.AddAsync(jobIdValue, stream.Id, templateRefId, templateHash, graph, name, change, codeChange, preflightChange, clonedPreflightChange, startedByUserId, priority, autoSubmit, updateIssues, promoteIssuesByDefault, jobTriggers, showUgsBadges, showUgsAlerts, notificationChannel, notificationChannelFilter, expandedArguments);
+			_jobTaskSource.UpdateQueuedJob(newJob, graph);
+
+			await _jobTaskSource.UpdateUgsBadges(newJob, graph, new List<(LabelState, LabelOutcome)>());
+
+			if (startedByUserId != null)
+			{
+				await _userCollection.UpdateSettingsAsync(startedByUserId.Value, addPinnedJobIds: new[] { newJob.Id });
+			}
+
+			await AbortAnyDuplicateJobs(newJob);
+
+			return newJob;
+		}
+
+		private async Task AbortAnyDuplicateJobs(IJob newJob)
+		{
+			using IScope scope = GlobalTracer.Instance.BuildSpan("JobService.AbortAnyDuplicateJobs").StartActive();
+			scope.Span.SetTag("JobId", newJob.Id.ToString());
+			scope.Span.SetTag("JobName", newJob.Name);
+			
+			List<IJob> jobsToAbort = new List<IJob>();
+			if (newJob.PreflightChange > 0)
+			{
+				jobsToAbort = await _jobs.FindAsync(preflightChange: newJob.PreflightChange);
+			}
+
+			foreach (IJob job in jobsToAbort)
+			{
+				if (job.GetState() == JobState.Complete)
+				{
+					continue;
+				}
+				if (job.Id == newJob.Id)
+				{
+					continue; // Don't remove the new job
+				}
+				if (job.TemplateId != newJob.TemplateId)
+				{
+					continue;
+				}
+				if (job.TemplateHash != newJob.TemplateHash)
+				{
+					continue;
+				}
+				if (job.StartedByUserId != newJob.StartedByUserId)
+				{
+					continue;
+				}
+				if (String.Join(",", job.Arguments) != String.Join(",", newJob.Arguments))
+				{
+					continue;
+				}
+
+				IJob? updatedJob = await UpdateJobAsync(job, null, null, null, KnownUsers.System, null, null, null);
+				if (updatedJob == null)
+				{
+					_logger.LogError("Failed marking duplicate job as aborted! Job ID: {JobId}", job.Id);
+				}
+
+				IJob? updatedJob2 = await GetJobAsync(job.Id);
+				if (updatedJob2?.AbortedByUserId != updatedJob?.AbortedByUserId)
+				{
+					throw new NotImplementedException();
+				}
+			}
+		}
+
+		/// <summary>
+		/// Deletes a job
+		/// </summary>
+		/// <param name="job">The job to delete</param>
+		public async Task<bool> DeleteJobAsync(IJob job)
+		{
+			using IScope traceScope = GlobalTracer.Instance.BuildSpan("JobService.DeleteJobAsync").StartActive();
+			traceScope.Span.SetTag("JobId", job.Id.ToString());
+			traceScope.Span.SetTag("JobName", job.Name);
+
+			using IDisposable scope = _logger.BeginScope("DeleteJobAsync({JobId})", job.Id);
+
+			// Delete the job
+			while (!await _jobs.RemoveAsync(job))
+			{
+				IJob? newJob = await _jobs.GetAsync(job.Id);
+				if (newJob == null)
+				{
+					return false;
+				}
+				job = newJob;
+			}
+
+			// Remove all the triggers from it
+			List<ObjectId> triggerIds = new List<ObjectId>();
+			if (job.NotificationTriggerId != null)
+			{
+				triggerIds.Add(job.NotificationTriggerId.Value);
+			}
+			foreach (IJobStep step in job.Batches.SelectMany(x => x.Steps))
+			{
+				if (step.NotificationTriggerId != null)
+				{
+					triggerIds.Add(step.NotificationTriggerId.Value);
+				}
+			}
+			await _triggerCollection.DeleteAsync(triggerIds);
+
+			return true;
+		}
+
+		/// <summary>
+		/// Delete all the jobs for a stream
+		/// </summary>
+		/// <param name="streamId">Unique id of the stream</param>
+		/// <returns>Async task</returns>
+		public async Task DeleteJobsForStreamAsync(StreamId streamId)
+		{
+			await _jobs.RemoveStreamAsync(streamId);
+		}
+
+		/// <summary>
+		/// Updates a new job
+		/// </summary>
+		/// <param name="job">The job document to update</param>
+		/// <param name="name">Name of the job</param>
+		/// <param name="priority">Priority of the job</param>
+		/// <param name="autoSubmit">Whether to automatically submit the preflighted change on completion</param>
+		/// <param name="abortedByUserId">Name of the user that aborted this job</param>
+		/// <param name="onCompleteTriggerId">Object id for a notification trigger</param>
+		/// <param name="reports">New reports to add</param>
+		/// <param name="arguments">New arguments for the job</param>
+		/// <param name="labelIdxToTriggerId">New trigger ID for a label in the job</param>
+		public async Task<IJob?> UpdateJobAsync(IJob job, string? name = null, Priority? priority = null, bool? autoSubmit = null, UserId? abortedByUserId = null, ObjectId? onCompleteTriggerId = null, List<Report>? reports = null, List<string>? arguments = null, KeyValuePair<int, ObjectId>? labelIdxToTriggerId = null)
+		{
+			using IScope traceScope = GlobalTracer.Instance.BuildSpan("JobService.UpdateJobAsync").StartActive();
+			traceScope.Span.SetTag("JobId", job.Id.ToString());
+			traceScope.Span.SetTag("Name", name);
+			
+			using IDisposable scope = _logger.BeginScope("UpdateJobAsync({JobId})", job.Id);
+			for(IJob? newJob = job; newJob != null; newJob = await GetJobAsync(job.Id))
+			{
+				IGraph graph = await GetGraphAsync(newJob);
+
+				// Capture the previous label states
+				IReadOnlyList<(LabelState, LabelOutcome)> oldLabelStates = newJob.GetLabelStates(graph);
+
+				// Update the new list of job steps
+				newJob = await _jobs.TryUpdateJobAsync(newJob, graph, name, priority, autoSubmit, null, null, abortedByUserId, onCompleteTriggerId, reports, arguments, labelIdxToTriggerId);
+				if (newJob != null)
+				{
+					// Update any badges that have been modified
+					await _jobTaskSource.UpdateUgsBadges(newJob, graph, oldLabelStates);
+
+					// Cancel any leases which are no longer required
+					foreach (IJobStepBatch batch in newJob.Batches)
+					{
+						if (batch.Error == JobStepBatchError.Cancelled && (batch.State == JobStepBatchState.Starting || batch.State == JobStepBatchState.Running) && batch.AgentId != null && batch.LeaseId != null)
+						{
+							await CancelLeaseAsync(batch.AgentId.Value, batch.LeaseId.Value);
+						}
+					}
+					return newJob;
+				}
+			}
+			return null;
+		}
+
+		/// <summary>
+		/// Cancel an active lease
+		/// </summary>
+		/// <param name="agentId">The agent to retreive</param>
+		/// <param name="leaseId">The lease id to update</param>
+		/// <returns></returns>
+		async Task CancelLeaseAsync(AgentId agentId, LeaseId leaseId)
+		{
+			using IScope scope = GlobalTracer.Instance.BuildSpan("JobService.CancelLeaseAsync").StartActive();
+			scope.Span.SetTag("AgentId", agentId.ToString());
+			scope.Span.SetTag("LeaseId", leaseId.ToString());
+			
+			for (; ; )
+			{
+				IAgent? agent = await _agents.GetAsync(agentId);
+				if (agent == null)
+				{
+					break;
+				}
+
+				int index = 0;
+				while (index < agent.Leases.Count && agent.Leases[index].Id != leaseId)
+				{
+					index++;
+				}
+				if (index == agent.Leases.Count)
+				{
+					break;
+				}
+
+				AgentLease lease = agent.Leases[index];
+				if (lease.State != LeaseState.Active)
+				{
+					break;
+				}
+
+				IAgent? newAgent = await _agents.TryCancelLeaseAsync(agent, index);
+				if (newAgent != null)
+				{
+					_jobTaskSource.CancelLongPollForAgent(agent.Id);
+					break;
+				}
+			}
+		}
+
+		/// <summary>
+		/// Gets a job with the given unique id
+		/// </summary>
+		/// <param name="jobId">Job id to search for</param>
+		/// <returns>Information about the given job</returns>
+		public Task<IJob?> GetJobAsync(JobId jobId)
+		{
+			return _jobs.GetAsync(jobId);
+		}
+
+		/// <summary>
+		/// Gets the graph for a job
+		/// </summary>
+		/// <param name="job">Job to retrieve the graph for</param>
+		/// <returns>The graph for this job</returns>
+		public Task<IGraph> GetGraphAsync(IJob job)
+		{
+			return _graphs.GetAsync(job.GraphHash);
+		}
+
+		/// <summary>
+		/// Gets a job's permissions info by ID
+		/// </summary>
+		/// <param name="jobId">Unique id of the job</param>
+		/// <returns>The job document</returns>
+		public Task<IJobPermissions?> GetJobPermissionsAsync(JobId jobId)
+		{
+			return _jobs.GetPermissionsAsync(jobId);
+		}
+
+		/// <summary>
+		/// Searches for jobs matching the given criteria
+		/// </summary>
+		/// <param name="jobIds">List of job ids to return</param>
+		/// <param name="streamId">The stream containing the job</param>
+		/// <param name="name">Name of the job</param>
+		/// <param name="templates">Templates to look for</param>
+		/// <param name="minChange">The minimum changelist number</param>
+		/// <param name="maxChange">The maximum changelist number</param>		
+		/// <param name="preflightChange">The preflight change to look for</param>
+		/// <param name="preflightOnly">Whether to only include preflights</param>
+		/// <param name="preflightStartedByUser">User for which to include preflight jobs</param>		
+		/// <param name="startedByUser">User for which to include jobs</param>
+		/// <param name="minCreateTime">The minimum creation time</param>
+		/// <param name="maxCreateTime">The maximum creation time</param>
+		/// <param name="target">The target to query</param>
+		/// <param name="state">State to query</param>
+		/// <param name="outcome">Outcomes to return</param>
+		/// <param name="modifiedBefore">Filter the results by last modified time</param>
+		/// <param name="modifiedAfter">Filter the results by last modified time</param>
+		/// <param name="index">Index of the first result to return</param>
+		/// <param name="count">Number of results to return</param>
+		/// <param name="consistentRead">If the database read should be made to the replica server</param>
+		/// <param name="excludeUserJobs">Whether to exclude user jobs from the find</param>
+		/// <returns>List of jobs matching the given criteria</returns>
+		public async Task<List<IJob>> FindJobsAsync(JobId[]? jobIds = null, StreamId? streamId = null, string? name = null, TemplateRefId[]? templates = null, int? minChange = null, int? maxChange = null, int? preflightChange = null, bool? preflightOnly = null, UserId? preflightStartedByUser = null, UserId? startedByUser = null, DateTimeOffset ? minCreateTime = null, DateTimeOffset? maxCreateTime = null, string? target = null, JobStepState[]? state = null, JobStepOutcome[]? outcome = null, DateTimeOffset? modifiedBefore = null, DateTimeOffset? modifiedAfter = null, int? index = null, int? count = null, bool consistentRead = true, bool? excludeUserJobs = null)
+		{
+			using IScope scope = GlobalTracer.Instance.BuildSpan("JobService.FindJobsAsync").StartActive();
+			scope.Span.SetTag("JobIds", jobIds);
+			scope.Span.SetTag("StreamId", streamId);
+			scope.Span.SetTag("Name", name);
+			scope.Span.SetTag("Templates", templates);
+			scope.Span.SetTag("MinChange", minChange);
+			scope.Span.SetTag("MaxChange", maxChange);
+			scope.Span.SetTag("PreflightChange", preflightChange);
+			scope.Span.SetTag("PreflightStartedByUser", preflightStartedByUser);
+			scope.Span.SetTag("StartedByUser", startedByUser);
+			scope.Span.SetTag("MinCreateTime", minCreateTime);
+			scope.Span.SetTag("MaxCreateTime", maxCreateTime);
+			scope.Span.SetTag("Target", target);
+			scope.Span.SetTag("State", state?.ToString());
+			scope.Span.SetTag("Outcome", outcome?.ToString());
+			scope.Span.SetTag("ModifiedBefore", modifiedBefore);
+			scope.Span.SetTag("ModifiedAfter", modifiedAfter);
+			scope.Span.SetTag("Index", index);
+			scope.Span.SetTag("Count", count);
+			
+			if (target == null && (state == null || state.Length == 0) && (outcome == null || outcome.Length == 0))
+			{
+				return await _jobs.FindAsync(jobIds, streamId, name, templates, minChange, maxChange, preflightChange, preflightOnly, preflightStartedByUser, startedByUser, minCreateTime, maxCreateTime, modifiedBefore, modifiedAfter, index, count, consistentRead, null, excludeUserJobs);
+			}
+			else
+			{
+				List<IJob> results = new List<IJob>();
+				_logger.LogInformation("Performing scan for job with ");
+
+				int maxCount = (count ?? 1);
+				while (results.Count < maxCount)
+				{
+					List<IJob> scanJobs = await _jobs.FindAsync(jobIds, streamId, name, templates, minChange, maxChange, preflightChange, preflightOnly, preflightStartedByUser, startedByUser, minCreateTime, maxCreateTime, modifiedBefore, modifiedAfter, 0, 5, consistentRead, null, excludeUserJobs);
+					if (scanJobs.Count == 0)
+					{
+						break;
+					}
+
+					foreach (IJob job in scanJobs.OrderByDescending(x => x.Change))
+					{
+						(JobStepState, JobStepOutcome)? result;
+						if (target == null)
+						{
+							result = job.GetTargetState();
+						}
+						else
+						{
+							result = job.GetTargetState(await GetGraphAsync(job), target);
+						}
+
+						if (result != null)
+						{
+							(JobStepState jobState, JobStepOutcome jobOutcome) = result.Value;
+							if ((state == null || state.Length == 0 || state.Contains(jobState)) && (outcome == null || outcome.Length == 0 || outcome.Contains(jobOutcome)))
+							{
+								results.Add(job);
+								if (results.Count == maxCount)
+								{
+									break;
+								}
+							}
+						}
+					}
+
+					maxChange = scanJobs.Min(x => x.Change) - 1;
+				}
+
+				return results;
+			}
+		}
+		
+		/// <summary>
+		/// Searches for jobs matching a stream with given templates
+		/// </summary>
+		/// <param name="streamId">The stream containing the job</param>
+		/// <param name="templates">Templates to look for</param>
+		/// <param name="preflightStartedByUser">User for which to include preflight jobs</param>
+		/// <param name="maxCreateTime">The maximum creation time</param>
+		/// <param name="modifiedAfter">Filter the results by last modified time</param>	
+		/// <param name="index">Index of the first result to return</param>
+		/// <param name="count">Number of results to return</param>
+		/// <param name="consistentRead">If the database read should be made to the replica server</param>
+		/// <returns>List of jobs matching the given criteria</returns>
+		public async Task<List<IJob>> FindJobsByStreamWithTemplatesAsync(StreamId streamId, TemplateRefId[] templates, UserId? preflightStartedByUser = null, DateTimeOffset? maxCreateTime = null, DateTimeOffset? modifiedAfter = null, int? index = null, int? count = null, bool consistentRead = true)
+		{
+			using IScope scope = GlobalTracer.Instance.BuildSpan("JobService.FindJobsByStreamWithTemplatesAsync").StartActive();
+			scope.Span.SetTag("StreamId", streamId);
+			scope.Span.SetTag("Templates", templates);
+			scope.Span.SetTag("PreflightStartedByUser", preflightStartedByUser);
+			scope.Span.SetTag("MaxCreateTime", maxCreateTime);
+			scope.Span.SetTag("ModifiedAfter", modifiedAfter);
+			scope.Span.SetTag("Index", index);
+			scope.Span.SetTag("Count", count);
+			
+			return await _jobs.FindLatestByStreamWithTemplatesAsync(streamId, templates, preflightStartedByUser, maxCreateTime, modifiedAfter, index, count, consistentRead);
+		}
+
+		/// <summary>
+		/// Attempts to update the node groups to be executed for a job. Fails if another write happens in the meantime.
+		/// </summary>
+		/// <param name="job">The job to update</param>
+		/// <param name="newGraph">New graph for this job</param>
+		/// <returns>True if the groups were updated to the given list. False if another write happened first.</returns>
+		public async Task<IJob?> TryUpdateGraphAsync(IJob job, IGraph newGraph)
+		{
+			using IScope traceScope = GlobalTracer.Instance.BuildSpan("JobService.TryUpdateGraphAsync").StartActive();
+			traceScope.Span.SetTag("Job", job.Id);
+			traceScope.Span.SetTag("NewGraph", newGraph.Id);
+
+			using IDisposable scope = _logger.BeginScope("TryUpdateGraphAsync({JobId})", job.Id);
+
+			IReadOnlyList<(LabelState, LabelOutcome)> oldLabelStates = job.GetLabelStates(newGraph);
+
+			IJob? newJob = await _jobs.TryUpdateGraphAsync(job, newGraph);
+			if(newJob != null)
+			{
+				await _jobTaskSource.UpdateUgsBadges(newJob, newGraph, oldLabelStates);
+				_jobTaskSource.UpdateQueuedJob(newJob, newGraph);
+			}
+			return newJob;
+		}
+
+		/// <summary>
+		/// Gets the timing info for a particular job.
+		/// </summary>
+		/// <param name="job">The job to get timing info for</param>
+		/// <returns>Timing info for the given job</returns>
+		public async Task<IJobTiming> GetJobTimingAsync(IJob job)
+		{
+			using IScope traceScope = GlobalTracer.Instance.BuildSpan("JobService.GetJobTimingAsync").StartActive();
+			traceScope.Span.SetTag("Job", job.Id);
+
+			using IDisposable scope = _logger.BeginScope("GetJobTimingAsync({JobId})", job.Id);
+
+			IGraph graph = await _graphs.GetAsync(job.GraphHash);
+
+			Dictionary<string, JobStepTimingData> cachedNewSteps = new Dictionary<string, JobStepTimingData>();
+			for (; ; )
+			{
+				// Try to get the current timing document
+				IJobTiming? jobTiming = await _jobTimings.TryGetAsync(job.Id);
+
+				// Calculate timing info for any missing steps
+				List<JobStepTimingData> newSteps = new List<JobStepTimingData>();
+				foreach (IJobStepBatch batch in job.Batches)
+				{
+					INodeGroup group = graph.Groups[batch.GroupIdx];
+					foreach (IJobStep step in batch.Steps)
+					{
+						INode node = group.Nodes[step.NodeIdx];
+						if (jobTiming == null || !jobTiming.TryGetStepTiming(node.Name, out _))
+						{
+							JobStepTimingData? stepTimingData;
+							if (!cachedNewSteps.TryGetValue(node.Name, out stepTimingData))
+							{
+								stepTimingData = await GetStepTimingInfo(job.StreamId, job.TemplateId, node.Name, job.Change);
+							}
+							newSteps.Add(stepTimingData);
+						}
+					}
+				}
+
+				// Try to add or update the timing document
+				if (jobTiming == null)
+				{
+					IJobTiming? newJobTiming = await _jobTimings.TryAddAsync(job.Id, newSteps);
+					if (newJobTiming != null)
+					{
+						return newJobTiming;
+					}
+				}
+				else if (newSteps.Count == 0)
+				{
+					return jobTiming;
+				}
+				else
+				{
+					IJobTiming? newJobTiming = await _jobTimings.TryAddStepsAsync(jobTiming, newSteps);
+					if (newJobTiming != null)
+					{
+						return newJobTiming;
+					}
+				}
+
+				// Update the cached list of steps for the next iteration
+				foreach (JobStepTimingData newStep in newSteps)
+				{
+					cachedNewSteps[newStep.Name] = newStep;
+				}
+			}
+		}
+
+		/// <summary>
+		/// Finds the average duration for the given step
+		/// </summary>
+		/// <param name="streamId">The stream to search</param>
+		/// <param name="templateId">The template id</param>
+		/// <param name="nodeName">Name of the node</param>
+		/// <param name="change">Maximum changelist to consider</param>
+		/// <returns>Expected duration for the given step</returns>
+		async Task<JobStepTimingData> GetStepTimingInfo(StreamId streamId, TemplateRefId templateId, string nodeName, int? change)
+		{
+			using IScope traceScope = GlobalTracer.Instance.BuildSpan("JobService.GetStepTimingInfo").StartActive();
+			traceScope.Span.SetTag("StreamId", streamId);
+			traceScope.Span.SetTag("TemplateId", templateId);
+			traceScope.Span.SetTag("NodeName", nodeName);
+			traceScope.Span.SetTag("Change", change);
+			
+			// Find all the steps matching the given criteria
+			List<IJobStepRef> steps = await _jobStepRefs.GetStepsForNodeAsync(streamId, templateId, nodeName, change, false, 10);
+
+			// Sum up all the durations and wait times
+			int count = 0;
+			float waitTimeSum = 0;
+			float initTimeSum = 0;
+			float durationSum = 0;
+
+			foreach (IJobStepRef step in steps)
+			{
+				if (step.FinishTimeUtc != null)
+				{
+					waitTimeSum += step.BatchWaitTime;
+					initTimeSum += step.BatchInitTime;
+					durationSum += (float)(step.FinishTimeUtc.Value - step.StartTimeUtc).TotalSeconds;
+					count++;
+				}
+			}
+
+			// Compute the averages
+			if (count == 0)
+			{
+				return new JobStepTimingData(nodeName, null, null, null);
+			}
+			else
+			{
+				return new JobStepTimingData(nodeName, waitTimeSum / count, initTimeSum / count, durationSum / count);
+			}
+		}
+
+		/// <summary>
+		/// Updates the state of a batch
+		/// </summary>
+		/// <param name="job">Job to update</param>
+		/// <param name="batchId">Unique id of the batch to update</param>
+		/// <param name="newLogId">The new log file id</param>
+		/// <param name="newState">New state of the jobstep</param>
+		/// <returns>True if the job was updated, false if it was deleted</returns>
+		public async Task<IJob?> UpdateBatchAsync(IJob job, SubResourceId batchId, LogId? newLogId = null, JobStepBatchState? newState = null)
+		{
+			using IScope traceScope = GlobalTracer.Instance.BuildSpan("JobService.UpdateBatchAsync").StartActive();
+			traceScope.Span.SetTag("Job", job.Id);
+			traceScope.Span.SetTag("BatchId", batchId);
+			traceScope.Span.SetTag("NewLogId", newLogId);
+			traceScope.Span.SetTag("NewState", newState.ToString());
+
+			using IDisposable scope = _logger.BeginScope("UpdateBatchAsync({JobId})", job.Id);
+
+			bool bCheckForBadAgent = true;
+			for (; ; )
+			{
+				// Find the index of the appropriate batch
+				int batchIdx = job.Batches.FindIndex(x => x.Id == batchId);
+				if (batchIdx == -1)
+				{
+					return null;
+				}
+
+				IJobStepBatch batch = job.Batches[batchIdx];
+
+				// If the batch has already been marked as complete, error out
+				if (batch.State == JobStepBatchState.Complete)
+				{
+					return null;
+				}
+
+				// If we're marking the batch as complete before the agent has run everything, mark it to conform
+				JobStepBatchError? newError = null;
+				if (newState == JobStepBatchState.Complete && batch.AgentId != null)
+				{
+					if (batch.Steps.Any(x => x.State == JobStepState.Waiting || x.State == JobStepState.Ready))
+					{
+						// Mark the batch as incomplete
+						newError = JobStepBatchError.Incomplete;
+
+						// Find the agent and set the conform flag
+						if (bCheckForBadAgent)
+						{
+							for (; ; )
+							{
+								IAgent? agent = await _agents.GetAsync(batch.AgentId.Value);
+								if (agent == null || agent.RequestConform)
+								{
+									break;
+								}
+								if (await _agents.TryUpdateSettingsAsync(agent, bRequestConform: true) != null)
+								{
+									_logger.LogError("Agent {AgentId} did not complete lease; marking for conform", agent.Id);
+									break;
+								}
+							}
+							bCheckForBadAgent = false;
+						}
+					}
+				}
+
+				// Update the batch state
+				IJob? newJob = await TryUpdateBatchAsync(job, batchId, newLogId, newState, newError);
+				if (newJob != null)
+				{
+					return newJob;
+				}
+
+				// Update the job
+				newJob = await GetJobAsync(job.Id);
+				if (newJob == null)
+				{
+					return null;
+				}
+
+				job = newJob;
+			}
+		}
+
+		/// <summary>
+		/// Attempts to update the state of a batch
+		/// </summary>
+		/// <param name="job">Job to update</param>
+		/// <param name="batchId">Unique id of the batch to update</param>
+		/// <param name="newLogId">The new log file id</param>
+		/// <param name="newState">New state of the jobstep</param>
+		/// <param name="newError">New error state</param>
+		/// <returns>The updated job, otherwise null</returns>
+		public async Task<IJob?> TryUpdateBatchAsync(IJob job, SubResourceId batchId, LogId? newLogId = null, JobStepBatchState? newState = null, JobStepBatchError? newError = null)
+		{
+			IGraph graph = await GetGraphAsync(job);
+			return await _jobs.TryUpdateBatchAsync(job, graph, batchId, newLogId, newState, newError);
+		}
+
+		/// <summary>
+		/// Update a jobstep state
+		/// </summary>
+		/// <param name="job">Job to update</param>
+		/// <param name="batchId">Unique id of the batch containing the step</param>
+		/// <param name="stepId">Unique id of the step to update</param>
+		/// <param name="newState">New state of the jobstep</param>
+		/// <param name="newOutcome">New outcome of the jobstep</param>
+		/// <param name="newError">New error for the step</param>
+		/// <param name="newAbortRequested">New state of abort request</param>
+		/// <param name="newAbortByUserId">New user that requested the abort</param>
+		/// <param name="newLogId">New log id for the jobstep</param>
+		/// <param name="newNotificationTriggerId">New notification trigger id for the jobstep</param>
+		/// <param name="newRetryByUserId">Whether the step should be retried</param>
+		/// <param name="newPriority">New priority for this step</param>
+		/// <param name="newReports">New list of reports</param>
+		/// <param name="newProperties">Property changes. Any properties with a null value will be removed.</param>
+		/// <returns>True if the job was updated, false if it was deleted in the meantime</returns>
+		public async Task<IJob?> UpdateStepAsync(IJob job, SubResourceId batchId, SubResourceId stepId, JobStepState newState = JobStepState.Unspecified, JobStepOutcome newOutcome = JobStepOutcome.Unspecified, JobStepError? newError = null, bool? newAbortRequested = null, UserId? newAbortByUserId = null, LogId? newLogId = null, ObjectId? newNotificationTriggerId = null, UserId? newRetryByUserId = null, Priority? newPriority = null, List<Report>? newReports = null, Dictionary<string, string?>? newProperties = null)
+		{
+			using IScope traceScope = GlobalTracer.Instance.BuildSpan("JobService.UpdateStepAsync").StartActive();
+			traceScope.Span.SetTag("Job", job.Id);
+			traceScope.Span.SetTag("BatchId", batchId);
+			traceScope.Span.SetTag("StepId", stepId);
+			
+			using IDisposable scope = _logger.BeginScope("UpdateStepAsync({JobId})", job.Id);
+			for (; ;)
+			{
+				IJob? newJob = await TryUpdateStepAsync(job, batchId, stepId, newState, newOutcome, newError, newAbortRequested, newAbortByUserId, newLogId, newNotificationTriggerId, newRetryByUserId, newPriority, newReports, newProperties);
+				if (newJob != null)
+				{
+					return newJob;
+				}
+
+				newJob = await GetJobAsync(job.Id);
+				if(newJob == null)
+				{
+					return null;
+				}
+
+				job = newJob;
+			}
+		}
+
+		/// <summary>
+		/// Update a jobstep state
+		/// </summary>
+		/// <param name="job">Job to update</param>
+		/// <param name="batchId">Unique id of the batch containing the step</param>
+		/// <param name="stepId">Unique id of the step to update</param>
+		/// <param name="newState">New state of the jobstep</param>
+		/// <param name="newOutcome">New outcome of the jobstep</param>
+		/// <param name="newError">New error for the step</param>
+		/// <param name="newAbortRequested">New state for request abort</param>
+		/// <param name="newAbortByUserId">New name of user that requested the abort</param>
+		/// <param name="newLogId">New log id for the jobstep</param>
+		/// <param name="newTriggerId">New trigger id for the jobstep</param>
+		/// <param name="newRetryByUserId">Whether the step should be retried</param>
+		/// <param name="newPriority">New priority for this step</param>
+		/// <param name="newReports">New reports</param>
+		/// <param name="newProperties">Property changes. Any properties with a null value will be removed.</param>
+		/// <returns>True if the job was updated, false if it was deleted in the meantime</returns>
+		public async Task<IJob?> TryUpdateStepAsync(IJob job, SubResourceId batchId, SubResourceId stepId, JobStepState newState = JobStepState.Unspecified, JobStepOutcome newOutcome = JobStepOutcome.Unspecified, JobStepError? newError = null, bool? newAbortRequested = null, UserId? newAbortByUserId = null, LogId? newLogId = null, ObjectId? newTriggerId = null, UserId? newRetryByUserId = null, Priority? newPriority = null, List<Report>? newReports = null, Dictionary<string, string?>? newProperties = null)
+		{
+			using IScope traceScope = GlobalTracer.Instance.BuildSpan("JobService.TryUpdateStepAsync").StartActive();
+			traceScope.Span.SetTag("Job", job.Id);
+			traceScope.Span.SetTag("BatchId", batchId);
+			traceScope.Span.SetTag("StepId", stepId);
+
+			using IDisposable scope = _logger.BeginScope("TryUpdateStepAsync({JobId})", job.Id);
+
+			// Get the graph for this job
+			IGraph graph = await GetGraphAsync(job);
+
+			// Make sure the job state is set to unspecified
+			JobStepState oldState = JobStepState.Unspecified;
+			JobStepOutcome oldOutcome = JobStepOutcome.Unspecified;
+			if (newState != JobStepState.Unspecified)
+			{
+				if (job.TryGetStep(batchId, stepId, out IJobStep? step))
+				{
+					oldState = step.State;
+					oldOutcome = step.Outcome;
+				}
+			}
+
+			// If we're changing the state of a step, capture the label states beforehand so we can update any that change
+			IReadOnlyList<(LabelState, LabelOutcome)>? oldLabelStates = null;
+			if (oldState != newState || oldOutcome != newOutcome)
+			{
+				oldLabelStates = job.GetLabelStates(graph);
+			}
+
+			// Update the step
+			IJob? newJob = await _jobs.TryUpdateStepAsync(job, graph, batchId, stepId, newState, newOutcome, newError, newAbortRequested, newAbortByUserId, newLogId, newTriggerId, newRetryByUserId, newPriority, newReports, newProperties);
+			if (newJob != null)
+			{
+				job = newJob;
+
+				using IScope ddScope = GlobalTracer.Instance.BuildSpan("TryUpdateStepAsync").StartActive();
+				ddScope.Span.SetTag("JobId", job.Id.ToString());
+				ddScope.Span.SetTag("BatchId", batchId.ToString());
+				ddScope.Span.SetTag("StepId", stepId.ToString());
+				
+				if (oldState != newState || oldOutcome != newOutcome)
+				{
+					_logger.LogDebug("Transitioned job {JobId}, batch {BatchId}, step {StepId} from {OldState} to {NewState}", job.Id, batchId, stepId, oldState, newState);
+
+					// Send any updates for modified badges
+					if (oldLabelStates != null)
+					{
+						using IScope _ = GlobalTracer.Instance.BuildSpan("Send badge updates").StartActive();
+						IReadOnlyList<(LabelState, LabelOutcome)> newLabelStates = job.GetLabelStates(graph);
+						OnLabelUpdate?.Invoke(job, oldLabelStates, newLabelStates);
+						await _jobTaskSource.UpdateUgsBadges(job, graph, oldLabelStates, newLabelStates);
+					}
+
+					// Submit the change if auto-submit is enabled
+					if (job.PreflightChange != 0 && newState == JobStepState.Completed)
+					{
+						(JobStepState state, JobStepOutcome outcome) = job.GetTargetState();
+						if (state == JobStepState.Completed)
+						{
+							if (job.AutoSubmit && outcome == JobStepOutcome.Success && job.AbortedByUserId == null)
+							{
+								job = await AutoSubmitChangeAsync(job, graph);
+							}
+							else if (job.ClonedPreflightChange != 0)
+							{
+								IStream? stream = await _streamService.GetCachedStream(job.StreamId);
+								if (stream != null)
+								{
+									await DeleteShelvedChangeAsync(stream.ClusterName, job.ClonedPreflightChange);
+								}
+							}
+						}
+					}
+
+					// Notify subscribers
+					if (newState == JobStepState.Completed)
+					{
+						using IScope _ = GlobalTracer.Instance.BuildSpan("Notify subscribers").StartActive();
+						OnJobStepComplete?.Invoke(job, graph, batchId, stepId);
+					}
+
+					// Create any downstream jobs
+					if (newState == JobStepState.Completed && newOutcome != JobStepOutcome.Failure)
+					{
+						using IScope _ = GlobalTracer.Instance.BuildSpan("Create downstream jobs").StartActive();
+						for (int idx = 0; idx < job.ChainedJobs.Count; idx++)
+						{
+							IChainedJob jobTrigger = job.ChainedJobs[idx];
+							if (jobTrigger.JobId == null)
+							{
+								JobStepOutcome jobTriggerOutcome = job.GetTargetOutcome(graph, jobTrigger.Target);
+								if (jobTriggerOutcome == JobStepOutcome.Success || jobTriggerOutcome == JobStepOutcome.Warnings)
+								{
+									job = await FireJobTriggerAsync(job, graph, jobTrigger) ?? job;
+								}
+							}
+						}
+					}
+
+					// Update the jobstep ref if it completed
+					if (newState == JobStepState.Running || newState == JobStepState.Completed || newState == JobStepState.Aborted)
+					{
+						using IScope _ = GlobalTracer.Instance.BuildSpan("Update job step ref").StartActive();
+						if (job.TryGetBatch(batchId, out IJobStepBatch? batch) && batch.TryGetStep(stepId, out IJobStep? step) && step.StartTimeUtc != null)
+						{
+							await _jobStepRefs.UpdateAsync(job, batch, step, graph);
+						}
+					}
+
+					// Update any issues that depend on this step
+					if (newState == JobStepState.Completed && job.UpdateIssues)
+					{
+						if (_issueService != null)
+						{
+							using IScope _ = GlobalTracer.Instance.BuildSpan("Update issues (V2)").StartActive();
+							try
+							{
+								await _issueService.UpdateCompleteStep(job, graph, batchId, stepId);
+							}
+							catch(Exception ex)
+							{
+								_logger.LogError(ex, "Exception while updating issue service for job {JobId}:{BatchId}:{StepId}: {Message}", job.Id, batchId, stepId, ex.Message);
+							}
+						}
+					}
+				}
+
+				using (IScope dispatchScope = GlobalTracer.Instance.BuildSpan("Update queued jobs").StartActive())
+				{
+					// Notify the dispatch service that the job has changed
+					_jobTaskSource.UpdateQueuedJob(job, graph);
+				}
+
+				return job;
+			}
+
+			return null;
+		}
+
+		/// <summary>
+		/// Submit the given change for a preflight
+		/// </summary>
+		/// <param name="job">The job being run</param>
+		/// <param name="graph">Graph for the job</param>
+		/// <returns></returns>
+		private async Task<IJob> AutoSubmitChangeAsync(IJob job, IGraph graph)
+		{
+			using IScope traceScope = GlobalTracer.Instance.BuildSpan("JobService.AutoSubmitChangeAsync").StartActive();
+			traceScope.Span.SetTag("Job", job.Id);
+			traceScope.Span.SetTag("Graph", graph.Id);
+			
+			int? change;
+			string message;
+			try
+			{
+				IStream? stream = await _streamService.GetCachedStream(job.StreamId);
+				if (stream != null)
+				{
+					int clonedPreflightChange = job.ClonedPreflightChange;
+					if (clonedPreflightChange == 0)
+					{
+						if (ShouldClonePreflightChange(job.StreamId))
+						{
+							clonedPreflightChange = await CloneShelvedChangeAsync(stream.ClusterName, job.PreflightChange);
+						}
+						else
+						{
+							clonedPreflightChange = job.PreflightChange;
+						}
+					}
+
+					_logger.LogInformation("Updating description for {ClonedPreflightChange}", clonedPreflightChange);
+
+					ChangeDetails details = await _perforceService.GetChangeDetailsAsync(stream.ClusterName, stream.Name, clonedPreflightChange, null);
+					await _perforceService.UpdateChangelistDescription(stream.ClusterName, clonedPreflightChange, details.Description.TrimEnd() + $"\n#preflight {job.Id}");
+
+					_logger.LogInformation("Submitting change {Change} (through {ChangeCopy}) after successful completion of {JobId}", job.PreflightChange, clonedPreflightChange, job.Id);
+					(change, message) = await _perforceService.SubmitShelvedChangeAsync(stream.ClusterName, clonedPreflightChange, job.PreflightChange);
+
+					_logger.LogInformation("Attempt to submit {Change} (through {ChangeCopy}): {Message}", job.PreflightChange, clonedPreflightChange, message);
+
+					if (!String.IsNullOrEmpty(message))
+					{
+						message = Regex.Replace(message, @"^Submit validation failed.*\n(?:\s*'[^']*' validation failed:\s*\n)?", "");
+					}
+
+					if (ShouldClonePreflightChange(job.StreamId))
+					{
+						if (change != null && job.ClonedPreflightChange != 0)
+						{
+							await DeleteShelvedChangeAsync(stream.ClusterName, job.PreflightChange);
+						}
+					}
+					else
+					{
+						if (change != null && job.PreflightChange != 0)
+						{
+							await DeleteShelvedChangeAsync(stream.ClusterName, job.PreflightChange);
+						}
+					}
+				}
+				else
+				{
+					(change, message) = ((int?)null, "Stream not found");
+				}
+			}
+			catch (Exception ex)
+			{
+				(change, message) = ((int?)null, "Internal error");
+				_logger.LogError(ex, "Unable to submit shelved change");
+			}
+
+			for (; ; )
+			{
+				IJob? newJob = await _jobs.TryUpdateJobAsync(job, graph, autoSubmitChange: change, autoSubmitMessage: message);
+				if (newJob != null)
+				{
+					return newJob;
+				}
+
+				newJob = await GetJobAsync(job.Id);
+				if (newJob == null)
+				{
+					return job;
+				}
+
+				job = newJob;
+			}
+		}
+
+		/// <summary>
+		/// Clone a shelved changelist for running a preflight
+		/// </summary>
+		/// <param name="clusterName">Name of the Perforce cluster</param>
+		/// <param name="change">The changelist to clone</param>
+		/// <returns></returns>
+		private async Task<int> CloneShelvedChangeAsync(string clusterName, int change)
+		{
+			int clonedChange;
+			try
+			{
+				clonedChange = await _perforceService.DuplicateShelvedChangeAsync(clusterName, change);
+				_logger.LogInformation("CL {Change} was duplicated into {ClonedChange}", change, clonedChange);
+			}
+			catch (Exception ex)
+			{
+				_logger.LogError(ex, "Unable to CL {Change} for preflight: {Message}", change, ex.Message);
+				throw;
+			}
+			return clonedChange;
+		}
+
+		/// <summary>
+		/// Deletes a shelved changelist, catching and logging any exceptions
+		/// </summary>
+		/// <param name="clusterName"></param>
+		/// <param name="change">The changelist to delete</param>
+		/// <returns>True if the change was deleted successfully, false otherwise</returns>
+		private async Task<bool> DeleteShelvedChangeAsync(string clusterName, int change)
+		{
+			_logger.LogInformation("Removing shelf {Change}", change);
+			try
+			{
+				await _perforceService.DeleteShelvedChangeAsync(clusterName, change);
+				return true;
+			}
+			catch (Exception ex)
+			{
+				_logger.LogError(ex, "Unable to delete shelved change {Change}", change);
+				return false;
+			}
+		}
+
+		/// <summary>
+		/// Fires a trigger for a chained job
+		/// </summary>
+		/// <param name="job">The job containing the trigger</param>
+		/// <param name="graph">Graph for the job containing the trigger</param>
+		/// <param name="jobTrigger">The trigger object to fire</param>
+		/// <returns>New job instance</returns>
+		private async Task<IJob?> FireJobTriggerAsync(IJob job, IGraph graph, IChainedJob jobTrigger)
+		{
+			using IScope traceScope = GlobalTracer.Instance.BuildSpan("JobService.FireJobTriggerAsync").StartActive();
+			traceScope.Span.SetTag("Job", job.Id);
+			traceScope.Span.SetTag("Graph", graph.Id);
+			traceScope.Span.SetTag("JobTrigger.Target", jobTrigger.Target);
+			traceScope.Span.SetTag("JobTrigger.JobId", jobTrigger.JobId);
+			traceScope.Span.SetTag("JobTrigger.TemplateRefId", jobTrigger.TemplateRefId);
+			
+			for (; ; )
+			{
+				// Update the job
+				JobId chainedJobId = JobId.GenerateNewId();
+
+				IJob? newJob = await _jobs.TryUpdateJobAsync(job, graph, jobTrigger: new KeyValuePair<TemplateRefId, JobId>(jobTrigger.TemplateRefId, chainedJobId));
+				if(newJob != null)
+				{
+					IStream? stream = await _streamService.GetStreamAsync(newJob.StreamId);
+					if(stream == null)
+					{
+						_logger.LogWarning("Cannot find stream {StreamId} for downstream job", newJob.StreamId);
+						break;
+					}
+
+					TemplateRef? templateRef;
+					if (!stream.Templates.TryGetValue(jobTrigger.TemplateRefId, out templateRef))
+					{
+						_logger.LogWarning("Cannot find template {TemplateId} in stream {StreamId}", jobTrigger.TemplateRefId, newJob.StreamId);
+						break;
+					}
+
+					ITemplate? template = await _templateCollection.GetAsync(templateRef.Hash);
+					if (template == null)
+					{
+						_logger.LogWarning("Cannot find template {TemplateHash}", templateRef.Hash);
+						break;
+					}
+
+					IGraph triggerGraph = await _graphs.AddAsync(template);
+					_logger.LogInformation("Creating downstream job {ChainedJobId} from job {JobId}", chainedJobId, newJob.Id);
+
+					await CreateJobAsync(chainedJobId, stream, jobTrigger.TemplateRefId, templateRef.Hash, triggerGraph, templateRef.Name, newJob.Change, newJob.CodeChange, newJob.PreflightChange, newJob.ClonedPreflightChange, newJob.StartedByUserId, template.Priority, null, newJob.UpdateIssues, newJob.PromoteIssuesByDefault, templateRef.ChainedJobs, false, false, templateRef.NotificationChannel, templateRef.NotificationChannelFilter, template.Arguments);
+					return newJob;
+				}
+
+				// Fetch the job again
+				newJob = await _jobs.GetAsync(job.Id);
+				if(newJob == null)
+				{
+					break;
+				}
+
+				job = newJob;
+			}
+			return null;
+		}
+
+		/// <summary>
+		/// Determines if the user is authorized to perform an action on a particular stream
+		/// </summary>
+		/// <param name="acl">The ACL to check</param>
+		/// <param name="streamId">The stream containing this job</param>
+		/// <param name="action">The action being performed</param>
+		/// <param name="user">The principal to authorize</param>
+		/// <param name="cache">Cache of stream permissions</param>
+		/// <returns>True if the action is authorized</returns>
+		private Task<bool> AuthorizeAsync(Acl? acl, StreamId streamId, AclAction action, ClaimsPrincipal user, StreamPermissionsCache? cache)
+		{
+			// Do the regular authorization
+			bool? result = acl?.Authorize(action, user);
+			if (result == null)
+			{
+				return _streamService.AuthorizeAsync(streamId, action, user, cache);
+			}
+			else
+			{
+				return Task.FromResult(result.Value);
+			}
+		}
+
+		/// <summary>
+		/// Determines if the user is authorized to perform an action on a particular stream
+		/// </summary>
+		/// <param name="job">The job to check</param>
+		/// <param name="action">The action being performed</param>
+		/// <param name="user">The principal to authorize</param>
+		/// <param name="cache">Cache of stream permissions</param>
+		/// <returns>True if the action is authorized</returns>
+		public Task<bool> AuthorizeAsync(IJob job, AclAction action, ClaimsPrincipal user, StreamPermissionsCache? cache)
+		{
+			return AuthorizeAsync(job.Acl, job.StreamId, action, user, cache);
+		}
+
+		/// <summary>
+		/// Determines if the user is authorized to perform an action on a particular stream
+		/// </summary>
+		/// <param name="jobId">The job to check</param>
+		/// <param name="action">The action being performed</param>
+		/// <param name="user">The principal to authorize</param>
+		/// <param name="cache">Cache of job permissions</param>
+		/// <returns>True if the action is authorized</returns>
+		public async Task<bool> AuthorizeAsync(JobId jobId, AclAction action, ClaimsPrincipal user, JobPermissionsCache? cache)
+		{
+			IJobPermissions? permissions;
+			if (cache == null)
+			{
+				permissions = await GetJobPermissionsAsync(jobId);
+			}
+			else if (!cache.Jobs.TryGetValue(jobId, out permissions))
+			{
+				permissions = await GetJobPermissionsAsync(jobId);
+				cache.Jobs.Add(jobId, permissions);
+			}
+			return permissions != null && await AuthorizeAsync(permissions.Acl, permissions.StreamId, action, user, cache);
+		}
+
+		/// <summary>
+		/// Determines if the user is authorized to perform an action on a particular stream
+		/// </summary>
+		/// <param name="job">The job to check</param>
+		/// <param name="user">The principal to authorize</param>
+		/// <returns>True if the action is authorized</returns>
+		public static bool AuthorizeSession(IJob job, ClaimsPrincipal user)
+		{
+			return job.Batches.Any(x => AuthorizeSession(x, user));
+		}
+
+		/// <summary>
+		/// Determines if the user is authorized to perform an action on a particular stream
+		/// </summary>
+		/// <param name="batch">The batch being executed</param>
+		/// <param name="user">The principal to authorize</param>
+		/// <returns>True if the action is authorized</returns>
+		public static bool AuthorizeSession(IJobStepBatch batch, ClaimsPrincipal user)
+		{
+			if(batch.SessionId != null)
+			{
+				foreach (Claim claim in user.Claims)
+				{
+					if(claim.Type == HordeClaimTypes.AgentSessionId)
+					{
+						SessionId sessionIdValue;
+						if (SessionId.TryParse(claim.Value, out sessionIdValue) && sessionIdValue == batch.SessionId.Value)
+						{
+							return true;
+						}
+					}
+				}
+			}
+			return false;
+		}
+		
+		/// <summary>
+		/// Get a list of each batch's session ID formatted as a string for debugging purposes
+		/// </summary>
+		/// <param name="job">The job to list</param>
+		/// <returns>List of session IDs</returns>
+		public static string GetAllBatchSessionIds(IJob job)
+		{
+			return String.Join(",", job.Batches.Select(b => b.SessionId));
+		}
+	}
+}
diff --git a/Engine/Source/Programs/Horde/Horde.Build/Jobs/JobStepRefCollection.cs b/Engine/Source/Programs/Horde/Horde.Build/Jobs/JobStepRefCollection.cs
new file mode 100644
--- /dev/null
+++ b/Engine/Source/Programs/Horde/Horde.Build/Jobs/JobStepRefCollection.cs
@@ -0,0 +1,145 @@
+// Copyright Epic Games, Inc. All Rights Reserved.
+
+using System;
+using System.Collections.Generic;
+using System.Linq;
+using System.Threading.Tasks;
+using Horde.Build.Models;
+using Horde.Build.Server;
+using Horde.Build.Utilities;
+using HordeCommon;
+using MongoDB.Bson.Serialization.Attributes;
+using MongoDB.Driver;
+
+namespace Horde.Build.Collections.Impl
+{
+	using LogId = ObjectId<ILogFile>;
+	using PoolId = StringId<IPool>;
+	using StreamId = StringId<IStream>;
+	using TemplateRefId = StringId<TemplateRef>;
+
+	/// <summary>
+	/// Collection of JobStepRef documents
+	/// </summary>
+	public class JobStepRefCollection : IJobStepRefCollection
+	{
+		class JobStepRef : IJobStepRef
+		{
+			[BsonId]
+			public JobStepRefId Id { get; set; }
+
+			public string JobName { get; set; } = "Unknown";
+
+			public string Name { get; set; }
+
+			public StreamId StreamId { get; set; }
+			public TemplateRefId TemplateId { get; set; }
+			public int Change { get; set; }
+			public LogId? LogId { get; set; }
+			public PoolId? PoolId { get; set; }
+			public AgentId? AgentId { get; set; }
+			public JobStepOutcome? Outcome { get; set; }
+
+			[BsonIgnoreIfNull]
+			public int? LastSuccess { get; }
+
+			[BsonIgnoreIfNull]
+			public int? LastWarning { get; }
+
+			public float BatchWaitTime { get; set; }
+			public float BatchInitTime { get; set; }
+
+			[BsonIgnoreIfNull]
+			public DateTimeOffset? StartTime { get; set; }
+
+			[BsonIgnoreIfNull]
+			public DateTime? StartTimeUtc { get; set; }
+
+			[BsonIgnoreIfNull]
+			public DateTimeOffset? FinishTime { get; set; }
+
+			[BsonIgnoreIfNull]
+			public DateTime? FinishTimeUtc { get; set; }
+
+			DateTime IJobStepRef.StartTimeUtc => StartTimeUtc ?? StartTime?.UtcDateTime ?? default;
+			DateTime? IJobStepRef.FinishTimeUtc => FinishTimeUtc ?? FinishTime?.UtcDateTime;
+			string IJobStepRef.NodeName => Name;
+
+			public JobStepRef(JobStepRefId id, string jobName, string nodeName, StreamId streamId, TemplateRefId templateId, int change, LogId? logId, PoolId? poolId, AgentId? agentId, JobStepOutcome? outcome, int? lastSuccess, int? lastWarning, float batchWaitTime, float batchInitTime, DateTime startTimeUtc, DateTime? finishTimeUtc)
+			{
+				Id = id;
+				JobName = jobName;
+				Name = nodeName;
+				StreamId = streamId;
+				TemplateId = templateId;
+				Change = change;
+				LogId = logId;
+				PoolId = poolId;
+				AgentId = agentId;
+				Outcome = outcome;
+				LastSuccess = lastSuccess;
+				LastWarning = lastWarning;
+				BatchWaitTime = batchWaitTime;
+				BatchInitTime = batchInitTime;
+				StartTime = startTimeUtc;
+				StartTimeUtc = startTimeUtc;
+				FinishTime = finishTimeUtc;
+				FinishTimeUtc = finishTimeUtc;
+			}
+		}
+
+		readonly IMongoCollection<JobStepRef> _jobStepRefs;
+
+		/// <summary>
+		/// Constructor
+		/// </summary>
+		/// <param name="mongoService">The database service instance</param>
+		public JobStepRefCollection(MongoService mongoService)
+		{
+			_jobStepRefs = mongoService.GetCollection<JobStepRef>("JobStepRefs", keys => keys.Ascending(x => x.StreamId).Ascending(x => x.TemplateId).Ascending(x => x.Name).Descending(x => x.Change));
+		}
+
+		/// <inheritdoc/>
+		public async Task<IJobStepRef> InsertOrReplaceAsync(JobStepRefId id, string jobName, string stepName, StreamId streamId, TemplateRefId templateId, int change, LogId? logId, PoolId? poolId, AgentId? agentId, JobStepOutcome? outcome, int? lastSuccess, int? lastWarning, float waitTime, float initTime, DateTime startTimeUtc, DateTime? finishTimeUtc)
+		{
+			JobStepRef newJobStepRef = new JobStepRef(id, jobName, stepName, streamId, templateId, change, logId, poolId, agentId, outcome, lastSuccess, lastWarning, waitTime, initTime, startTimeUtc, finishTimeUtc);
+			await _jobStepRefs.ReplaceOneAsync(Builders<JobStepRef>.Filter.Eq(x => x.Id, newJobStepRef.Id), newJobStepRef, new ReplaceOptions { IsUpsert = true });
+			return newJobStepRef;
+		}
+
+		/// <inheritdoc/>
+		public async Task<List<IJobStepRef>> GetStepsForNodeAsync(StreamId streamId, TemplateRefId templateId, string nodeName, int? change, bool includeFailed, int count)
+		{
+			// Find all the steps matching the given criteria
+			FilterDefinitionBuilder<JobStepRef> filterBuilder = Builders<JobStepRef>.Filter;
+
+			FilterDefinition<JobStepRef> filter = FilterDefinition<JobStepRef>.Empty;
+			filter &= filterBuilder.Eq(x => x.StreamId, streamId);
+			filter &= filterBuilder.Eq(x => x.TemplateId, templateId);
+			filter &= filterBuilder.Eq(x => x.Name, nodeName);
+			if (change != null)
+			{
+				filter &= filterBuilder.Lte(x => x.Change, change.Value);
+			}
+			if (!includeFailed)
+			{
+				filter &= filterBuilder.Ne(x => x.Outcome, JobStepOutcome.Failure);
+			}
+
+			List<JobStepRef> steps = await _jobStepRefs.Find(filter).SortByDescending(x => x.Change).ThenByDescending(x => x.StartTime).Limit(count).ToListAsync();
+			return steps.ConvertAll<IJobStepRef>(x => x);
+		}
+
+		/// <inheritdoc/>
+		public async Task<IJobStepRef?> GetPrevStepForNodeAsync(StreamId streamId, TemplateRefId templateId, string nodeName, int change)
+		{
+			return await _jobStepRefs.Find(x => x.StreamId == streamId && x.TemplateId == templateId && x.Name == nodeName && x.Change < change && x.Outcome != null).SortByDescending(x => x.Change).FirstOrDefaultAsync();
+		}
+
+		/// <inheritdoc/>
+		public async Task<IJobStepRef?> GetNextStepForNodeAsync(StreamId streamId, TemplateRefId templateId, string nodeName, int change)
+		{
+			return await _jobStepRefs.Find(x => x.StreamId == streamId && x.TemplateId == templateId && x.Name == nodeName && x.Change > change && x.Outcome != null).SortBy(x => x.Change).FirstOrDefaultAsync();
+		}
+	}
+}
