Horde: Finish re-org of Horde solution started a while ago.

diff --git a/Engine/Source/Programs/Horde/Horde.Build/Jobs/JobTaskSource.cs b/Engine/Source/Programs/Horde/Horde.Build/Jobs/JobTaskSource.cs
new file mode 100644
--- /dev/null
+++ b/Engine/Source/Programs/Horde/Horde.Build/Jobs/JobTaskSource.cs
@@ -0,0 +1,1136 @@
+// Copyright Epic Games, Inc. All Rights Reserved.
+
+using System;
+using System.Collections.Generic;
+using System.Diagnostics;
+using System.Diagnostics.CodeAnalysis;
+using System.Globalization;
+using System.Linq;
+using System.Text;
+using System.Threading;
+using System.Threading.Tasks;
+using EpicGames.Core;
+using Google.Protobuf;
+using Google.Protobuf.WellKnownTypes;
+using Horde.Build.Api;
+using Horde.Build.Collections;
+using Horde.Build.Commits;
+using Horde.Build.Models;
+using Horde.Build.Server;
+using Horde.Build.Services;
+using Horde.Build.Utilities;
+using HordeCommon;
+using HordeCommon.Rpc.Tasks;
+using Microsoft.Extensions.Hosting;
+using Microsoft.Extensions.Logging;
+using Microsoft.Extensions.Options;
+using MongoDB.Driver;
+
+namespace Horde.Build.Tasks.Impl
+{
+	using JobId = ObjectId<IJob>;
+	using LeaseId = ObjectId<ILease>;
+	using LogId = ObjectId<ILogFile>;
+	using PoolId = StringId<IPool>;
+	using StreamId = StringId<IStream>;
+
+	/// <summary>
+	/// Background service to dispatch pending work to agents in priority order.
+	/// </summary>
+	public sealed class JobTaskSource : TaskSourceBase<ExecuteJobTask>, IHostedService, IDisposable
+	{
+		/// <inheritdoc/>
+		public override string Type => "Job";
+
+		/// <inheritdoc/>
+		public override TaskSourceFlags Flags => TaskSourceFlags.None;
+
+		/// <summary>
+		/// An item in the queue to be executed
+		/// </summary>
+		[DebuggerDisplay("{_job.Id}:{Batch.Id} ({_poolId})")]
+		internal class QueueItem
+		{
+			/// <summary>
+			/// The stream for this job
+			/// </summary>
+			public IStream _stream;
+
+			/// <summary>
+			/// The job instance
+			/// </summary>
+			public IJob _job;
+
+			/// <summary>
+			/// Index of the batch within this job to be executed
+			/// </summary>
+			public int _batchIdx;
+
+			/// <summary>
+			/// The pool of machines to allocate from
+			/// </summary>
+			public PoolId _poolId;
+
+			/// <summary>
+			/// The type of workspace that this item should run in
+			/// </summary>
+			public AgentWorkspace _workspace;
+
+			/// <summary>
+			/// Whether or not to use the AutoSDK.
+			/// </summary>
+			public bool _useAutoSdk;
+
+			/// <summary>
+			/// Task for creating a lease and assigning to a waiter
+			/// </summary>
+			public Task? _assignTask;
+
+			/// <summary>
+			/// Accessor for the batch referenced by this item
+			/// </summary>
+			public IJobStepBatch Batch => _job.Batches[_batchIdx];
+
+			/// <summary>
+			/// Returns an identifier describing this unique batch
+			/// </summary>
+			public (JobId, SubResourceId) Id => (_job.Id, Batch.Id);
+
+			/// <summary>
+			/// Constructor
+			/// </summary>
+			/// <param name="stream">The stream containing this job</param>
+			/// <param name="job">The job instance</param>
+			/// <param name="batchIdx">The batch index to execute</param>
+			/// <param name="poolId">Unique id of the pool of machines to allocate from</param>
+			/// <param name="workspace">The workspace that this job should run in</param>
+			/// <param name="useAutoSdk">Whether or not to use the AutoSDK</param>
+			public QueueItem(IStream stream, IJob job, int batchIdx, PoolId poolId, AgentWorkspace workspace, bool useAutoSdk)
+			{
+				_stream = stream;
+				_job = job;
+				_batchIdx = batchIdx;
+				_poolId = poolId;
+				_workspace = workspace;
+				_useAutoSdk = useAutoSdk;
+			}
+		}
+
+		/// <summary>
+		/// Comparer for items in the queue
+		/// </summary>
+		class QueueItemComparer : IComparer<QueueItem>
+		{
+			/// <summary>
+			/// Compare two items
+			/// </summary>
+			/// <param name="x">First item to compare</param>
+			/// <param name="y">Second item to compare</param>
+			/// <returns>Negative value if X is a higher priority than Y</returns>
+			public int Compare([AllowNull] QueueItem x, [AllowNull] QueueItem y)
+			{
+				if (x == null)
+				{
+					return (y == null) ? 0 : -1;
+				}
+				else if (y == null)
+				{
+					return 1;
+				}
+
+				int delta = y.Batch.SchedulePriority - x.Batch.SchedulePriority;
+				if (delta == 0)
+				{
+					delta = x._job.Id.CompareTo(y._job.Id);
+					if (delta == 0)
+					{
+						delta = (int)x.Batch.Id.Value - (int)y.Batch.Id.Value;
+					}
+				}
+				return delta;
+			}
+		}
+
+		/// <summary>
+		/// Information about an agent waiting for work
+		/// </summary>
+		class QueueWaiter
+		{
+			/// <summary>
+			/// The agent performing the wait
+			/// </summary>
+			public IAgent Agent { get; }
+
+			/// <summary>
+			/// Task to wait for a lease to be assigned
+			/// </summary>
+			public Task<AgentLease?> Task => LeaseSource.Task;
+
+			/// <summary>
+			/// Completion source for the waiting agent. If a new queue item becomes available, the result will be passed through 
+			/// </summary>
+			public TaskCompletionSource<AgentLease?> LeaseSource { get; } = new TaskCompletionSource<AgentLease?>();
+
+			/// <summary>
+			/// Constructor
+			/// </summary>
+			/// <param name="agent">The agent waiting for a task</param>
+			public QueueWaiter(IAgent agent)
+			{
+				Agent = agent;
+			}
+		}
+
+		readonly MongoService _mongoService;
+		readonly StreamService _streamService;
+		readonly ILogFileService _logFileService;
+		readonly IAgentCollection _agentsCollection;
+		readonly IJobCollection _jobs;
+		readonly IJobStepRefCollection _jobStepRefs;
+		readonly IGraphCollection _graphs;
+		readonly IPoolCollection _poolCollection;
+		readonly IUgsMetadataCollection _ugsMetadataCollection;
+		readonly PerforceLoadBalancer _perforceLoadBalancer;
+		readonly IOptionsMonitor<ServerSettings> _settings;
+		readonly ICommitService _commitService;
+		readonly ILogger<JobTaskSource> _logger;
+		readonly ITicker _ticker;
+
+		// Object used for ensuring mutual exclusion to the queues
+		readonly object _lockObject = new object();
+
+		// List of items waiting to be executed
+		SortedSet<QueueItem> _queue = new SortedSet<QueueItem>(new QueueItemComparer());
+
+		// Map from batch id to the corresponding queue item
+		Dictionary<(JobId, SubResourceId), QueueItem> _batchIdToQueueItem = new Dictionary<(JobId, SubResourceId), QueueItem>();
+
+		// Set of long-poll tasks waiting to be satisfied 
+		readonly HashSet<QueueWaiter> _waiters = new HashSet<QueueWaiter>();
+
+		// During a background queue refresh operation, any updated batches are added to this dictionary for merging into the updated queue.
+		List<QueueItem>? _newQueueItemsDuringUpdate;
+
+		// Cache of pools
+		Dictionary<PoolId, IPool> _cachedPoolIdToInstance = new Dictionary<PoolId, IPool>();
+
+		// Cache of stream objects. Used to resolve agent types.
+		private Dictionary<StreamId, IStream> _streams = new Dictionary<StreamId, IStream>();
+
+		/// <summary>
+		/// Delegate for job schedule events
+		/// </summary>
+		public delegate void JobScheduleEvent(IPool pool, bool hasAgentsOnline, IJob job, IGraph graph, SubResourceId batchId);
+		
+		/// <summary>
+		/// Event triggered when a job is scheduled
+		/// </summary>
+		public event JobScheduleEvent? OnJobScheduled;
+
+		// Interval between querying the database for jobs to execute
+		static readonly TimeSpan s_refreshInterval = TimeSpan.FromSeconds(5.0);
+
+		/// <summary>
+		/// Constructor
+		/// </summary>
+		public JobTaskSource(MongoService mongoService, IAgentCollection agents, IJobCollection jobs, IJobStepRefCollection jobStepRefs, IGraphCollection graphs, IPoolCollection pools, IUgsMetadataCollection ugsMetadataCollection, StreamService streamService, ILogFileService logFileService, PerforceLoadBalancer perforceLoadBalancer, ICommitService commitService, IClock clock, IOptionsMonitor<ServerSettings> settings, ILogger<JobTaskSource> logger)
+		{
+			_mongoService = mongoService;
+			_agentsCollection = agents;
+			_jobs = jobs;
+			_jobStepRefs = jobStepRefs;
+			_graphs = graphs;
+			_poolCollection = pools;
+			_ugsMetadataCollection = ugsMetadataCollection;
+			_streamService = streamService;
+			_logFileService = logFileService;
+			_perforceLoadBalancer = perforceLoadBalancer;
+			_ticker = clock.AddTicker<JobTaskSource>(s_refreshInterval, TickAsync, logger);
+			_settings = settings;
+			_commitService = commitService;
+			_logger = logger;
+
+			OnLeaseStartedProperties.Add(nameof(ExecuteJobTask.JobId), x => new JobId(x.JobId)).Add(nameof(ExecuteJobTask.BatchId), x => SubResourceId.Parse(x.BatchId)).Add(nameof(ExecuteJobTask.LogId), x => new LogId(x.LogId));
+		}
+
+		/// <inheritdoc/>
+		public Task StartAsync(CancellationToken cancellationToken) => _ticker.StartAsync();
+
+		/// <inheritdoc/>
+		public Task StopAsync(CancellationToken cancellationToken) => _ticker.StopAsync();
+
+		/// <inheritdoc/>
+		public void Dispose() => _ticker.Dispose();
+
+		/// <summary>
+		/// Gets an object containing the stats of the queue for diagnostic purposes.
+		/// </summary>
+		/// <returns>Status object</returns>
+		public object GetStatus()
+		{
+			lock (_lockObject)
+			{
+				List<object> outputItems = new List<object>();
+				foreach (QueueItem queueItem in _queue)
+				{
+					outputItems.Add(new { JobId = queueItem._job.Id.ToString(), BatchId = queueItem.Batch.Id.ToString(), PoolId = queueItem._poolId.ToString(), Workspace = queueItem._workspace });
+				}
+
+				List<object> outputWaiters = new List<object>();
+				foreach (QueueWaiter waiter in _waiters)
+				{
+					outputWaiters.Add(new { Id = waiter.Agent.Id.ToString(), Pools = waiter.Agent.GetPools().Select(x => x.ToString()).ToList(), waiter.Agent.Workspaces });
+				}
+
+				return new { Items = outputItems, Waiters = outputWaiters };
+			}
+		}
+
+		/// <summary>
+		/// Cancel any pending wait for an agent, allowing it to cycle its session state immediately
+		/// </summary>
+		/// <param name="agentId">The agent id</param>
+		public void CancelLongPollForAgent(AgentId agentId)
+		{
+			QueueWaiter? waiter;
+			lock (_lockObject)
+			{
+				waiter = _waiters.FirstOrDefault(x => x.Agent.Id == agentId);
+			}
+			if(waiter != null)
+			{
+				waiter.LeaseSource.TrySetCanceled();
+			}
+		}
+
+		/// <summary>
+		/// Background task
+		/// </summary>
+		/// <param name="stoppingToken">Token that indicates that the service should shut down</param>
+		/// <returns>Async task</returns>
+		internal async ValueTask TickAsync(CancellationToken stoppingToken)
+		{
+			// Set the NewBatchIdToQueueItem member, so we capture any updated jobs during the DB query.
+			lock (_lockObject)
+			{
+				_newQueueItemsDuringUpdate = new List<QueueItem>();
+			}
+
+			// Query all the current streams
+			List<IStream> streamsList = await _streamService.GetStreamsAsync();
+			_streams = streamsList.ToDictionary(x => x.Id, x => x);
+
+			// Find all the pools which are valid (ie. have at least one online agent)
+			DateTime utcNow = DateTime.UtcNow;
+			List<IAgent> agents = await _agentsCollection.FindAsync();
+			List<IPool> pools = await _poolCollection.GetAsync();
+
+			// Find all the pools which are currently online
+			HashSet<PoolId> onlinePools = new HashSet<PoolId>(agents.Where(x => x.IsSessionValid(utcNow) && x.Enabled).SelectMany(x => x.ExplicitPools));
+			foreach (IPool pool in pools)
+			{
+				if (pool.Condition != null && !onlinePools.Contains(pool.Id) && agents.Any(x => x.IsSessionValid(utcNow) && x.SatisfiesCondition(pool.Condition) && x.Enabled))
+				{
+					onlinePools.Add(pool.Id);
+				}
+			}
+
+			// Find lists of valid pools and online pools
+			HashSet<PoolId> validPools = new HashSet<PoolId>(onlinePools.Union(agents.Where(x => !x.IsSessionValid(utcNow) || !x.Enabled).SelectMany(x => x.ExplicitPools)));
+			foreach (IPool pool in pools)
+			{
+				if (pool.Condition != null && !validPools.Contains(pool.Id) && agents.Any(x => !x.IsSessionValid(utcNow) && x.SatisfiesCondition(pool.Condition) && x.Enabled))
+				{
+					validPools.Add(pool.Id);
+				}
+			}
+
+			// Query all the current pools
+			_cachedPoolIdToInstance = pools.ToDictionary(x => x.Id, x => x);
+
+			// New list of queue items
+			SortedSet<QueueItem> newQueue = new SortedSet<QueueItem>(_queue.Comparer);
+			Dictionary<(JobId, SubResourceId), QueueItem> newBatchIdToQueueItem = new Dictionary<(JobId, SubResourceId), QueueItem>();
+			
+			// Returns true if agents are online and available for scheduling for a pool
+			bool IsPoolOnline(PoolId poolId)
+			{
+				return onlinePools.Contains(poolId);
+			}
+			
+			// Returns true if a pool can be auto-scaled
+			bool IsPoolAutoScaled(PoolId poolId)
+			{
+				IPool? pool = pools.Find(p => p.Id == poolId);
+				return validPools.Contains(poolId) && pool != null && pool.EnableAutoscaling;
+			}
+			
+			bool HasAgentsOnlineOrIsAutoScaled(PoolId poolId)
+			{
+				// If pool is auto-scaled, it will be considered online even if it has no agents online
+				return IsPoolOnline(poolId) || IsPoolAutoScaled(poolId);
+			}
+
+			// Query for a new list of jobs for the queue
+			List<IJob> newJobs = await _jobs.GetDispatchQueueAsync();
+			for (int idx = 0; idx < newJobs.Count; idx++)
+			{
+				IJob? newJob = newJobs[idx];
+
+				if (newJob.GraphHash == null)
+				{
+					_logger.LogError("Job {JobId} has a null graph hash and can't be started.", newJob.Id);
+					await _jobs.TryRemoveFromDispatchQueueAsync(newJob);
+					continue;
+				}
+				if (newJob.AbortedByUserId != null)
+				{
+					_logger.LogError("Job {JobId} was aborted but not removed from dispatch queue", newJob.Id);
+					await _jobs.TryRemoveFromDispatchQueueAsync(newJob);
+					continue;
+				}
+
+				// Get the graph for this job
+				IGraph graph = await _graphs.GetAsync(newJob.GraphHash);
+
+				// Get the stream. If it fails, skip the whole job.
+				IStream? stream;
+				if (!_streams.TryGetValue(newJob.StreamId, out stream))
+				{
+					newJob = await _jobs.SkipAllBatchesAsync(newJob, graph, JobStepBatchError.UnknownStream);
+					continue;
+				}
+
+				// Update all the batches
+				HashSet<SubResourceId> checkedBatchIds = new HashSet<SubResourceId>();
+				while (newJob != null)
+				{
+					// Find the next batch within this job that is ready
+					int batchIdx = newJob.Batches.FindIndex(x => x.State == JobStepBatchState.Ready && checkedBatchIds.Add(x.Id));
+					if (batchIdx == -1)
+					{
+						break;
+					}
+
+					// Validate the agent type and workspace settings
+					IJobStepBatch batch = newJob.Batches[batchIdx];
+					if (!stream.AgentTypes.TryGetValue(graph.Groups[batch.GroupIdx].AgentType, out AgentType? agentType))
+					{
+						newJob = await SkipBatchAsync(newJob, batch.Id, graph, JobStepBatchError.UnknownAgentType);
+					}
+					else if (!_cachedPoolIdToInstance.TryGetValue(agentType.Pool, out IPool? pool))
+					{
+						newJob = await SkipBatchAsync(newJob, batch.Id, graph, JobStepBatchError.UnknownPool);
+					}
+					else if (!validPools.Contains(agentType.Pool))
+					{
+						newJob = await SkipBatchAsync(newJob, batch.Id, graph, JobStepBatchError.NoAgentsInPool);
+					}
+					else if (!HasAgentsOnlineOrIsAutoScaled(agentType.Pool))
+					{
+						newJob = await SkipBatchAsync(newJob, batch.Id, graph, JobStepBatchError.NoAgentsOnline);
+					}
+					else if (!stream.TryGetAgentWorkspace(agentType, out (AgentWorkspace, bool)? workspaceResult))
+					{
+						newJob = await SkipBatchAsync(newJob, batch.Id, graph, JobStepBatchError.UnknownWorkspace);
+					}
+					else
+					{
+						TemplateRef? templateRef;
+						if (stream.Templates.TryGetValue(newJob.TemplateId, out templateRef))
+						{							
+							if (templateRef.StepStates != null)
+							{
+								for (int i = 0; i < templateRef.StepStates.Count; i++)
+								{
+									TemplateStepState state = templateRef.StepStates[i];
+									if (state.PausedByUserId != null)
+									{
+										IJobStep? step = batch.Steps.FirstOrDefault(x => graph.Groups[batch.GroupIdx].Nodes[x.NodeIdx].Name.Equals(state.Name, StringComparison.Ordinal));
+
+										if (step != null)
+										{
+											JobId jobId = newJob.Id;
+											newJob = await _jobs.TryUpdateStepAsync(newJob, graph, batch.Id, step.Id, JobStepState.Skipped, newError: JobStepError.Paused);
+											if (newJob == null)
+											{
+												_logger.LogError("Job {JobId} failed to update step {StepName} pause state", jobId, state.Name);
+												break;
+											}
+											else
+											{
+												_logger.LogInformation("Job {JobId} step {StepName} has been skipped due to being paused", jobId, state.Name);
+											}
+										}
+									}
+								}
+							}
+						}
+
+						if (newJob != null)
+						{
+							(AgentWorkspace workspace, bool useAutoSdk) = workspaceResult.Value;
+							QueueItem newQueueItem = new QueueItem(stream, newJob, batchIdx, agentType.Pool, workspace, useAutoSdk);
+							newQueue.Add(newQueueItem);
+							newBatchIdToQueueItem[(newJob.Id, batch.Id)] = newQueueItem;
+
+							IPool? newJobPool = pools.Find(p => p.Id == agentType.Pool);
+							if (newJobPool != null)
+							{
+								OnJobScheduled?.Invoke(newJobPool, IsPoolOnline(agentType.Pool), newJob, graph, batch.Id);
+							}
+						}
+					}
+				}
+
+				if (newJob != null)
+				{
+					if (!newJob.Batches.Any(batch => batch.State == JobStepBatchState.Ready || batch.State == JobStepBatchState.Starting || batch.State == JobStepBatchState.Running || batch.State == JobStepBatchState.Stopping))
+					{
+						_logger.LogError("Job {JobId} is in dispatch queue but not currently executing", newJob.Id);
+						await _jobs.TryRemoveFromDispatchQueueAsync(newJob);
+					}
+				}
+			}
+
+			// Update the queue
+			lock (_lockObject)
+			{
+				_queue = newQueue;
+				_batchIdToQueueItem = newBatchIdToQueueItem;
+
+				// Merge the new queue items with the queue
+				foreach (QueueItem newQueueItem in _newQueueItemsDuringUpdate)
+				{
+					QueueItem? existingQueueItem;
+					if (!newBatchIdToQueueItem.TryGetValue((newQueueItem._job.Id, newQueueItem.Batch.Id), out existingQueueItem))
+					{
+						// Always just add this item
+						_queue.Add(newQueueItem);
+						_batchIdToQueueItem[newQueueItem.Id] = newQueueItem;
+					}
+					else if (newQueueItem._job.UpdateIndex > existingQueueItem._job.UpdateIndex)
+					{
+						// Replace the existing item
+						_queue.Remove(existingQueueItem);
+						_queue.Add(newQueueItem);
+						_batchIdToQueueItem[newQueueItem.Id] = newQueueItem;
+					}
+				}
+
+				// Clear out the list to capture queue items during an update
+				_newQueueItemsDuringUpdate = null;
+			}
+		}
+
+		private async Task<IJob?> SkipBatchAsync(IJob job, SubResourceId batchId, IGraph graph, JobStepBatchError reason)
+		{
+			_logger.LogInformation("Skipping batch {BatchId} for job {JobId} (reason: {Reason})", batchId, job.Id, reason);
+
+			IReadOnlyList<(LabelState, LabelOutcome)> oldLabelStates = job.GetLabelStates(graph);
+			IJob? newJob = await _jobs.SkipBatchAsync(job, batchId, graph, reason);
+			if(newJob != null)
+			{
+				IReadOnlyList<(LabelState, LabelOutcome)> newLabelStates = newJob.GetLabelStates(graph);
+				await UpdateUgsBadges(newJob, graph, oldLabelStates, newLabelStates);
+			}
+			return newJob;
+		}
+
+		/// <summary>
+		/// Get the queue, for internal testing only
+		/// </summary>
+		/// <returns>A copy of the queue</returns>
+		internal SortedSet<QueueItem> GetQueueForTesting()
+		{
+			lock (_lockObject)
+			{
+				return new SortedSet<QueueItem>(_queue);
+			}
+		}
+
+		/// <summary>
+		/// Updates the current state of a job
+		/// </summary>
+		/// <param name="job">The job that has been updated</param>
+		/// <param name="graph">Graph for the job</param>
+		/// <returns>Async task</returns>
+		public void UpdateQueuedJob(IJob job, IGraph graph)
+		{
+			IStream? stream;
+			_streams.TryGetValue(job.StreamId, out stream);
+			UpdateQueuedJob(job, graph, stream);
+		}
+
+		void AssignAnyQueueItemToWaiter(QueueWaiter waiter)
+		{
+			lock (_waiters)
+			{
+				foreach(QueueItem item in _batchIdToQueueItem.Values)
+				{
+					if (TryAssignItemToWaiter(item, waiter))
+					{
+						break;
+					}
+				}
+			}
+		}
+
+		/// <summary>
+		/// Attempt to find a waiter that can handle the given queue item
+		/// </summary>
+		/// <param name="item">The queue item</param>
+		/// <returns></returns>
+		[SuppressMessage("Maintainability", "CA1508:Avoid dead conditional code", Justification = "<Pending>")]
+		void AssignQueueItemToAnyWaiter(QueueItem item)
+		{
+			if (item._assignTask == null && item.Batch.SessionId == null)
+			{
+				lock (_waiters)
+				{
+					if (item._assignTask == null && item.Batch.SessionId == null)
+					{
+						foreach (QueueWaiter waiter in _waiters)
+						{
+							if (TryAssignItemToWaiter(item, waiter))
+							{
+								break;
+							}
+						}
+					}
+				}
+			}
+		}
+
+		/// <summary>
+		/// Attempts to assign a queue item to an agent waiting for work
+		/// </summary>
+		/// <param name="item"></param>
+		/// <param name="waiter"></param>
+		/// <returns></returns>
+		bool TryAssignItemToWaiter(QueueItem item, QueueWaiter waiter)
+		{
+			if (item._assignTask == null && item.Batch.SessionId == null && waiter.Agent.Enabled && waiter.Agent.Leases.Count == 0 && waiter.Agent.IsInPool(item._poolId))
+			{
+				Task startTask = new Task<Task>(() => TryCreateLeaseAsync(item, waiter));
+				Task executeTask = startTask.ContinueWith(task => task, TaskScheduler.Default);
+
+				if (Interlocked.CompareExchange(ref item._assignTask, executeTask, null) == null)
+				{
+					startTask.Start(TaskScheduler.Default);
+					return true;
+				}
+			}
+			return false;
+		}
+
+		/// <summary>
+		/// Updates the current state of a job
+		/// </summary>
+		/// <param name="job">The job that has been updated</param>
+		/// <param name="graph">Graph for the job</param>
+		/// <param name="stream">The stream containing the job</param>
+		public void UpdateQueuedJob(IJob job, IGraph graph, IStream? stream)
+		{
+			List<TaskCompletionSource<bool>> completeWaiters = new List<TaskCompletionSource<bool>>();
+			lock (_lockObject)
+			{
+				for (int batchIdx = 0; batchIdx < job.Batches.Count; batchIdx++)
+				{
+					IJobStepBatch batch = job.Batches[batchIdx];
+					if (batch.State == JobStepBatchState.Ready && stream != null && batch.AgentId == null)
+					{
+						// Check if this item is already in the list.
+						QueueItem? existingItem;
+						if (_batchIdToQueueItem.TryGetValue((job.Id, batch.Id), out existingItem))
+						{
+							// Make sure this is a newer version of the job. There's no guarantee that this is the latest revision.
+							if (job.UpdateIndex > existingItem._job.UpdateIndex)
+							{
+								if (batch.SchedulePriority == existingItem.Batch.SchedulePriority)
+								{
+									existingItem._job = job;
+									existingItem._batchIdx = batchIdx;
+								}
+								else
+								{
+									RemoveQueueItem(existingItem);
+									InsertQueueItem(stream, job, batchIdx, existingItem._poolId, existingItem._workspace, existingItem._useAutoSdk);
+								}
+							}
+							continue;
+						}
+
+						// Get the group being executed by this batch
+						INodeGroup group = graph.Groups[batch.GroupIdx];
+
+						// Get the requirements for the new queue item
+						AgentType? agentType;
+						if (stream.AgentTypes.TryGetValue(group.AgentType, out agentType))
+						{
+							(AgentWorkspace, bool)? result;
+							if (stream.TryGetAgentWorkspace(agentType, out result))
+							{
+								(AgentWorkspace agentWorkspace, bool useAutoSdk) = result.Value;
+								InsertQueueItem(stream, job, batchIdx, agentType.Pool, agentWorkspace, useAutoSdk);
+							}
+						}
+					}
+					else
+					{
+						// Check if this item is already in the list. Remove it if it is.
+						QueueItem? existingItem;
+						if (_batchIdToQueueItem.TryGetValue((job.Id, batch.Id), out existingItem))
+						{
+							if (job.UpdateIndex > existingItem._job.UpdateIndex)
+							{
+								RemoveQueueItem(existingItem);
+							}
+						}
+					}
+				}
+			}
+
+			// Awake all the threads that have been assigned new work items. Has do be done outside the lock to prevent continuations running within it (see Waiter.CompletionSource for more info).
+			foreach (TaskCompletionSource<bool> completeWaiter in completeWaiters)
+			{
+				completeWaiter.TrySetResult(true);
+			}
+		}
+
+		/// <inheritdoc/>
+		public override Task<Task<AgentLease?>> AssignLeaseAsync(IAgent agent, CancellationToken cancellationToken)
+		{
+			QueueWaiter waiter = new QueueWaiter(agent);
+			lock (_lockObject)
+			{
+				AssignAnyQueueItemToWaiter(waiter);
+				if (waiter.LeaseSource.Task.TryGetResult(out AgentLease? result))
+				{
+					if (result == null)
+					{
+						return Task.FromResult(Skip(cancellationToken));
+					}
+					return Task.FromResult(Lease(result));
+				}
+				_waiters.Add(waiter);
+			}
+			return Task.FromResult(WaitForLeaseAsync(waiter, cancellationToken));
+		}
+
+		private async Task<AgentLease?> WaitForLeaseAsync(QueueWaiter waiter, CancellationToken cancellationToken)
+		{
+			try
+			{
+				using (cancellationToken.Register(() => waiter.LeaseSource.TrySetResult(null)))
+				{
+					return await waiter.Task;
+				}
+			}
+			finally
+			{
+				lock (_lockObject)
+				{
+					_waiters.Remove(waiter);
+				}
+			}
+		}
+
+		/// <inheritdoc/>
+		public override Task CancelLeaseAsync(IAgent agent, LeaseId leaseId, ExecuteJobTask task)
+		{
+			return CancelLeaseAsync(agent, new JobId(task.JobId), task.BatchId.ToSubResourceId());
+		}
+
+		/// <summary>
+		/// Assign a new batch to be executed by the given agent
+		/// </summary>
+		/// <param name="item">The item to create a lease for</param>
+		/// <param name="waiter">The agent waiting for work</param>
+		/// <returns>New work to execute</returns>
+		private async Task<AgentLease?> TryCreateLeaseAsync(QueueItem item, QueueWaiter waiter)
+		{
+			IJob job = item._job;
+			IJobStepBatch batch = item.Batch;
+			IAgent agent = waiter.Agent;
+			_logger.LogDebug("Assigning job {JobId}, batch {BatchId} to waiter (agent {AgentId})", job.Id, batch.Id, agent.Id);
+
+			// Generate a new unique id for the lease
+			LeaseId leaseId = LeaseId.GenerateNewId();
+
+			// The next time to try assigning to another agent
+			DateTime backOffTime = DateTime.UtcNow + TimeSpan.FromMinutes(1.0);
+
+			// Try to update the job with this agent id
+			LogId logId = (await _logFileService.CreateLogFileAsync(job.Id, agent.SessionId, LogType.Json)).Id;
+			IJob? newJob = await _jobs.TryAssignLeaseAsync(item._job, item._batchIdx, item._poolId, agent.Id, agent.SessionId!.Value, leaseId, logId);
+			if (newJob != null)
+			{
+				job = newJob;
+
+				// Get the lease name
+				StringBuilder leaseName = new StringBuilder($"{item._stream.Name} - ");
+				if (job.PreflightChange > 0)
+				{
+					leaseName.Append((job.Change > 0) ? $"Preflight CL {job.PreflightChange} against CL {job.Change}" : $"Preflight CL {job.PreflightChange} against latest");
+				}
+				else
+				{
+					leaseName.Append((job.Change > 0) ? $"CL {job.Change}" : "Latest CL");
+				}
+				leaseName.Append(CultureInfo.InvariantCulture, $" - {job.Name}");
+
+				// Get the global settings
+				Globals globals = await _mongoService.GetGlobalsAsync();
+
+				// Encode the payload
+				ExecuteJobTask? task = await CreateExecuteJobTaskAsync(item._stream, job, batch, agent, item._workspace, item._useAutoSdk, logId);
+				if (task != null)
+				{
+					byte[] payload = Any.Pack(task).ToByteArray();
+
+					// Create the lease and try to set it on the waiter. If this fails, the waiter has already moved on, and the lease can be cancelled.
+					AgentLease lease = new AgentLease(leaseId, leaseName.ToString(), job.StreamId, item._poolId, logId, LeaseState.Pending, null, true, payload);
+					if (waiter.LeaseSource.TrySetResult(lease))
+					{
+						_logger.LogDebug("Assigned lease {LeaseId} to agent {AgentId}", leaseId, agent.Id);
+						return lease;
+					}
+				}
+
+				// Cancel the lease
+				_logger.LogDebug("Unable to assign lease {LeaseId} to agent {AgentId}, cancelling", leaseId, agent.Id);
+				await CancelLeaseAsync(waiter.Agent, job.Id, batch.Id);
+			}
+			else
+			{
+				// Unable to assign job
+				_logger.LogDebug("Failed to assign job {JobId}, batch {BatchId} to agent {AgentId}. Refreshing queue entries.", job.Id, batch.Id, agent.Id);
+
+				// Get the new copy of the job
+				newJob = await _jobs.GetAsync(job.Id);
+				if (newJob == null)
+				{
+					lock (_lockObject)
+					{
+						List<QueueItem> removeItems = _queue.Where(x => x._job == job).ToList();
+						foreach (QueueItem removeItem in removeItems)
+						{
+							RemoveQueueItem(removeItem);
+						}
+					}
+				}
+				else
+				{
+					IGraph graph = await _graphs.GetAsync(newJob.GraphHash);
+					UpdateQueuedJob(newJob, graph);
+				}
+			}
+
+			// Clear out the assignment for this item, and try to reassign it
+			item._assignTask = null;
+			AssignQueueItemToAnyWaiter(item);
+			return null;
+		}
+
+		async Task<ExecuteJobTask?> CreateExecuteJobTaskAsync(IStream stream, IJob job, IJobStepBatch batch, IAgent agent, AgentWorkspace workspace, bool bUseAutoSdk, LogId logId)
+		{
+			// Get the lease name
+			StringBuilder leaseName = new StringBuilder($"{stream.Name} - ");
+			if (job.PreflightChange > 0)
+			{
+				leaseName.Append((job.Change > 0) ? $"Preflight CL {job.PreflightChange} against CL {job.Change}" : $"Preflight CL {job.PreflightChange} against latest");
+			}
+			else
+			{
+				leaseName.Append((job.Change > 0) ? $"CL {job.Change}" : "Latest CL");
+			}
+			leaseName.Append(CultureInfo.InvariantCulture, $" - {job.Name}");
+
+			// Get the global settings
+			Globals globals = await _mongoService.GetGlobalsAsync();
+
+			// Encode the payload
+			ExecuteJobTask task = new ExecuteJobTask();
+			task.JobId = job.Id.ToString();
+			task.BatchId = batch.Id.ToString();
+			task.LogId = logId.ToString();
+			task.JobName = leaseName.ToString();
+
+			List<HordeCommon.Rpc.Messages.AgentWorkspace> workspaces = new List<HordeCommon.Rpc.Messages.AgentWorkspace>();
+
+			PerforceCluster? cluster = globals.FindPerforceCluster(workspace.Cluster);
+			if (cluster == null)
+			{
+				return null;
+			}
+
+			AgentWorkspace? autoSdkWorkspace = bUseAutoSdk ? agent.GetAutoSdkWorkspace(cluster) : null;
+			if (autoSdkWorkspace != null)
+			{
+				if (!await agent.TryAddWorkspaceMessage(autoSdkWorkspace, cluster, _perforceLoadBalancer, workspaces))
+				{
+					return null;
+				}
+				task.AutoSdkWorkspace = workspaces.Last();
+			}
+
+			if (!await agent.TryAddWorkspaceMessage(workspace, cluster, _perforceLoadBalancer, workspaces))
+			{
+				return null;
+			}
+
+			task.Workspace = workspaces.Last();
+
+			IStream? Stream;
+			if (_streams.TryGetValue(job.StreamId, out Stream))
+			{
+				task.ContentRef = _commitService.GetReplicatedContentRef(Stream, job.Change);
+			}
+
+			return task;
+		}
+
+		/// <summary>
+		/// Send any badge updates for this job
+		/// </summary>
+		/// <param name="job">The job being updated</param>
+		/// <param name="graph">Graph for the job</param>
+		/// <param name="oldLabelStates">Previous badge states for the job</param>
+		/// <returns>Async task</returns>
+		public async Task UpdateUgsBadges(IJob job, IGraph graph, IReadOnlyList<(LabelState, LabelOutcome)> oldLabelStates)
+		{
+			await UpdateUgsBadges(job, graph, oldLabelStates, job.GetLabelStates(graph));
+		}
+
+		/// <summary>
+		/// Send any badge updates for this job
+		/// </summary>
+		/// <param name="job">The job being updated</param>
+		/// <param name="graph">Graph for the job</param>
+		/// <param name="oldLabelStates">Previous badge states for the job</param>
+		/// <param name="newLabelStates">The new badge states for the job</param>
+		/// <returns>Async task</returns>
+		public async Task UpdateUgsBadges(IJob job, IGraph graph, IReadOnlyList<(LabelState, LabelOutcome)> oldLabelStates, IReadOnlyList<(LabelState, LabelOutcome)> newLabelStates)
+		{
+			if (!job.ShowUgsBadges || job.PreflightChange != 0)
+			{
+				return;
+			}
+
+			IReadOnlyDictionary<int, UgsBadgeState> oldStates = job.GetUgsBadgeStates(graph, oldLabelStates);
+			IReadOnlyDictionary<int, UgsBadgeState> newStates = job.GetUgsBadgeStates(graph, newLabelStates);
+
+			// Figure out a list of all the badges that have been modified
+			List<int> updateLabels = new List<int>();
+			foreach (KeyValuePair<int, UgsBadgeState> pair in oldStates)
+			{
+				if (!newStates.ContainsKey(pair.Key))
+				{
+					updateLabels.Add(pair.Key);
+				}
+			}
+			foreach (KeyValuePair<int, UgsBadgeState> pair in newStates)
+			{
+				if (!oldStates.TryGetValue(pair.Key, out UgsBadgeState prevState) || prevState != pair.Value)
+				{
+					updateLabels.Add(pair.Key);
+				}
+			}
+
+			// Cached stream for this job
+			IStream? stream = null;
+
+			// Send all the updates
+			Dictionary<int, IUgsMetadata> metadataCache = new Dictionary<int, IUgsMetadata>();
+			foreach (int labelIdx in updateLabels)
+			{
+				ILabel label = graph.Labels[labelIdx];
+
+				// Skip if this label has no UGS name.
+				if (label.UgsName == null)
+				{
+					continue;
+				}
+
+				// Get the new state
+				if (!newStates.TryGetValue(labelIdx, out UgsBadgeState newState))
+				{
+					newState = UgsBadgeState.Skipped;
+				}
+
+				// Get the stream
+				if (stream == null)
+				{
+					stream = await _streamService.GetStreamAsync(job.StreamId);
+					if (stream == null)
+					{
+						_logger.LogError("Unable to fetch definition for stream {StreamId}", job.StreamId);
+						break;
+					}
+				}
+
+				// The changelist number to display the badge for
+				int change;
+				if (label.Change == LabelChange.Code)
+				{
+					change = job.CodeChange;
+				}
+				else
+				{
+					change = job.Change;
+				}
+
+				// Get the current metadata state
+				IUgsMetadata? metadata;
+				if (!metadataCache.TryGetValue(change, out metadata))
+				{
+					metadata = await _ugsMetadataCollection.FindOrAddAsync(stream.Name, change, label.UgsProject);
+					metadataCache[change] = metadata;
+				}
+
+				// Apply the update
+				Uri labelUrl = new Uri(_settings.CurrentValue.DashboardUrl, $"job/{job.Id}?label={labelIdx}");
+				_logger.LogInformation("Updating state of badge {BadgeName} at {Change} to {NewState} ({LabelUrl})", label.UgsName, change, newState, labelUrl);
+				metadata = await _ugsMetadataCollection.UpdateBadgeAsync(metadata, label.UgsName!, labelUrl, newState);
+				metadataCache[change] = metadata;
+			}
+		}
+
+		/// <inheritdoc/>
+		public override async Task OnLeaseFinishedAsync(IAgent agent, LeaseId leaseId, ExecuteJobTask task, LeaseOutcome outcome, ReadOnlyMemory<byte> output, ILogger logger)
+		{
+			await base.OnLeaseFinishedAsync(agent, leaseId, task, outcome, output, logger);
+
+			if (outcome != LeaseOutcome.Success)
+			{
+				AgentId agentId = agent.Id;
+				JobId jobId = new JobId(task.JobId);
+				SubResourceId batchId = task.BatchId.ToSubResourceId();
+
+				// Update the batch
+				for (; ; )
+				{
+					IJob? job = await _jobs.GetAsync(jobId);
+					if (job == null)
+					{
+						break;
+					}
+
+					int batchIdx = job.Batches.FindIndex(x => x.Id == batchId);
+					if (batchIdx == -1)
+					{
+						break;
+					}
+
+					IJobStepBatch batch = job.Batches[batchIdx];
+					if (batch.AgentId != agentId)
+					{
+						break;
+					}
+
+					int runningStepIdx = batch.Steps.FindIndex(x => x.State == JobStepState.Running);
+
+					JobStepBatchError error;
+					if (outcome == LeaseOutcome.Cancelled)
+					{
+						error = JobStepBatchError.Cancelled;
+					}
+					else
+					{
+						error = JobStepBatchError.ExecutionError;
+					}
+
+					IGraph graph = await _graphs.GetAsync(job.GraphHash);
+					job = await _jobs.TryFailBatchAsync(job, batchIdx, graph, error);
+
+					if (job != null)
+					{
+						if (batch.Error != JobStepBatchError.None)
+						{
+							logger.LogInformation("Failed job {JobId}, batch {BatchId} with error {Error}", job.Id, batch.Id, batch.Error);
+						}
+						if (runningStepIdx != -1)
+						{
+							await _jobStepRefs.UpdateAsync(job, batch, batch.Steps[runningStepIdx], graph);
+						}
+						break;
+					}
+				}
+			}
+		}
+
+		/// <summary>
+		/// 
+		/// </summary>
+		/// <param name="agent"></param>
+		/// <param name="jobId"></param>
+		/// <param name="batchId"></param>
+		/// <returns></returns>
+		async Task CancelLeaseAsync(IAgent agent, JobId jobId, SubResourceId batchId)
+		{
+			_logger.LogDebug("Cancelling lease for job {JobId}, batch {BatchId}", jobId, batchId);
+
+			// Update the batch
+			for (; ; )
+			{
+				IJob? job = await _jobs.GetAsync(jobId);
+				if (job == null)
+				{
+					break;
+				}
+
+				int batchIdx = job.Batches.FindIndex(x => x.Id == batchId);
+				if (batchIdx == -1)
+				{
+					break;
+				}
+
+				IJobStepBatch batch = job.Batches[batchIdx];
+				if (batch.AgentId != agent.Id)
+				{
+					break;
+				}
+
+				IJob? newJob = await _jobs.TryCancelLeaseAsync(job, batchIdx);
+				if (newJob != null)
+				{
+					break;
+				}
+			}
+		}
+
+		/// <summary>
+		/// Inserts an item into the queue
+		/// </summary>
+		/// <param name="stream">The stream containing the job</param>
+		/// <param name="job"></param>
+		/// <param name="batchIdx"></param>
+		/// <param name="poolId">The pool to use</param>
+		/// <param name="workspace">The workspace for this item to run in</param>
+		/// <param name="useAutoSdk">Whether or not to use the AutoSDK</param>
+		/// <returns></returns>
+		void InsertQueueItem(IStream stream, IJob job, int batchIdx, PoolId poolId, AgentWorkspace workspace, bool useAutoSdk)
+		{
+			_logger.LogDebug("Adding queued job {JobId}, batch {BatchId} [Pool: {Pool}, Workspace: {Workspace}]", job.Id, job.Batches[batchIdx].Id, poolId, workspace.Identifier);
+
+			QueueItem newItem = new QueueItem(stream, job, batchIdx, poolId, workspace, useAutoSdk);
+			_batchIdToQueueItem[newItem.Id] = newItem;
+			_queue.Add(newItem);
+
+			AssignQueueItemToAnyWaiter(newItem);
+		}
+
+		/// <summary>
+		/// Removes an item from the queue
+		/// </summary>
+		/// <param name="item">Item to remove</param>
+		void RemoveQueueItem(QueueItem item)
+		{
+			_logger.LogDebug("Removing queued job {JobId}, batch {BatchId}", item._job.Id, item.Batch.Id);
+
+			_queue.Remove(item);
+			_batchIdToQueueItem.Remove(item.Id);
+		}
+	}
+}
diff --git a/Engine/Source/Programs/Horde/Horde.Build/Jobs/JobsController.Messages.cs b/Engine/Source/Programs/Horde/Horde.Build/Jobs/JobsController.Messages.cs
new file mode 100644
--- /dev/null
+++ b/Engine/Source/Programs/Horde/Horde.Build/Jobs/JobsController.Messages.cs
@@ -0,0 +1,1135 @@
+// Copyright Epic Games, Inc. All Rights Reserved.
+
+using System;
+using System.Collections.Generic;
+using System.ComponentModel.DataAnnotations;
+using System.Linq;
+using System.Text.Json.Serialization;
+using EpicGames.Core;
+using Horde.Build.Models;
+using HordeCommon;
+
+namespace Horde.Build.Api
+{
+	/// <summary>
+	/// State of the job
+	/// </summary>
+	public enum JobState
+	{
+		/// <summary>
+		/// Waiting for resources
+		/// </summary>
+		Waiting,
+
+		/// <summary>
+		/// Currently running one or more steps
+		/// </summary>
+		Running,
+
+		/// <summary>
+		/// All steps have completed
+		/// </summary>
+		Complete,
+	}
+
+	/// <summary>
+	/// Parameters required to create a job
+	/// </summary>
+	public class CreateJobRequest
+	{
+		/// <summary>
+		/// The stream that this job belongs to
+		/// </summary>
+		[Required]
+		public string StreamId { get; set; }
+
+		/// <summary>
+		/// The template for this job
+		/// </summary>
+		[Required]
+		public string TemplateId { get; set; }
+
+		/// <summary>
+		/// Name of the job
+		/// </summary>
+		public string? Name { get; set; }
+
+		/// <summary>
+		/// The changelist number to build. Can be null for latest.
+		/// </summary>
+		public int? Change { get; set; }
+
+		/// <summary>
+		/// Parameters to use when selecting the change to execute at.
+		/// </summary>
+		public ChangeQueryRequest? ChangeQuery { get; set; }
+
+		/// <summary>
+		/// The preflight changelist number
+		/// </summary>
+		public int? PreflightChange { get; set; }
+
+		/// <summary>
+		/// Priority for the job
+		/// </summary>
+		public Priority? Priority { get; set; }
+
+		/// <summary>
+		/// Whether to automatically submit the preflighted change on completion
+		/// </summary>
+		public bool? AutoSubmit { get; set; }
+
+		/// <summary>
+		/// Whether to update issues based on the outcome of this job
+		/// </summary>
+		public bool? UpdateIssues { get; set; }
+
+		/// <summary>
+		/// Arguments for the job
+		/// </summary>
+		public List<string>? Arguments { get; set; }
+
+		/// <summary>
+		/// Private constructor for serialization
+		/// </summary>
+		public CreateJobRequest(string streamId, string templateId)
+		{
+			StreamId = streamId;
+			TemplateId = templateId;
+		}
+	}
+
+	/// <summary>
+	/// Response from creating a new job
+	/// </summary>
+	public class CreateJobResponse
+	{
+		/// <summary>
+		/// Unique id for the new job
+		/// </summary>
+		public string Id { get; set; }
+
+		/// <summary>
+		/// Constructor
+		/// </summary>
+		/// <param name="id">Unique id for the new job</param>
+		public CreateJobResponse(string id)
+		{
+			Id = id;
+		}
+	}
+
+	/// <summary>
+	/// Updates an existing job
+	/// </summary>
+	public class UpdateJobRequest
+	{
+		/// <summary>
+		/// New name for the job
+		/// </summary>
+		public string? Name { get; set; }
+
+		/// <summary>
+		/// New priority for the job
+		/// </summary>
+		public Priority? Priority { get; set; }
+
+		/// <summary>
+		/// Set whether the job should be automatically submitted or not
+		/// </summary>
+		public bool? AutoSubmit { get; set; }
+
+		/// <summary>
+		/// Mark this job as aborted
+		/// </summary>
+		public bool? Aborted { get; set; }
+
+		/// <summary>
+		/// New list of arguments for the job. Only -Target= arguments can be modified after the job has started.
+		/// </summary>
+		public List<string>? Arguments { get; set; }
+
+		/// <summary>
+		/// Custom permissions for this object
+		/// </summary>
+		public UpdateAclRequest? Acl { get; set; }
+	}
+
+	/// <summary>
+	/// Information about a report associated with a job
+	/// </summary>
+	public class GetReportResponse
+	{
+		/// <summary>
+		/// Name of the report
+		/// </summary>
+		public string Name { get; set; }
+
+		/// <summary>
+		/// Report placement
+		/// </summary>
+		public HordeCommon.Rpc.ReportPlacement Placement { get; set; }
+
+		/// <summary>
+		/// The artifact id
+		/// </summary>
+		public string ArtifactId { get; set; }
+
+		/// <summary>
+		/// Constructor
+		/// </summary>
+		/// <param name="report"></param>
+		public GetReportResponse(IReport report)
+		{
+			Name = report.Name;
+			Placement = report.Placement;
+			ArtifactId = report.ArtifactId.ToString();
+		}
+	}
+
+	/// <summary>
+	/// Information about a job
+	/// </summary>
+	public class GetJobResponse
+	{
+		/// <summary>
+		/// Unique Id for the job
+		/// </summary>
+		public string Id { get; set; }
+
+		/// <summary>
+		/// Unique id of the stream containing this job
+		/// </summary>
+		public string StreamId { get; set; }
+
+		/// <summary>
+		/// Name of the job
+		/// </summary>
+		public string Name { get; set; }
+
+		/// <summary>
+		/// The changelist number to build
+		/// </summary>
+		public int Change { get; set; }
+
+		/// <summary>
+		/// The code changelist
+		/// </summary>
+		public int? CodeChange { get; set; }
+
+		/// <summary>
+		/// The preflight changelist number
+		/// </summary>
+		public int? PreflightChange { get; set; }
+
+		/// <summary>
+		/// The cloned preflight changelist number
+		/// </summary>
+		public int? ClonedPreflightChange { get; set; }
+
+		/// <summary>
+		/// The template type
+		/// </summary>
+		public string TemplateId { get; set; }
+
+		/// <summary>
+		/// Hash of the actual template data
+		/// </summary>
+		public string TemplateHash { get; set; }
+
+		/// <summary>
+		/// Hash of the graph for this job
+		/// </summary>
+		public string? GraphHash { get; set; }
+
+		/// <summary>
+		/// The user that started this job [DEPRECATED]
+		/// </summary>
+		public string? StartedByUserId { get; set; }
+
+		/// <summary>
+		/// The user that started this job [DEPRECATED]
+		/// </summary>
+		public string? StartedByUser { get; set; }
+
+		/// <summary>
+		/// The user that started this job
+		/// </summary>
+		public GetThinUserInfoResponse? StartedByUserInfo { get; set; }
+
+		/// <summary>
+		/// The user that aborted this job [DEPRECATED]
+		/// </summary>
+		public string? AbortedByUser { get; set; }
+
+		/// <summary>
+		/// The user that aborted this job
+		/// </summary>
+		public GetThinUserInfoResponse? AbortedByUserInfo { get; set; }
+
+		/// <summary>
+		/// Priority of the job
+		/// </summary>
+		public Priority Priority { get; set; }
+
+		/// <summary>
+		/// Whether the change will automatically be submitted or not
+		/// </summary>
+		public bool AutoSubmit { get; set; }
+
+		/// <summary>
+		/// The submitted changelist number
+		/// </summary>
+		public int? AutoSubmitChange { get; }
+
+		/// <summary>
+		/// Message produced by trying to auto-submit the change
+		/// </summary>
+		public string? AutoSubmitMessage { get; }
+
+		/// <summary>
+		/// Time that the job was created
+		/// </summary>
+		public DateTimeOffset CreateTime { get; set; }
+
+		/// <summary>
+		/// The global job state
+		/// </summary>
+		public JobState State { get; set; }
+
+		/// <summary>
+		/// Array of jobstep batches
+		/// </summary>
+		public List<GetBatchResponse>? Batches { get; set; }
+
+		/// <summary>
+		/// List of labels
+		/// </summary>
+		public List<GetLabelStateResponse>? Labels { get; set; }
+
+		/// <summary>
+		/// The default label, containing the state of all steps that are otherwise not matched.
+		/// </summary>
+		public GetDefaultLabelStateResponse? DefaultLabel { get; set; } 
+
+		/// <summary>
+		/// List of reports
+		/// </summary>
+		public List<GetReportResponse>? Reports { get; set; }
+
+		/// <summary>
+		/// Parameters for the job
+		/// </summary>
+		public List<string> Arguments { get; set; }
+
+		/// <summary>
+		/// The last update time for this job
+		/// </summary>
+		public DateTimeOffset UpdateTime { get; set; }
+
+		/// <summary>
+		/// Whether issues are being updated by this job
+		/// </summary>
+		public bool UpdateIssues { get; set; }
+
+		/// <summary>
+		/// Per-object permissions
+		/// </summary>
+		public GetAclResponse? Acl { get; set; }
+
+		/// <summary>
+		/// Constructor
+		/// </summary>
+		/// <param name="job">Job to create a response for</param>
+		/// <param name="startedByUserInfo">User that started this job</param>
+		/// <param name="abortedByUserInfo">User that aborted this job</param>
+		/// <param name="aclResponse">The ACL response object</param>
+		public GetJobResponse(IJob job, GetThinUserInfoResponse? startedByUserInfo, GetThinUserInfoResponse? abortedByUserInfo, GetAclResponse? aclResponse)
+		{
+			Id = job.Id.ToString();
+			StreamId = job.StreamId.ToString();
+			Name = job.Name;
+			Change = job.Change;
+			CodeChange = (job.CodeChange != 0) ? (int?)job.CodeChange : null;
+			PreflightChange = (job.PreflightChange != 0) ? (int?)job.PreflightChange : null;
+			ClonedPreflightChange = (job.ClonedPreflightChange != 0) ? (int?)job.ClonedPreflightChange : null;
+			TemplateId = job.TemplateId.ToString();
+			TemplateHash = job.TemplateHash?.ToString() ?? String.Empty;
+			GraphHash = job.GraphHash.ToString();
+			StartedByUserId = job.StartedByUserId?.ToString();
+			StartedByUser = startedByUserInfo?.Login;
+			StartedByUserInfo = startedByUserInfo;
+			AbortedByUser = abortedByUserInfo?.Login;
+			AbortedByUserInfo = abortedByUserInfo;
+			CreateTime = new DateTimeOffset(job.CreateTimeUtc);
+			State = job.GetState();
+			Priority = job.Priority;
+			AutoSubmit = job.AutoSubmit;
+			AutoSubmitChange = job.AutoSubmitChange;
+			AutoSubmitMessage = job.AutoSubmitMessage;
+			Reports = job.Reports?.ConvertAll(x => new GetReportResponse(x));
+			Arguments = job.Arguments.ToList();
+			UpdateTime = new DateTimeOffset(job.UpdateTimeUtc);
+			UpdateIssues = job.UpdateIssues;
+			Acl = aclResponse;
+		}
+	}
+
+	/// <summary>
+	/// The timing info for 
+	/// </summary>
+	public class GetJobTimingResponse
+	{
+		/// <summary>
+		/// The job object
+		/// </summary>
+		public IJob? Job { get; set; }
+		
+		/// <summary>
+		/// A full job response
+		/// </summary>
+		public GetJobResponse? JobResponse { get; set; }
+		
+		/// <summary>
+		/// Timing info for each step
+		/// </summary>
+		public Dictionary<string, GetStepTimingInfoResponse> Steps { get; set; }
+
+		/// <summary>
+		/// Timing information for each label
+		/// </summary>
+		public List<GetLabelTimingInfoResponse> Labels { get; set; }
+
+		/// <summary>
+		/// Constructor
+		/// </summary>
+		/// <param name="job">The job object</param>
+		/// <param name="jobResponse">The job response</param>
+		/// <param name="steps">Timing info for each steps</param>
+		/// <param name="labels">Timing info for each label</param>
+		public GetJobTimingResponse(IJob? job, GetJobResponse? jobResponse, Dictionary<string, GetStepTimingInfoResponse> steps, List<GetLabelTimingInfoResponse> labels)
+		{
+			Job = job;
+			JobResponse = jobResponse;
+			Steps = steps;
+			Labels = labels;
+		}
+	}
+	
+	/// <summary>
+	/// The timing info for 
+	/// </summary>
+	public class FindJobTimingsResponse
+	{
+		/// <summary>
+		/// Timing info for each job
+		/// </summary>
+		public Dictionary<string, GetJobTimingResponse> Timings { get; set; }
+
+		/// <summary>
+		/// Constructor
+		/// </summary>
+		/// <param name="timings">Timing info for each job</param>
+		public FindJobTimingsResponse(Dictionary<string, GetJobTimingResponse> timings)
+		{
+			Timings = timings;
+		}
+	}
+
+	/// <summary>
+	/// Request used to update a jobstep
+	/// </summary>
+	public class UpdateStepRequest
+	{
+		/// <summary>
+		/// The new jobstep state
+		/// </summary>
+		public JobStepState State { get; set; } = JobStepState.Unspecified;
+
+		/// <summary>
+		/// Outcome from the jobstep
+		/// </summary>
+		public JobStepOutcome Outcome { get; set; } = JobStepOutcome.Unspecified;
+		
+		/// <summary>
+		/// If the step has been requested to abort
+		/// </summary>
+		public bool? AbortRequested { get; set; }
+
+		/// <summary>
+		/// Specifies the log file id for this step
+		/// </summary>
+		public string? LogId { get; set; }
+
+		/// <summary>
+		/// Whether the step should be re-run
+		/// </summary>
+		public bool? Retry { get; set; }
+
+		/// <summary>
+		/// New priority for this step
+		/// </summary>
+		public Priority? Priority { get; set; }
+
+		/// <summary>
+		/// Properties to set. Any entries with a null value will be removed.
+		/// </summary>
+		public Dictionary<string, string?>? Properties { get; set; }
+	}
+
+	/// <summary>
+	/// Response object when updating a jobstep
+	/// </summary>
+	public class UpdateStepResponse
+	{
+		/// <summary>
+		/// If a new step is created (due to specifying the retry flag), specifies the batch id
+		/// </summary>
+		public string? BatchId { get; set; }
+
+		/// <summary>
+		/// If a step is retried, includes the new step id
+		/// </summary>
+		public string? StepId { get; set; }
+	}
+
+	/// <summary>
+	/// Returns information about a jobstep
+	/// </summary>
+	public class GetStepResponse
+	{
+		/// <summary>
+		/// The unique id of the step
+		/// </summary>
+		public string Id { get; set; }
+
+		/// <summary>
+		/// Index of the node which this jobstep is to execute
+		/// </summary>
+		public int NodeIdx { get; set; }
+
+		/// <summary>
+		/// Current state of the job step. This is updated automatically when runs complete.
+		/// </summary>
+		public JobStepState State { get; set; }
+
+		/// <summary>
+		/// Current outcome of the jobstep
+		/// </summary>
+		public JobStepOutcome Outcome { get; set; }
+
+		/// <summary>
+		/// Error describing additional context for why a step failed to complete
+		/// </summary>
+		public JobStepError Error { get; set; }
+		
+		/// <summary>
+		/// If the step has been requested to abort
+		/// </summary>
+		public bool AbortRequested { get; set; }
+		
+		/// <summary>
+		/// Name of the user that requested the abort of this step [DEPRECATED]
+		/// </summary>
+		public string? AbortByUser { get; set; }
+
+		/// <summary>
+		/// The user that requested this step be run again 
+		/// </summary>
+		public GetThinUserInfoResponse? AbortedByUserInfo { get; set; }
+
+		/// <summary>
+		/// Name of the user that requested this step be run again [DEPRECATED]
+		/// </summary>
+		public string? RetryByUser { get; set; }
+
+		/// <summary>
+		/// The user that requested this step be run again 
+		/// </summary>
+		public GetThinUserInfoResponse? RetriedByUserInfo { get; set; }
+
+		/// <summary>
+		/// The log id for this step
+		/// </summary>
+		public string? LogId { get; set; }
+
+		/// <summary>
+		/// Time at which the batch was ready (UTC).
+		/// </summary>
+		public DateTimeOffset? ReadyTime { get; set; }
+
+		/// <summary>
+		/// Time at which the batch started (UTC).
+		/// </summary>
+		public DateTimeOffset? StartTime { get; set; }
+
+		/// <summary>
+		/// Time at which the batch finished (UTC)
+		/// </summary>
+		public DateTimeOffset? FinishTime { get; set; }
+
+		/// <summary>
+		/// List of reports
+		/// </summary>
+		public List<GetReportResponse>? Reports { get; set; }
+
+		/// <summary>
+		/// User-defined properties for this jobstep.
+		/// </summary>
+		public Dictionary<string, string>? Properties { get; set; }
+
+		/// <summary>
+		/// Constructor
+		/// </summary>
+		/// <param name="step">The step to construct from</param>
+		/// <param name="abortedByUserInfo">User that aborted this step</param>
+		/// <param name="retriedByUserInfo">User that retried this step</param>
+		public GetStepResponse(IJobStep step, GetThinUserInfoResponse? abortedByUserInfo, GetThinUserInfoResponse? retriedByUserInfo)
+		{
+			Id = step.Id.ToString();
+			NodeIdx = step.NodeIdx;
+			State = step.State;
+			Outcome = step.Outcome;
+			Error = step.Error;
+			AbortRequested = step.AbortRequested;
+			AbortByUser = abortedByUserInfo?.Login;
+			AbortedByUserInfo = abortedByUserInfo;
+			RetryByUser = retriedByUserInfo?.Login;
+			RetriedByUserInfo = retriedByUserInfo;
+			LogId = step.LogId?.ToString();
+			ReadyTime = step.ReadyTimeUtc;
+			StartTime = step.StartTimeUtc;
+			FinishTime = step.FinishTimeUtc;
+			Reports = step.Reports?.ConvertAll(x => new GetReportResponse(x));
+
+			if (step.Properties != null && step.Properties.Count > 0)
+			{
+				Properties = step.Properties;
+			}
+		}
+	}
+
+	/// <summary>
+	/// The state of a particular run
+	/// </summary>
+	[JsonConverter(typeof(JsonStringEnumConverter))]
+	public enum JobStepBatchState
+	{
+		/// <summary>
+		/// Waiting for dependencies of at least one jobstep to complete
+		/// </summary>
+		Waiting = 0,
+
+		/// <summary>
+		/// Ready to execute
+		/// </summary>
+		Ready = 1,
+
+		/// <summary>
+		/// Preparing to execute work
+		/// </summary>
+		Starting = 2,
+
+		/// <summary>
+		/// Executing work
+		/// </summary>
+		Running = 3,
+
+		/// <summary>
+		/// Preparing to stop
+		/// </summary>
+		Stopping = 4,
+
+		/// <summary>
+		/// All steps have finished executing
+		/// </summary>
+		Complete = 5
+	}
+
+	/// <summary>
+	/// Error code for a batch not being executed
+	/// </summary>
+	[JsonConverter(typeof(JsonStringEnumConverter))]
+	public enum JobStepBatchError
+	{
+		/// <summary>
+		/// No error
+		/// </summary>
+		None = 0,
+
+		/// <summary>
+		/// The stream for this job is unknown
+		/// </summary>
+		UnknownStream = 1,
+
+		/// <summary>
+		/// The given agent type for this batch was not valid for this stream
+		/// </summary>
+		UnknownAgentType = 2,
+
+		/// <summary>
+		/// The pool id referenced by the agent type was not found
+		/// </summary>
+		UnknownPool = 3,
+
+		/// <summary>
+		/// There are no agents in the given pool currently online
+		/// </summary>
+		NoAgentsInPool = 4,
+
+		/// <summary>
+		/// There are no agents in this pool that are onlinbe
+		/// </summary>
+		NoAgentsOnline = 5,
+
+		/// <summary>
+		/// Unknown workspace referenced by the agent type
+		/// </summary>
+		UnknownWorkspace = 6,
+
+		/// <summary>
+		/// Cancelled
+		/// </summary>
+		Cancelled = 7,
+
+		/// <summary>
+		/// Lost connection with the agent machine
+		/// </summary>
+		LostConnection = 8,
+
+		/// <summary>
+		/// Lease terminated prematurely but can be retried.
+		/// </summary>
+		Incomplete = 9,
+
+		/// <summary>
+		/// An error ocurred while executing the lease. Cannot be retried.
+		/// </summary>
+		ExecutionError = 10,
+	}
+
+	/// <summary>
+	/// Request to update a jobstep batch
+	/// </summary>
+	public class UpdateBatchRequest
+	{
+		/// <summary>
+		/// The new log file id
+		/// </summary>
+		public string? LogId { get; set; }
+
+		/// <summary>
+		/// The state of this batch
+		/// </summary>
+		public JobStepBatchState? State { get; set; }
+	}
+
+	/// <summary>
+	/// Information about a jobstep batch
+	/// </summary>
+	public class GetBatchResponse
+	{
+		/// <summary>
+		/// Unique id for this batch
+		/// </summary>
+		public string Id { get; set; }
+
+		/// <summary>
+		/// The unique log file id
+		/// </summary>
+		public string? LogId { get; set; }
+
+		/// <summary>
+		/// Index of the group being executed
+		/// </summary>
+		public int GroupIdx { get; set; }
+
+		/// <summary>
+		/// The state of this batch
+		/// </summary>
+		public JobStepBatchState State { get; set; }
+
+		/// <summary>
+		/// Error code for this batch
+		/// </summary>
+		public JobStepBatchError Error { get; set; }
+
+		/// <summary>
+		/// Steps within this run
+		/// </summary>
+		public List<GetStepResponse> Steps { get; set; }
+
+		/// <summary>
+		/// The agent assigned to execute this group
+		/// </summary>
+		public string? AgentId { get; set; }
+
+		/// <summary>
+		/// Rate for using this agent (per hour)
+		/// </summary>
+		public double? AgentRate { get; set; }
+
+		/// <summary>
+		/// The agent session holding this lease
+		/// </summary>
+		public string? SessionId { get; set; }
+
+		/// <summary>
+		/// The lease that's executing this group
+		/// </summary>
+		public string? LeaseId { get; set; }
+
+		/// <summary>
+		/// The priority of this batch
+		/// </summary>
+		public int WeightedPriority { get; set; }
+
+		/// <summary>
+		/// Time at which the group started (UTC).
+		/// </summary>
+		public DateTimeOffset? StartTime { get; set; }
+
+		/// <summary>
+		/// Time at which the group finished (UTC)
+		/// </summary>
+		public DateTimeOffset? FinishTime { get; set; }
+
+		/// <summary>
+		/// Converts this batch into a public response object
+		/// </summary>
+		/// <param name="batch">The batch to construct from</param>
+		/// <param name="steps">Steps in this batch</param>
+		/// <param name="agentRate">Rate for this agent</param>
+		/// <returns>Response instance</returns>
+		public GetBatchResponse(IJobStepBatch batch, List<GetStepResponse> steps, double? agentRate)
+		{
+			Id = batch.Id.ToString();
+			LogId = batch.LogId?.ToString();
+			GroupIdx = batch.GroupIdx;
+			State = batch.State;
+			Error = batch.Error;
+			Steps = steps;
+			AgentId = batch.AgentId?.ToString();
+			AgentRate = agentRate;
+			SessionId = batch.SessionId?.ToString();
+			LeaseId = batch.LeaseId?.ToString();
+			WeightedPriority = batch.SchedulePriority;
+			StartTime = batch.StartTimeUtc;
+			FinishTime = batch.FinishTimeUtc;
+		}
+	}
+
+	/// <summary>
+	/// State of an aggregate
+	/// </summary>
+	public enum LabelState
+	{
+		/// <summary>
+		/// Aggregate is not currently being built (no required nodes are present)
+		/// </summary>
+		Unspecified,
+
+		/// <summary>
+		/// Steps are still running
+		/// </summary>
+		Running,
+
+		/// <summary>
+		/// All steps are complete
+		/// </summary>
+		Complete
+	}
+
+	/// <summary>
+	/// Outcome of an aggregate
+	/// </summary>
+	public enum LabelOutcome
+	{
+		/// <summary>
+		/// Aggregate is not currently being built
+		/// </summary>
+		Unspecified,
+
+		/// <summary>
+		/// A step dependency failed
+		/// </summary>
+		Failure,
+
+		/// <summary>
+		/// A dependency finished with warnings
+		/// </summary>
+		Warnings,
+
+		/// <summary>
+		/// Successful
+		/// </summary>
+		Success,
+	}
+
+	/// <summary>
+	/// State of a label within a job
+	/// </summary>
+	public class GetLabelStateResponse
+	{
+		/// <summary>
+		/// State of the label
+		/// </summary>
+		public LabelState? State { get; set; }
+
+		/// <summary>
+		/// Outcome of the label
+		/// </summary>
+		public LabelOutcome? Outcome { get; set; }
+
+		/// <summary>
+		/// Constructor
+		/// </summary>
+		/// <param name="state">State of the label</param>
+		/// <param name="outcome">Outcome of the label</param>
+		public GetLabelStateResponse(LabelState state, LabelOutcome outcome)
+		{
+			State = state;
+			Outcome = outcome;
+		}
+	}
+
+	/// <summary>
+	/// Information about the default label (ie. with inlined list of nodes)
+	/// </summary>
+	public class GetDefaultLabelStateResponse : GetLabelStateResponse
+	{
+		/// <summary>
+		/// List of nodes covered by default label
+		/// </summary>
+		public List<string> Nodes { get; set; }
+
+		/// <summary>
+		/// Constructor
+		/// </summary>
+		/// <param name="state">State of the label</param>
+		/// <param name="outcome">Outcome of the label</param>
+		/// <param name="nodes">List of nodes that are covered by the default label</param>
+		public GetDefaultLabelStateResponse(LabelState state, LabelOutcome outcome, List<string> nodes)
+			: base(state, outcome)
+		{
+			Nodes = nodes;
+		}
+	}
+
+	/// <summary>
+	/// Information about the timing info for a particular target
+	/// </summary>
+	public class GetTimingInfoResponse
+	{
+		/// <summary>
+		/// Wait time on the critical path
+		/// </summary>
+		public float? TotalWaitTime { get; set; }
+
+		/// <summary>
+		/// Sync time on the critical path
+		/// </summary>
+		public float? TotalInitTime { get; set; }
+
+		/// <summary>
+		/// Duration to this point
+		/// </summary>
+		public float? TotalTimeToComplete { get; set; }
+
+		/// <summary>
+		/// Average wait time by the time the job reaches this point
+		/// </summary>
+		public float? AverageTotalWaitTime { get; set; }
+
+		/// <summary>
+		/// Average sync time to this point
+		/// </summary>
+		public float? AverageTotalInitTime { get; set; }
+
+		/// <summary>
+		/// Average duration to this point
+		/// </summary>
+		public float? AverageTotalTimeToComplete { get; set; }
+
+		/// <summary>
+		/// Constructor
+		/// </summary>
+		/// <param name="timingInfo">Timing info to construct from</param>
+		public GetTimingInfoResponse(TimingInfo timingInfo)
+		{
+			TotalWaitTime = (float?)timingInfo.TotalWaitTime?.TotalSeconds;
+			TotalInitTime = (float?)timingInfo.TotalInitTime?.TotalSeconds;
+			TotalTimeToComplete = (float?)timingInfo.TotalTimeToComplete?.TotalSeconds;
+
+			AverageTotalWaitTime = (float?)timingInfo.AverageTotalWaitTime?.TotalSeconds;
+			AverageTotalInitTime = (float?)timingInfo.AverageTotalInitTime?.TotalSeconds;
+			AverageTotalTimeToComplete = (float?)timingInfo.AverageTotalTimeToComplete?.TotalSeconds;
+		}
+	}
+
+	/// <summary>
+	/// Information about the timing info for a particular target
+	/// </summary>
+	public class GetStepTimingInfoResponse : GetTimingInfoResponse
+	{
+		/// <summary>
+		/// Name of this node
+		/// </summary>
+		public string? Name { get; set; }
+
+		/// <summary>
+		/// Average wait time for this step
+		/// </summary>
+		public float? AverageStepWaitTime { get; set; }
+
+		/// <summary>
+		/// Average init time for this step
+		/// </summary>
+		public float? AverageStepInitTime { get; set; }
+
+		/// <summary>
+		/// Average duration for this step
+		/// </summary>
+		public float? AverageStepDuration { get; set; }
+
+		/// <summary>
+		/// Constructor
+		/// </summary>
+		/// <param name="name">Name of the node</param>
+		/// <param name="jobTimingInfo">Timing info to construct from</param>
+		public GetStepTimingInfoResponse(string? name, TimingInfo jobTimingInfo)
+			: base(jobTimingInfo)
+		{
+			Name = name;
+			AverageStepWaitTime = jobTimingInfo.StepTiming?.AverageWaitTime;
+			AverageStepInitTime = jobTimingInfo.StepTiming?.AverageInitTime;
+			AverageStepDuration = jobTimingInfo.StepTiming?.AverageDuration;
+		}
+	}
+
+	/// <summary>
+	/// Information about the timing info for a label
+	/// </summary>
+	public class GetLabelTimingInfoResponse : GetTimingInfoResponse
+	{
+		/// <summary>
+		/// Name of the label
+		/// </summary>
+		[Obsolete("Use DashboardName instead")]
+		public string? Name => DashboardName;
+
+		/// <summary>
+		/// Category for the label
+		/// </summary>
+		[Obsolete("Use DashboardCategory instead")]
+		public string? Category => DashboardCategory;
+
+		/// <summary>
+		/// Name of the label
+		/// </summary>
+		public string? DashboardName { get; set; }
+
+		/// <summary>
+		/// Category for the label
+		/// </summary>
+		public string? DashboardCategory { get; set; }
+
+		/// <summary>
+		/// Name of the label
+		/// </summary>
+		public string? UgsName { get; set; }
+
+		/// <summary>
+		/// Category for the label
+		/// </summary>
+		public string? UgsProject { get; set; }
+
+		/// <summary>
+		/// Constructor
+		/// </summary>
+		/// <param name="label">The label to construct from</param>
+		/// <param name="timingInfo">Timing info to construct from</param>
+		public GetLabelTimingInfoResponse(ILabel label, TimingInfo timingInfo)
+			: base(timingInfo)
+		{
+			DashboardName = label.DashboardName;
+			DashboardCategory = label.DashboardCategory;
+			UgsName = label.UgsName;
+			UgsProject = label.UgsProject;
+		}
+	}
+
+	/// <summary>
+	/// Describes the history of a step
+	/// </summary>
+	public class GetJobStepRefResponse
+	{
+		/// <summary>
+		/// The job id
+		/// </summary>
+		public string JobId { get; set; }
+
+		/// <summary>
+		/// The batch containing the step
+		/// </summary>
+		public string BatchId { get; set; }
+
+		/// <summary>
+		/// The step identifier
+		/// </summary>
+		public string StepId { get; set; }
+
+		/// <summary>
+		/// The change number being built
+		/// </summary>
+		public int Change { get; set; }
+
+		/// <summary>
+		/// The step log id
+		/// </summary>
+		public string? LogId { get; set; }
+
+		/// <summary>
+		/// The pool id
+		/// </summary>
+		public string? PoolId { get; set; }
+
+		/// <summary>
+		/// The agent id
+		/// </summary>
+		public string? AgentId { get; set; }
+
+		/// <summary>
+		/// Outcome of the step, once complete.
+		/// </summary>
+		public JobStepOutcome? Outcome { get; set; }
+
+		/// <summary>
+		/// Time at which the step started.
+		/// </summary>
+		public DateTimeOffset StartTime { get; }
+
+		/// <summary>
+		/// Time at which the step finished.
+		/// </summary>
+		public DateTimeOffset? FinishTime { get; }
+
+		/// <summary>
+		/// Constructor
+		/// </summary>
+		/// <param name="jobStepRef">The jobstep ref to construct from</param>
+		public GetJobStepRefResponse(IJobStepRef jobStepRef)
+		{
+			JobId = jobStepRef.Id.JobId.ToString();
+			BatchId = jobStepRef.Id.BatchId.ToString();
+			StepId = jobStepRef.Id.StepId.ToString();
+			Change = jobStepRef.Change;
+			LogId = jobStepRef.LogId.ToString();
+			PoolId = jobStepRef.PoolId?.ToString();
+			AgentId = jobStepRef.AgentId?.ToString();
+			Outcome = jobStepRef.Outcome;
+			StartTime = jobStepRef.StartTimeUtc;
+			FinishTime = jobStepRef.FinishTimeUtc;
+		}
+	}
+}
