Horde: Finish re-org of Horde solution started a while ago.

diff --git a/Engine/Source/Programs/Horde/Horde.Build/Jobs/JobsController.cs b/Engine/Source/Programs/Horde/Horde.Build/Jobs/JobsController.cs
new file mode 100644
--- /dev/null
+++ b/Engine/Source/Programs/Horde/Horde.Build/Jobs/JobsController.cs
@@ -0,0 +1,1555 @@
+// Copyright Epic Games, Inc. All Rights Reserved.
+
+using System;
+using System.Collections.Generic;
+using System.Linq;
+using System.Threading.Tasks;
+using EpicGames.Core;
+using EpicGames.Perforce;
+using Horde.Build.Acls;
+using Horde.Build.Api;
+using Horde.Build.Collections;
+using Horde.Build.Jobs;
+using Horde.Build.Models;
+using Horde.Build.Notifications;
+using Horde.Build.Services;
+using Horde.Build.Utilities;
+using HordeCommon;
+using Microsoft.AspNetCore.Authorization;
+using Microsoft.AspNetCore.Mvc;
+using Microsoft.Extensions.Logging;
+using MongoDB.Bson;
+using OpenTracing;
+using OpenTracing.Util;
+
+namespace Horde.Build.Controllers
+{
+	using JobId = ObjectId<IJob>;
+	using StreamId = StringId<IStream>;
+	using TemplateRefId = StringId<TemplateRef>;
+	using UserId = ObjectId<IUser>;
+
+	/// <summary>
+	/// Controller for the /api/v1/jobs endpoing
+	/// </summary>
+	[ApiController]
+	[Authorize]
+	[Route("[controller]")]
+	public class JobsController : HordeControllerBase
+	{
+		private readonly IGraphCollection _graphs;
+		private readonly IPerforceService _perforce;
+		private readonly StreamService _streamService;
+		private readonly JobService _jobService;
+		private readonly ITemplateCollection _templateCollection;
+		private readonly IArtifactCollection _artifactCollection;
+		private readonly IUserCollection _userCollection;
+		private readonly INotificationService _notificationService;
+		private readonly AgentService _agentService;
+		private readonly ILogger<JobsController> _logger;
+
+		/// <summary>
+		/// Constructor
+		/// </summary>
+		public JobsController(IGraphCollection graphs, IPerforceService perforce, StreamService streamService, JobService jobService, ITemplateCollection templateCollection, IArtifactCollection artifactCollection, IUserCollection userCollection, INotificationService notificationService, AgentService agentService, ILogger<JobsController> logger)
+		{
+			_graphs = graphs;
+			_perforce = perforce;
+			_streamService = streamService;
+			_jobService = jobService;
+			_templateCollection = templateCollection;
+			_artifactCollection = artifactCollection;
+			_userCollection = userCollection;
+			_notificationService = notificationService;
+			_agentService = agentService;
+			_logger = logger;
+		}
+
+		/// <summary>
+		/// Creates a new job
+		/// </summary>
+		/// <param name="create">Properties of the new job</param>
+		/// <returns>Id of the new job</returns>
+		[HttpPost]
+		[Route("/api/v1/jobs")]
+		public async Task<ActionResult<CreateJobResponse>> CreateJobAsync([FromBody] CreateJobRequest create)
+		{
+			IStream? stream = await _streamService.GetStreamAsync(new StreamId(create.StreamId));
+			if (stream == null)
+			{
+				return NotFound(create.StreamId);
+			}
+			if (!await _streamService.AuthorizeAsync(stream, AclAction.CreateJob, User, null))
+			{
+				return Forbid(AclAction.CreateJob, stream.Id);
+			}
+
+			// Get the name of the template ref
+			TemplateRefId templateRefId = new TemplateRefId(create.TemplateId);
+
+			// Augment the request with template properties
+			TemplateRef? templateRef;
+			if (!stream.Templates.TryGetValue(templateRefId, out templateRef))
+			{
+				return BadRequest($"Template {create.TemplateId} is not available for stream {stream.Id}");
+			}
+			if (!await _streamService.AuthorizeAsync(stream, templateRef, AclAction.CreateJob, User, null))
+			{
+				return Forbid(AclAction.CreateJob, stream.Id);
+			}
+
+			ITemplate? template = await _templateCollection.GetAsync(templateRef.Hash);
+			if (template == null)
+			{
+				return BadRequest("Missing template referenced by {TemplateId}", create.TemplateId);
+			}
+			if (!template.AllowPreflights && create.PreflightChange > 0)
+			{
+				return BadRequest("Template {TemplateId} does not allow preflights", create.TemplateId);
+			}
+
+			// Get the name of the new job
+			string name = create.Name ?? template.Name;
+			if (create.TemplateId.Equals("stage-to-marketplace", StringComparison.Ordinal) && create.Arguments != null)
+			{
+				foreach (string argument in create.Arguments)
+				{
+					const string Prefix = "-set:UserContentItems=";
+					if (argument.StartsWith(Prefix, StringComparison.Ordinal))
+					{
+						name += $" - {argument.Substring(Prefix.Length)}";
+						break;
+					}
+				}
+			}
+
+			// Get the priority of the new job
+			Priority priority = create.Priority ?? template.Priority ?? Priority.Normal;
+
+			// New groups for the job
+			IGraph graph = await _graphs.AddAsync(template);
+
+			// Get the change to build
+			int change;
+			if (create.Change.HasValue)
+			{
+				change = create.Change.Value;
+			}
+			else if (create.ChangeQuery != null)
+			{
+				change = await ExecuteChangeQueryAsync(stream, new TemplateRefId(create.ChangeQuery.TemplateId ?? create.TemplateId), create.ChangeQuery.Target, create.ChangeQuery.Outcomes ?? new List<JobStepOutcome> { JobStepOutcome.Success });
+			}
+			else if (create.PreflightChange == null && template.SubmitNewChange != null)
+			{
+				change = await _perforce.CreateNewChangeForTemplateAsync(stream, template);
+			}
+			else
+			{
+				change = await _perforce.GetLatestChangeAsync(stream.ClusterName, stream.Name, null);
+			}
+
+			// And get the matching code changelist
+			int codeChange = await _perforce.GetCodeChangeAsync(stream.ClusterName, stream.Name, change);
+
+			// New properties for the job
+			List<string> arguments = create.Arguments ?? template.GetDefaultArguments();
+
+			// Check the preflight change is valid
+			if (create.PreflightChange != null)
+			{
+				CheckShelfResult result = await _perforce.CheckShelfAsync(stream.ClusterName, stream.Name, create.PreflightChange.Value, null);
+				switch (result)
+				{
+					case CheckShelfResult.Ok:
+						break;
+					case CheckShelfResult.NoChange:
+						return BadRequest(KnownLogEvents.Horde_InvalidPreflight, "CL {Change} does not exist", create.PreflightChange);
+					case CheckShelfResult.NoShelvedFiles:
+						return BadRequest(KnownLogEvents.Horde_InvalidPreflight, "CL {Change} does not contain any shelved files", create.PreflightChange);
+					case CheckShelfResult.WrongStream:
+						return BadRequest(KnownLogEvents.Horde_InvalidPreflight, "CL {Change} does not contain files in {Stream}", create.PreflightChange, stream.Name);
+					case CheckShelfResult.MixedStream:
+						return BadRequest(KnownLogEvents.Horde_InvalidPreflight, "CL {Change} contains files from multiple streams", create.PreflightChange);
+					default:
+						return BadRequest(KnownLogEvents.Horde_InvalidPreflight, "CL {Change} cannot be preflighted ({Result})", create.PreflightChange, result);
+				}
+			}
+
+			bool? updateIssues = null;
+			if (template.UpdateIssues)
+			{
+				updateIssues = true;
+			}
+			else if (create.UpdateIssues.HasValue)
+			{
+				updateIssues = create.UpdateIssues.Value;
+			}
+
+			// Create the job
+			IJob job = await _jobService.CreateJobAsync(null, stream, templateRefId, template.Id, graph, name, change, codeChange, create.PreflightChange, null, User.GetUserId(), priority, create.AutoSubmit, updateIssues, false, templateRef.ChainedJobs, templateRef.ShowUgsBadges, templateRef.ShowUgsAlerts, templateRef.NotificationChannel, templateRef.NotificationChannelFilter, arguments);
+			await UpdateNotificationsAsync(job.Id, new UpdateNotificationsRequest { Slack = true });
+			return new CreateJobResponse(job.Id.ToString());
+		}
+
+		/// <summary>
+		/// Evaluate a change query to determine which CL to run a job at
+		/// </summary>
+		/// <param name="stream"></param>
+		/// <param name="templateId"></param>
+		/// <param name="target"></param>
+		/// <param name="outcomes"></param>
+		/// <returns></returns>
+		async Task<int> ExecuteChangeQueryAsync(IStream stream, TemplateRefId templateId, string? target, List<JobStepOutcome> outcomes)
+		{
+			IList<IJob> jobs = await _jobService.FindJobsAsync(streamId: stream.Id, templates: new[] { templateId }, target: target, state: new[] { JobStepState.Completed }, outcome: outcomes.ToArray(), count: 1, excludeUserJobs: true);
+			if (jobs.Count == 0)
+			{
+				_logger.LogInformation("Unable to find successful build of {TemplateId} target {Target}. Using latest change instead", templateId, target);
+				return await _perforce.GetLatestChangeAsync(stream.ClusterName, stream.Name, null);
+			}
+			else
+			{
+				_logger.LogInformation("Last successful build of {TemplateId} target {Target} was job {JobId} at change {Change}", templateId, target, jobs[0].Id, jobs[0].Change);
+				return jobs[0].Change;
+			}
+		}
+
+		/// <summary>
+		/// Deletes a specific job.
+		/// </summary>
+		/// <param name="jobId">Id of the job to delete</param>
+		/// <returns>Async task</returns>
+		[HttpDelete]
+		[Route("/api/v1/jobs/{jobId}")]
+		public async Task<ActionResult> DeleteJobAsync(JobId jobId)
+		{
+			IJob? job = await _jobService.GetJobAsync(jobId);
+			if (job == null)
+			{
+				return NotFound(jobId);
+			}
+			if (!await _jobService.AuthorizeAsync(job, AclAction.DeleteJob, User, null))
+			{
+				return Forbid(AclAction.DeleteJob, jobId);
+			}
+			if (!await _jobService.DeleteJobAsync(job))
+			{
+				return NotFound(jobId);
+			}
+			return Ok();
+		}
+
+		/// <summary>
+		/// Updates a specific job.
+		/// </summary>
+		/// <param name="jobId">Id of the job to find</param>
+		/// <param name="request">Settings to update in the job</param>
+		/// <returns>Async task</returns>
+		[HttpPut]
+		[Route("/api/v1/jobs/{jobId}")]
+		public async Task<ActionResult> UpdateJobAsync(JobId jobId, [FromBody] UpdateJobRequest request)
+		{
+			IJob? job = await _jobService.GetJobAsync(jobId);
+			if (job == null)
+			{
+				return NotFound(jobId);
+			}
+
+			StreamPermissionsCache permissionsCache = new StreamPermissionsCache();
+			if (!await _jobService.AuthorizeAsync(job, AclAction.UpdateJob, User, permissionsCache))
+			{
+				return Forbid(AclAction.UpdateJob, jobId);
+			}
+			if (request.Acl != null && !await _jobService.AuthorizeAsync(job, AclAction.ChangePermissions, User, permissionsCache))
+			{
+				return Forbid(AclAction.ChangePermissions, jobId);
+			}
+
+			// Convert legacy behavior of clearing out the argument to setting the aborted flag
+			if (request.Arguments != null && request.Arguments.Count == 0)
+			{
+				request.Aborted = true;
+				request.Arguments = null;
+			}
+
+			UserId? abortedByUserId = null;
+			if (request.Aborted ?? false)
+			{
+				abortedByUserId = User.GetUserId();
+			}
+
+			IJob? newJob = await _jobService.UpdateJobAsync(job, name: request.Name, priority: request.Priority, autoSubmit: request.AutoSubmit, abortedByUserId: abortedByUserId, arguments: request.Arguments);
+			if (newJob == null)
+			{
+				return NotFound(jobId);
+			}
+			return Ok();
+		}
+
+		/// <summary>
+		/// Updates notifications for a specific job.
+		/// </summary>
+		/// <param name="jobId">Id of the job to find</param>
+		/// <param name="request">The notification request</param>
+		/// <returns>Information about the requested job</returns>
+		[HttpPut]
+		[Route("/api/v1/jobs/{jobId}/notifications")]
+		public async Task<ActionResult> UpdateNotificationsAsync(JobId jobId, [FromBody] UpdateNotificationsRequest request)
+		{
+			IJob? job = await _jobService.GetJobAsync(jobId);
+			if (job == null)
+			{
+				return NotFound(jobId);
+			}
+			if (!await _jobService.AuthorizeAsync(job, AclAction.CreateSubscription, User, null))
+			{
+				return Forbid(AclAction.CreateSubscription, jobId);
+			}
+
+			ObjectId triggerId = job.NotificationTriggerId ?? ObjectId.GenerateNewId();
+
+			job = await _jobService.UpdateJobAsync(job, null, null, null, null, triggerId, null, null);
+			if (job == null)
+			{
+				return NotFound(jobId);
+			}
+
+			await _notificationService.UpdateSubscriptionsAsync(triggerId, User, request.Email, request.Slack);
+			return Ok();
+		}
+
+		/// <summary>
+		/// Gets information about a specific job.
+		/// </summary>
+		/// <param name="jobId">Id of the job to find</param>
+		/// <returns>Information about the requested job</returns>
+		[HttpGet]
+		[Route("/api/v1/jobs/{jobId}/notifications")]
+		public async Task<ActionResult<GetNotificationResponse>> GetNotificationsAsync(JobId jobId)
+		{
+			IJob? job = await _jobService.GetJobAsync(jobId);
+			if (job == null)
+			{
+				return NotFound(jobId);
+			}
+			if (!await _jobService.AuthorizeAsync(job, AclAction.CreateSubscription, User, null))
+			{
+				return Forbid(AclAction.CreateSubscription, jobId);
+			}
+
+			INotificationSubscription? subscription;
+			if (job.NotificationTriggerId == null)
+			{
+				subscription = null;
+			}
+			else
+			{
+				subscription = await _notificationService.GetSubscriptionsAsync(job.NotificationTriggerId.Value, User);
+			}
+			return new GetNotificationResponse(subscription);
+		}
+
+		/// <summary>
+		/// Gets information about a specific job.
+		/// </summary>
+		/// <param name="jobId">Id of the job to find</param>
+		/// <param name="modifiedAfter">If specified, returns an empty response unless the job's update time is equal to or less than the given value</param>
+		/// <param name="filter">Filter for the fields to return</param>
+		/// <returns>Information about the requested job</returns>
+		[HttpGet]
+		[Route("/api/v1/jobs/{jobId}")]
+		[ProducesResponseType(typeof(GetJobResponse), 200)]
+		public async Task<ActionResult<object>> GetJobAsync(JobId jobId, [FromQuery] DateTimeOffset? modifiedAfter = null, [FromQuery] PropertyFilter? filter = null)
+		{
+			IJob? job = await _jobService.GetJobAsync(jobId);
+			if (job == null)
+			{
+				return NotFound(jobId);
+			}
+
+			StreamPermissionsCache cache = new StreamPermissionsCache();
+			if (!await _jobService.AuthorizeAsync(job, AclAction.ViewJob, User, cache))
+			{
+				return Forbid(AclAction.ViewJob, jobId);
+			}
+			if (modifiedAfter != null && job.UpdateTimeUtc <= modifiedAfter.Value)
+			{
+				return new Dictionary<string, object>();
+			}
+
+			IGraph graph = await _jobService.GetGraphAsync(job);
+			bool bIncludeAcl = await _jobService.AuthorizeAsync(job, AclAction.ViewPermissions, User, cache);
+			bool bIncludeCosts = await _jobService.AuthorizeAsync(job, AclAction.ViewCosts, User, cache);
+			return await CreateJobResponseAsync(job, graph, bIncludeAcl, bIncludeCosts, filter);
+		}
+
+		/// <summary>
+		/// Creates a response containing information about this object
+		/// </summary>
+		/// <param name="job">The job document</param>
+		/// <param name="graph">The graph for this job</param>
+		/// <param name="bIncludeAcl">Whether to include the ACL in the response</param>
+		/// <param name="bIncludeCosts">Whether to include costs in the response</param>
+		/// <param name="filter">Filter for the properties to return</param>
+		/// <returns>Object containing the requested properties</returns>
+		async Task<object> CreateJobResponseAsync(IJob job, IGraph graph, bool bIncludeAcl, bool bIncludeCosts, PropertyFilter? filter)
+		{
+			if (filter == null)
+			{
+				return await CreateJobResponseAsync(job, graph, true, true, bIncludeAcl, bIncludeCosts);
+			}
+			else
+			{
+				return filter.ApplyTo(await CreateJobResponseAsync(job, graph, filter.Includes(nameof(GetJobResponse.Batches)), filter.Includes(nameof(GetJobResponse.Labels)) || filter.Includes(nameof(GetJobResponse.DefaultLabel)), bIncludeAcl, bIncludeCosts));
+			}
+		}
+
+		/// <summary>
+		/// Creates a response containing information about this object
+		/// </summary>
+		/// <param name="job">The job document</param>
+		/// <param name="graph">The graph definition</param>
+		/// <param name="bIncludeBatches">Whether to include the job batches in the response</param>
+		/// <param name="bIncludeLabels">Whether to include the job aggregates in the response</param>
+		/// <param name="bIncludeAcl">Whether to include the ACL in the response</param>
+		/// <param name="bIncludeCosts">Whether to include costs of running particular agents</param>
+		/// <returns>The response object</returns>
+		async ValueTask<GetJobResponse> CreateJobResponseAsync(IJob job, IGraph graph, bool bIncludeBatches, bool bIncludeLabels, bool bIncludeAcl, bool bIncludeCosts)
+		{
+			GetThinUserInfoResponse? startedByUserInfo = null;
+			if (job.StartedByUserId != null)
+			{
+				startedByUserInfo = new GetThinUserInfoResponse(await _userCollection.GetCachedUserAsync(job.StartedByUserId.Value));
+			}
+
+			GetThinUserInfoResponse? abortedByUserInfo = null;
+			if (job.AbortedByUserId != null)
+			{
+				abortedByUserInfo = new GetThinUserInfoResponse(await _userCollection.GetCachedUserAsync(job.AbortedByUserId.Value));
+			}
+
+			GetAclResponse? aclResponse = null;
+			if (bIncludeAcl && job.Acl != null)
+			{
+				aclResponse = new GetAclResponse(job.Acl);
+			}
+
+			GetJobResponse response = new GetJobResponse(job, startedByUserInfo, abortedByUserInfo, aclResponse);
+			if (bIncludeBatches || bIncludeLabels)
+			{
+				if (bIncludeBatches)
+				{
+					response.Batches = new List<GetBatchResponse>();
+					foreach (IJobStepBatch batch in job.Batches)
+					{
+						response.Batches.Add(await CreateBatchResponseAsync(batch, bIncludeCosts));
+					}
+				}
+				if (bIncludeLabels)
+				{
+					response.Labels = new List<GetLabelStateResponse>();
+					response.DefaultLabel = job.GetLabelStateResponses(graph, response.Labels);
+				}
+			}
+			return response;
+		}
+
+		/// <summary>
+		/// Get the response object for a batch
+		/// </summary>
+		/// <param name="batch"></param>
+		/// <param name="bIncludeCosts"></param>
+		/// <returns></returns>
+		async ValueTask<GetBatchResponse> CreateBatchResponseAsync(IJobStepBatch batch, bool bIncludeCosts)
+		{
+			List<GetStepResponse> steps = new List<GetStepResponse>();
+			foreach (IJobStep step in batch.Steps)
+			{
+				steps.Add(await CreateStepResponseAsync(step));
+			}
+
+			double? agentRate = null;
+			if (batch.AgentId != null && bIncludeCosts)
+			{
+				agentRate = await _agentService.GetRateAsync(batch.AgentId.Value);
+			}
+
+			return new GetBatchResponse(batch, steps, agentRate);
+		}
+
+		/// <summary>
+		/// Get the response object for a step
+		/// </summary>
+		/// <param name="step"></param>
+		/// <returns></returns>
+		async ValueTask<GetStepResponse> CreateStepResponseAsync(IJobStep step)
+		{
+			GetThinUserInfoResponse? abortedByUserInfo = null;
+			if (step.AbortedByUserId != null)
+			{
+				abortedByUserInfo = new GetThinUserInfoResponse(await _userCollection.GetCachedUserAsync(step.AbortedByUserId.Value));
+			}
+
+			GetThinUserInfoResponse? retriedByUserInfo = null;
+			if (step.RetriedByUserId != null)
+			{
+				retriedByUserInfo = new GetThinUserInfoResponse(await _userCollection.GetCachedUserAsync(step.RetriedByUserId.Value));
+			}
+
+			return new GetStepResponse(step, abortedByUserInfo, retriedByUserInfo);
+		}
+
+		/// <summary>
+		/// Gets information about the graph for a specific job.
+		/// </summary>
+		/// <param name="jobId">Id of the job to find</param>
+		/// <param name="filter">Filter for the fields to return</param>
+		/// <returns>Information about the requested job</returns>
+		[HttpGet]
+		[Route("/api/v1/jobs/{jobId}/graph")]
+		[ProducesResponseType(typeof(GetGraphResponse), 200)]
+		public async Task<ActionResult<object>> GetJobGraphAsync(JobId jobId, [FromQuery] PropertyFilter? filter = null)
+		{
+			IJob? job = await _jobService.GetJobAsync(jobId);
+			if (job == null)
+			{
+				return NotFound(jobId);
+			}
+			if (!await _jobService.AuthorizeAsync(job, AclAction.ViewJob, User, null))
+			{
+				return Forbid(AclAction.ViewJob, job.Id);
+			}
+
+			IGraph graph = await _jobService.GetGraphAsync(job);
+			return PropertyFilter.Apply(new GetGraphResponse(graph), filter);
+		}
+
+		/// <summary>
+		/// Gets timing information about the graph for a specific job.
+		/// </summary>
+		/// <param name="jobId">Id of the job to find</param>
+		/// <param name="filter">Filter for the fields to return</param>
+		/// <returns>Information about the requested job</returns>
+		[HttpGet]
+		[Route("/api/v1/jobs/{jobId}/timing")]
+		[ProducesResponseType(typeof(GetJobTimingResponse), 200)]
+		public async Task<ActionResult<object>> GetJobTimingAsync(JobId jobId, [FromQuery] PropertyFilter? filter = null)
+		{
+			IJob? job = await _jobService.GetJobAsync(jobId);
+			if (job == null)
+			{
+				return NotFound(jobId);
+			}
+			if (!await _jobService.AuthorizeAsync(job, AclAction.ViewJob, User, null))
+			{
+				return Forbid(AclAction.ViewJob, jobId);
+			}
+
+			IJobTiming jobTiming = await _jobService.GetJobTimingAsync(job);
+			IGraph graph = await _jobService.GetGraphAsync(job);
+			return PropertyFilter.Apply(await CreateJobTimingResponse(job, graph, jobTiming), filter);
+		}
+
+		private async Task<GetJobTimingResponse> CreateJobTimingResponse(IJob job, IGraph graph, IJobTiming jobTiming, bool includeJobResponse = false)
+		{
+			Dictionary<INode, TimingInfo> nodeToTimingInfo = job.GetTimingInfo(graph, jobTiming);
+
+			Dictionary<string, GetStepTimingInfoResponse> steps = new Dictionary<string, GetStepTimingInfoResponse>();
+			foreach (IJobStepBatch batch in job.Batches)
+			{
+				foreach (IJobStep step in batch.Steps)
+				{
+					INode node = graph.Groups[batch.GroupIdx].Nodes[step.NodeIdx];
+					steps[step.Id.ToString()] = new GetStepTimingInfoResponse(node.Name, nodeToTimingInfo[node]);
+				}
+			}
+
+			List<GetLabelTimingInfoResponse> labels = new List<GetLabelTimingInfoResponse>();
+			foreach (ILabel label in graph.Labels)
+			{
+				TimingInfo timingInfo = TimingInfo.Max(label.GetDependencies(graph.Groups).Select(x => nodeToTimingInfo[x]));
+				labels.Add(new GetLabelTimingInfoResponse(label, timingInfo));
+			}
+
+			GetJobResponse? jobResponse = null;
+			if (includeJobResponse)
+			{
+				jobResponse = await CreateJobResponseAsync(job, graph, true, true, false, true);
+			}
+
+			return new GetJobTimingResponse(job, jobResponse, steps, labels);
+		}
+		
+		/// <summary>
+		/// Find timing information about the graph for multiple jobs
+		/// </summary>
+		/// <param name="streamId">The stream to search in</param>
+		/// <param name="templates">List of templates to find</param>
+		/// <param name="filter">Filter for the fields to return</param>
+		/// <param name="count">Number of results to return</param>
+		/// <returns>Job timings for each job ID</returns>
+		[HttpGet]
+		[Route("/api/v1/jobs/timing")]
+		[ProducesResponseType(typeof(FindJobTimingsResponse), 200)]
+		public async Task<ActionResult<object>> FindJobTimingsAsync(
+			[FromQuery] string? streamId = null,
+			[FromQuery(Name = "template")] string[]? templates = null,
+			[FromQuery] PropertyFilter? filter = null,
+			[FromQuery] int count = 100)
+		{
+			if (streamId == null)
+			{
+				return BadRequest("Missing/invalid query parameter streamId");
+			}
+
+			TemplateRefId[] templateRefIds = templates switch
+			{
+				{ Length: > 0 } => templates.Select(x => new TemplateRefId(x)).ToArray(),
+				_ => Array.Empty<TemplateRefId>()
+			};
+
+			List<IJob> jobs = await _jobService.FindJobsByStreamWithTemplatesAsync(new StreamId(streamId), templateRefIds, count: count, consistentRead: false);
+
+			Dictionary<string, GetJobTimingResponse> jobTimings = await jobs.ToAsyncEnumerable()
+				.WhereAwait(async job => await _jobService.AuthorizeAsync(job, AclAction.ViewJob, User, null))
+				.ToDictionaryAwaitAsync(x => ValueTask.FromResult(x.Id.ToString()), async job =>
+				{
+					IJobTiming jobTiming = await _jobService.GetJobTimingAsync(job);
+					IGraph graph = await _jobService.GetGraphAsync(job);
+					return await CreateJobTimingResponse(job, graph, jobTiming, true);
+				});
+			
+			return PropertyFilter.Apply(new FindJobTimingsResponse(jobTimings), filter);
+		}
+
+		/// <summary>
+		/// Gets information about the template for a specific job.
+		/// </summary>
+		/// <param name="jobId">Id of the job to find</param>
+		/// <param name="filter">Filter for the fields to return</param>
+		/// <returns>Information about the requested job</returns>
+		[HttpGet]
+		[Route("/api/v1/jobs/{jobId}/template")]
+		[ProducesResponseType(typeof(GetTemplateResponse), 200)]
+		public async Task<ActionResult<object>> GetJobTemplateAsync(JobId jobId, [FromQuery] PropertyFilter? filter = null)
+		{
+			IJob? job = await _jobService.GetJobAsync(jobId);
+			if (job == null || job.TemplateHash == null)
+			{
+				return NotFound(jobId);
+			}
+			if (!await _jobService.AuthorizeAsync(job, AclAction.ViewJob, User, null))
+			{
+				return Forbid(AclAction.ViewJob, jobId);
+			}
+
+			ITemplate? template = await _templateCollection.GetAsync(job.TemplateHash);
+			if(template == null)
+			{
+				return NotFound(job.StreamId, job.TemplateId);
+			}
+
+			return new GetTemplateResponse(template).ApplyFilter(filter);
+		}
+
+		/// <summary>
+		/// Find jobs matching a criteria
+		/// </summary>
+		/// <param name="ids">The job ids to return</param>
+		/// <param name="name">Name of the job to find</param>
+		/// <param name="templates">List of templates to find</param>
+		/// <param name="streamId">The stream to search for</param>
+		/// <param name="minChange">The minimum changelist number</param>
+		/// <param name="maxChange">The maximum changelist number</param>
+		/// <param name="includePreflight">Whether to include preflight jobs</param>
+		/// <param name="preflightOnly">Whether to only include preflight jobs</param>
+		/// <param name="preflightChange">The preflighted changelist</param>
+		/// <param name="startedByUserId">User id for which to include jobs</param>
+		/// <param name="preflightStartedByUserId">User id for which to include preflight jobs</param>
+		/// <param name="minCreateTime">Minimum creation time</param>
+		/// <param name="maxCreateTime">Maximum creation time</param>
+		/// <param name="modifiedBefore">If specified, only jobs updated before the give time will be returned</param>
+		/// <param name="modifiedAfter">If specified, only jobs updated after the give time will be returned</param>
+		/// <param name="target">Target to filter the returned jobs by</param>
+		/// <param name="state">Filter state of the returned jobs</param>
+		/// <param name="outcome">Filter outcome of the returned jobs</param>
+		/// <param name="filter">Filter for properties to return</param>
+		/// <param name="index">Index of the first result to be returned</param>
+		/// <param name="count">Number of results to return</param>
+		/// <returns>List of jobs</returns>
+		[HttpGet]
+		[Route("/api/v1/jobs")]
+		[ProducesResponseType(typeof(List<GetJobResponse>), 200)]
+		public async Task<ActionResult<List<object>>> FindJobsAsync(
+			[FromQuery(Name = "Id")] string[]? ids = null,
+			[FromQuery] string? streamId = null,
+			[FromQuery] string? name = null,
+			[FromQuery(Name = "template")] string[]? templates = null,
+			[FromQuery] int? minChange = null,
+			[FromQuery] int? maxChange = null,
+			[FromQuery] bool includePreflight = true,
+			[FromQuery] bool? preflightOnly = null,
+			[FromQuery] int? preflightChange = null,
+			[FromQuery] string? preflightStartedByUserId = null,
+			[FromQuery] string? startedByUserId = null,
+			[FromQuery] DateTimeOffset? minCreateTime = null,
+			[FromQuery] DateTimeOffset? maxCreateTime = null,
+			[FromQuery] DateTimeOffset? modifiedBefore = null,
+			[FromQuery] DateTimeOffset? modifiedAfter = null,
+			[FromQuery] string? target = null,
+			[FromQuery] JobStepState[]? state = null,
+			[FromQuery] JobStepOutcome[]? outcome = null, 
+			[FromQuery] PropertyFilter? filter = null,
+			[FromQuery] int index = 0,
+			[FromQuery] int count = 100)
+		{
+			JobId[]? jobIdValues = (ids == null) ? (JobId[]?)null : Array.ConvertAll(ids, x => new JobId(x));
+			StreamId? streamIdValue = (streamId == null)? (StreamId?)null : new StreamId(streamId);
+			
+			TemplateRefId[]? templateRefIds = (templates != null && templates.Length > 0) ? templates.Select(x => new TemplateRefId(x)).ToArray() : null;
+
+			if (includePreflight == false)
+			{
+				preflightChange = 0;
+			}
+
+			UserId? preflightStartedByUserIdValue = null;
+
+			if (preflightStartedByUserId != null)
+			{
+				preflightStartedByUserIdValue = new UserId(preflightStartedByUserId);
+			}
+
+			UserId? startedByUserIdValue = null;
+
+			if (startedByUserId != null)
+			{
+				startedByUserIdValue = new UserId(startedByUserId);
+			}
+
+			List<IJob> jobs;
+			using (IScope _ = GlobalTracer.Instance.BuildSpan("FindJobs").StartActive())
+			{
+				jobs = await _jobService.FindJobsAsync(jobIdValues, streamIdValue, name, templateRefIds, minChange,
+					maxChange, preflightChange, preflightOnly, preflightStartedByUserIdValue, startedByUserIdValue, minCreateTime?.UtcDateTime, maxCreateTime?.UtcDateTime, target, state, outcome,
+					modifiedBefore, modifiedAfter, index, count, false);
+			}
+
+			StreamPermissionsCache permissionsCache = new StreamPermissionsCache();
+
+			List<object> responses = new List<object>();
+			foreach (IJob job in jobs)
+			{
+				using IScope jobScope = GlobalTracer.Instance.BuildSpan("JobIteration").StartActive();
+				jobScope.Span.SetTag("jobId", job.Id.ToString());
+
+				if (job.GraphHash == null)
+				{
+					_logger.LogWarning("Job {JobId} has no GraphHash", job.Id);
+					continue;
+				}
+
+				bool viewJobAuthorized;
+				using (IScope _ = GlobalTracer.Instance.BuildSpan("AuthorizeViewJob").StartActive())
+				{
+					viewJobAuthorized = await _jobService.AuthorizeAsync(job, AclAction.ViewJob, User, permissionsCache);
+				}
+				
+				if (viewJobAuthorized)
+				{
+					IGraph graph;
+					using (IScope _ = GlobalTracer.Instance.BuildSpan("GetGraph").StartActive())
+					{
+						graph = await _jobService.GetGraphAsync(job);
+					}
+
+					bool bIncludeAcl;
+					using (IScope _ = GlobalTracer.Instance.BuildSpan("AuthorizeViewPermissions").StartActive())
+					{
+						bIncludeAcl = await _jobService.AuthorizeAsync(job, AclAction.ViewPermissions, User, permissionsCache);
+					}
+
+					bool bIncludeCosts;
+					using (IScope _ = GlobalTracer.Instance.BuildSpan("AuthorizeViewCosts").StartActive())
+					{
+						bIncludeCosts = await _jobService.AuthorizeAsync(job, AclAction.ViewCosts, User, permissionsCache);
+					}
+
+					using (IScope _ = GlobalTracer.Instance.BuildSpan("CreateResponse").StartActive())
+					{
+						responses.Add(await CreateJobResponseAsync(job, graph, bIncludeAcl, bIncludeCosts, filter));
+					}
+				}
+			}
+			return responses;
+		}
+
+		/// <summary>
+		/// Find jobs for a stream with given templates, sorted by creation date
+		/// </summary>
+		/// <param name="streamId">The stream to search for</param>
+		/// <param name="templates">List of templates to find</param>
+		/// <param name="preflightStartedByUserId">User id for which to include preflight jobs</param>
+		/// <param name="maxCreateTime">Maximum creation time</param>
+		/// <param name="modifiedAfter">If specified, only jobs updated after the given time will be returned</param>
+		/// <param name="filter">Filter for properties to return</param>
+		/// <param name="index">Index of the first result to be returned</param>
+		/// <param name="count">Number of results to return</param>
+		/// <param name="consistentRead">If a read to the primary database is required, for read consistency. Usually not required.</param>
+		/// <returns>List of jobs</returns>
+		[HttpGet]
+		[Route("/api/v1/jobs/streams/{streamId}")]
+		[ProducesResponseType(typeof(List<GetJobResponse>), 200)]
+		public async Task<ActionResult<List<object>>> FindJobsByStreamWithTemplatesAsync(
+			string streamId,
+			[FromQuery(Name = "template")] string[] templates,
+			[FromQuery] string? preflightStartedByUserId = null,
+			[FromQuery] DateTimeOffset? maxCreateTime = null,
+			[FromQuery] DateTimeOffset? modifiedAfter = null,
+			[FromQuery] PropertyFilter? filter = null,
+			[FromQuery] int index = 0,
+			[FromQuery] int count = 100,
+			[FromQuery] bool consistentRead = false)
+		{
+			StreamId streamIdValue = new StreamId(streamId);
+			TemplateRefId[] templateRefIds = templates.Select(x => new TemplateRefId(x)).ToArray();
+			UserId? preflightStartedByUserIdValue = preflightStartedByUserId != null ? new UserId(preflightStartedByUserId) : null;
+			count = Math.Min(1000, count);
+
+			List<IJob> jobs = await _jobService.FindJobsByStreamWithTemplatesAsync(streamIdValue, templateRefIds, preflightStartedByUserIdValue, maxCreateTime, modifiedAfter, index, count, consistentRead);
+			return await CreateAuthorizedJobResponses(jobs, filter);
+		}
+
+		private async Task<List<object>> CreateAuthorizedJobResponses(List<IJob> jobs, PropertyFilter? filter = null)
+		{
+			StreamPermissionsCache permissionsCache = new ();
+			List<object> responses = new ();
+			foreach (IJob job in jobs)
+			{
+				using IScope jobScope = GlobalTracer.Instance.BuildSpan("JobIteration").StartActive();
+				jobScope.Span.SetTag("jobId", job.Id.ToString());
+				
+				bool viewJobAuthorized;
+				using (IScope _ = GlobalTracer.Instance.BuildSpan("AuthorizeViewJob").StartActive())
+				{
+					viewJobAuthorized = await _jobService.AuthorizeAsync(job, AclAction.ViewJob, User, permissionsCache);
+				}
+				
+				if (viewJobAuthorized)
+				{
+					IGraph graph;
+					using (IScope _ = GlobalTracer.Instance.BuildSpan("GetGraph").StartActive())
+					{
+						graph = await _jobService.GetGraphAsync(job);
+					}
+
+					bool bIncludeAcl;
+					using (IScope _ = GlobalTracer.Instance.BuildSpan("AuthorizeViewPermissions").StartActive())
+					{
+						bIncludeAcl = await _jobService.AuthorizeAsync(job, AclAction.ViewPermissions, User, permissionsCache);
+					}
+
+					bool bIncludeCosts;
+					using (IScope _ = GlobalTracer.Instance.BuildSpan("AuthorizeViewCosts").StartActive())
+					{
+						bIncludeCosts = await _jobService.AuthorizeAsync(job, AclAction.ViewCosts, User, permissionsCache);
+					}
+
+					using (IScope _ = GlobalTracer.Instance.BuildSpan("CreateResponse").StartActive())
+					{
+						responses.Add(await CreateJobResponseAsync(job, graph, bIncludeAcl, bIncludeCosts, filter));
+					}
+				}
+			}
+
+			
+			return responses;
+		}
+
+		/// <summary>
+		/// Adds an array of nodes to be executed for a job
+		/// </summary>
+		/// <param name="jobId">Unique id for the job</param>
+		/// <param name="requests">Properties of the new nodes</param>
+		/// <returns>Id of the new job</returns>
+		[HttpPost]
+		[Route("/api/v1/jobs/{jobId}/groups")]
+		public async Task<ActionResult> CreateGroupsAsync(JobId jobId, [FromBody] List<NewGroup> requests)
+		{
+			for (; ; )
+			{
+				IJob? job = await _jobService.GetJobAsync(jobId);
+				if (job == null)
+				{
+					return NotFound(jobId);
+				}
+				if (!await _jobService.AuthorizeAsync(job, AclAction.ExecuteJob, User, null))
+				{
+					return Forbid(AclAction.ExecuteJob, jobId);
+				}
+
+				IGraph graph = await _jobService.GetGraphAsync(job);
+				graph = await _graphs.AppendAsync(graph, requests, null, null);
+
+				IJob? newJob = await _jobService.TryUpdateGraphAsync(job, graph);
+				if (newJob != null)
+				{
+					return Ok();
+				}
+			}
+		}
+
+		/// <summary>
+		/// Gets the nodes to be executed for a job
+		/// </summary>
+		/// <param name="jobId">Unique id for the job</param>
+		/// <param name="filter">Filter for the properties to return</param>
+		/// <returns>List of nodes to be executed</returns>
+		[HttpGet]
+		[Route("/api/v1/jobs/{jobId}/groups")]
+		[ProducesResponseType(typeof(List<GetGroupResponse>), 200)]
+		public async Task<ActionResult<List<object>>> GetGroupsAsync(JobId jobId, [FromQuery] PropertyFilter? filter = null)
+		{
+			IJob? job = await _jobService.GetJobAsync(jobId);
+			if (job == null)
+			{
+				return NotFound(jobId);
+			}
+			if (!await _jobService.AuthorizeAsync(job, AclAction.ViewJob, User, null))
+			{
+				return Forbid(AclAction.ViewJob, jobId);
+			}
+
+			IGraph graph = await _jobService.GetGraphAsync(job);
+			return graph.Groups.ConvertAll(x => new GetGroupResponse(x, graph.Groups).ApplyFilter(filter));
+		}
+
+		/// <summary>
+		/// Gets the nodes in a group to be executed for a job
+		/// </summary>
+		/// <param name="jobId">Unique id for the job</param>
+		/// <param name="groupIdx">The group index</param>
+		/// <param name="filter">Filter for the properties to return</param>
+		/// <returns>List of nodes to be executed</returns>
+		[HttpGet]
+		[Route("/api/v1/jobs/{jobId}/groups/{groupIdx}")]
+		[ProducesResponseType(typeof(GetGroupResponse), 200)]
+		public async Task<ActionResult<object>> GetGroupAsync(JobId jobId, int groupIdx, [FromQuery] PropertyFilter? filter = null)
+		{
+			IJob? job = await _jobService.GetJobAsync(jobId);
+			if (job == null)
+			{
+				return NotFound(jobId);
+			}
+			if (!await _jobService.AuthorizeAsync(job, AclAction.ViewJob, User, null))
+			{
+				return Forbid(AclAction.ViewJob, jobId);
+			}
+
+			IGraph graph = await _jobService.GetGraphAsync(job);
+			if (groupIdx < 0 || groupIdx >= graph.Groups.Count)
+			{
+				return NotFound(jobId, groupIdx);
+			}
+
+			return new GetGroupResponse(graph.Groups[groupIdx], graph.Groups).ApplyFilter(filter);
+		}
+
+		/// <summary>
+		/// Gets the nodes for a particular group
+		/// </summary>
+		/// <param name="jobId">Unique id for the job</param>
+		/// <param name="groupIdx">Index of the group containing the node to update</param>
+		/// <param name="filter">Filter for the properties to return</param>
+		[HttpGet]
+		[Route("/api/v1/jobs/{jobId}/groups/{groupIdx}/nodes")]
+		[ProducesResponseType(typeof(List<GetNodeResponse>), 200)]
+		public async Task<ActionResult<List<object>>> GetNodesAsync(JobId jobId, int groupIdx, [FromQuery] PropertyFilter? filter = null)
+		{
+			IJob? job = await _jobService.GetJobAsync(jobId);
+			if (job == null)
+			{
+				return NotFound(jobId);
+			}
+			if (!await _jobService.AuthorizeAsync(job, AclAction.ViewJob, User, null))
+			{
+				return Forbid(AclAction.ViewJob, jobId);
+			}
+
+			IGraph graph = await _jobService.GetGraphAsync(job);
+			if (groupIdx < 0 || groupIdx >= graph.Groups.Count)
+			{
+				return NotFound(jobId, groupIdx);
+			}
+
+			return graph.Groups[groupIdx].Nodes.ConvertAll(x => new GetNodeResponse(x, graph.Groups).ApplyFilter(filter));
+		}
+
+		/// <summary>
+		/// Gets a particular node definition
+		/// </summary>
+		/// <param name="jobId">Unique id for the job</param>
+		/// <param name="groupIdx">Index of the group containing the node to update</param>
+		/// <param name="nodeIdx">Index of the node to update</param>
+		/// <param name="filter">Filter for the properties to return</param>
+		[HttpGet]
+		[Route("/api/v1/jobs/{jobId}/groups/{groupIdx}/nodes/{nodeIdx}")]
+		[ProducesResponseType(typeof(GetNodeResponse), 200)]
+		public async Task<ActionResult<object>> GetNodeAsync(JobId jobId, int groupIdx, int nodeIdx, [FromQuery] PropertyFilter? filter = null)
+		{
+			IJob? job = await _jobService.GetJobAsync(jobId);
+			if (job == null)
+			{
+				return NotFound(jobId);
+			}
+			if (!await _jobService.AuthorizeAsync(job, AclAction.ViewJob, User, null))
+			{
+				return Forbid(AclAction.ViewJob, jobId);
+			}
+
+			IGraph graph = await _jobService.GetGraphAsync(job);
+			if (groupIdx < 0 || groupIdx >= graph.Groups.Count)
+			{
+				return NotFound(jobId, groupIdx);
+			}
+			if(nodeIdx < 0 || nodeIdx >= graph.Groups[groupIdx].Nodes.Count)
+			{
+				return NotFound(jobId, groupIdx, nodeIdx);
+			}
+
+			return new GetNodeResponse(graph.Groups[groupIdx].Nodes[nodeIdx], graph.Groups).ApplyFilter(filter);
+		}
+
+		/// <summary>
+		/// Gets the steps currently scheduled to be executed for a job
+		/// </summary>
+		/// <param name="jobId">Unique id for the job</param>
+		/// <param name="filter">Filter for the properties to return</param>
+		/// <returns>List of nodes to be executed</returns>
+		[HttpGet]
+		[Route("/api/v1/jobs/{jobId}/batches")]
+		[ProducesResponseType(typeof(List<GetBatchResponse>), 200)]
+		public async Task<ActionResult<List<object>>> GetBatchesAsync(JobId jobId, [FromQuery] PropertyFilter? filter = null)
+		{
+			IJob? job = await _jobService.GetJobAsync(jobId);
+			if (job == null)
+			{
+				return NotFound(jobId);
+			}
+
+			StreamPermissionsCache cache = new StreamPermissionsCache();
+			if (!await _jobService.AuthorizeAsync(job, AclAction.ViewJob, User, cache))
+			{
+				return Forbid(AclAction.ViewJob, jobId);
+			}
+
+			bool bIncludeCosts = await _jobService.AuthorizeAsync(job, AclAction.ViewCosts, User, cache);
+
+			List<object> responses = new List<object>();
+			foreach (IJobStepBatch batch in job.Batches)
+			{
+				GetBatchResponse response = await CreateBatchResponseAsync(batch, bIncludeCosts);
+				responses.Add(response.ApplyFilter(filter));
+			}
+			return responses;
+		}
+
+		/// <summary>
+		/// Updates the state of a jobstep
+		/// </summary>
+		/// <param name="jobId">Unique id for the job</param>
+		/// <param name="batchId">Unique id for the step</param>
+		/// <param name="request">Updates to apply to the node</param>
+		[HttpPut]
+		[Route("/api/v1/jobs/{jobId}/batches/{batchId}")]
+		public async Task<ActionResult> UpdateBatchAsync(JobId jobId, SubResourceId batchId, [FromBody] UpdateBatchRequest request)
+		{
+			IJob? job = await _jobService.GetJobAsync(jobId);
+			if (job == null)
+			{
+				return NotFound(jobId);
+			}
+
+			IJobStepBatch? batch = job.Batches.FirstOrDefault(x => x.Id == batchId);
+			if (batch == null)
+			{
+				return NotFound(jobId, batchId);
+			}
+			if (batch.SessionId == null || !User.HasSessionClaim(batch.SessionId.Value))
+			{
+				return Forbid("Missing session claim for job {JobId} batch {BatchId}", jobId, batchId);
+			}
+
+			IJob? newJob = await _jobService.UpdateBatchAsync(job, batchId, request.LogId?.ToObjectId<ILogFile>(), request.State);
+			if (newJob == null)
+			{
+				return NotFound(jobId);
+			}
+			return Ok();
+		}
+
+		/// <summary>
+		/// Gets a particular step currently scheduled to be executed for a job
+		/// </summary>
+		/// <param name="jobId">Unique id for the job</param>
+		/// <param name="batchId">Unique id for the step</param>
+		/// <param name="filter">Filter for the properties to return</param>
+		/// <returns>List of nodes to be executed</returns>
+		[HttpGet]
+		[Route("/api/v1/jobs/{jobId}/batches/{batchId}")]
+		[ProducesResponseType(typeof(GetBatchResponse), 200)]
+		public async Task<ActionResult<object>> GetBatchAsync(JobId jobId, SubResourceId batchId, [FromQuery] PropertyFilter? filter = null)
+		{
+			IJob? job = await _jobService.GetJobAsync(jobId);
+			if (job == null)
+			{
+				return NotFound(jobId);
+			}
+
+			StreamPermissionsCache cache = new StreamPermissionsCache();
+			if (!await _jobService.AuthorizeAsync(job, AclAction.ViewJob, User, cache))
+			{
+				return Forbid(AclAction.ViewJob, jobId);
+			}
+
+			bool bIncludeCosts = await _jobService.AuthorizeAsync(job, AclAction.ViewCosts, User, cache);
+
+			foreach (IJobStepBatch batch in job.Batches)
+			{
+				if (batch.Id == batchId)
+				{
+					GetBatchResponse response = await CreateBatchResponseAsync(batch, bIncludeCosts);
+					return response.ApplyFilter(filter);
+				}
+			}
+
+			return NotFound(jobId, batchId);
+		}
+
+		/// <summary>
+		/// Gets the steps currently scheduled to be executed for a job
+		/// </summary>
+		/// <param name="jobId">Unique id for the job</param>
+		/// <param name="batchId">Unique id for the batch</param>
+		/// <param name="filter">Filter for the properties to return</param>
+		/// <returns>List of nodes to be executed</returns>
+		[HttpGet]
+		[Route("/api/v1/jobs/{jobId}/batches/{batchId}/steps")]
+		[ProducesResponseType(typeof(List<GetStepResponse>), 200)]
+		public async Task<ActionResult<List<object>>> GetStepsAsync(JobId jobId, SubResourceId batchId, [FromQuery] PropertyFilter? filter = null)
+		{
+			IJob? job = await _jobService.GetJobAsync(jobId);
+			if (job == null)
+			{
+				return NotFound(jobId);
+			}
+			if (!await _jobService.AuthorizeAsync(job, AclAction.ViewJob, User, null))
+			{
+				return Forbid(AclAction.ViewJob, jobId);
+			}
+
+			foreach (IJobStepBatch batch in job.Batches)
+			{
+				if (batch.Id == batchId)
+				{
+					List<object> responses = new List<object>();
+					foreach (IJobStep step in batch.Steps)
+					{
+						GetStepResponse response = await CreateStepResponseAsync(step);
+						responses.Add(response.ApplyFilter(filter));
+					}
+					return responses;
+				}
+			}
+
+			return NotFound(jobId, batchId);
+		}
+
+		/// <summary>
+		/// Updates the state of a jobstep
+		/// </summary>
+		/// <param name="jobId">Unique id for the job</param>
+		/// <param name="batchId">Unique id for the batch</param>
+		/// <param name="stepId">Unique id for the step</param>
+		/// <param name="request">Updates to apply to the node</param>
+		[HttpPut]
+		[Route("/api/v1/jobs/{jobId}/batches/{batchId}/steps/{stepId}")]
+		public async Task<ActionResult<UpdateStepResponse>> UpdateStepAsync(JobId jobId, SubResourceId batchId, SubResourceId stepId, [FromBody] UpdateStepRequest request)
+		{
+			IJob? job = await _jobService.GetJobAsync(jobId);
+			if (job == null)
+			{
+				return NotFound(jobId);
+			}
+
+			// Check permissions for updating this step. Only the agent executing the step can modify the state of it.
+			if (request.State != JobStepState.Unspecified || request.Outcome != JobStepOutcome.Unspecified)
+			{
+				IJobStepBatch? batch = job.Batches.FirstOrDefault(x => x.Id == batchId);
+				if (batch == null)
+				{
+					return NotFound(jobId, batchId);
+				}
+				if (!batch.SessionId.HasValue || !User.HasSessionClaim(batch.SessionId.Value))
+				{
+					return Forbid();
+				}
+			}
+			if (request.Retry != null || request.Priority != null)
+			{
+				if (!await _jobService.AuthorizeAsync(job, AclAction.RetryJobStep, User, null))
+				{
+					return Forbid(AclAction.RetryJobStep, jobId);
+				}
+			}
+			if (request.Properties != null)
+			{
+				if (!await _jobService.AuthorizeAsync(job, AclAction.UpdateJob, User, null))
+				{
+					return Forbid(AclAction.UpdateJob, jobId);
+				}
+			}
+
+			UserId? retryByUser = (request.Retry.HasValue && request.Retry.Value) ? User.GetUserId() : null;
+			UserId? abortByUser = (request.AbortRequested.HasValue && request.AbortRequested.Value) ? User.GetUserId() : null;
+
+			try
+			{
+				IJob? newJob = await _jobService.UpdateStepAsync(job, batchId, stepId, request.State, request.Outcome, null, request.AbortRequested, abortByUser, request.LogId?.ToObjectId<ILogFile>(), null, retryByUser, request.Priority, null, request.Properties);
+				if (newJob == null)
+				{
+					return NotFound(jobId);
+				}
+
+				UpdateStepResponse response = new UpdateStepResponse();
+				if (request.Retry ?? false)
+				{
+					JobStepRefId? retriedStepId = FindRetriedStep(job, batchId, stepId);
+					if (retriedStepId != null)
+					{
+						response.BatchId = retriedStepId.Value.BatchId.ToString();
+						response.StepId = retriedStepId.Value.StepId.ToString();
+					}
+				}
+				return response;
+			}
+			catch (RetryNotAllowedException ex)
+			{
+				return BadRequest(ex.Message);
+			}
+		}
+
+		/// <summary>
+		/// Find the first retried step after the given step
+		/// </summary>
+		/// <param name="job">The job being run</param>
+		/// <param name="batchId">Batch id of the last step instance</param>
+		/// <param name="stepId">Step id of the last instance</param>
+		/// <returns>The retried step information</returns>
+		static JobStepRefId? FindRetriedStep(IJob job, SubResourceId batchId, SubResourceId stepId)
+		{
+			NodeRef? lastNodeRef = null;
+			foreach (IJobStepBatch batch in job.Batches)
+			{
+				if ((lastNodeRef == null && batch.Id == batchId) || (lastNodeRef != null && batch.GroupIdx == lastNodeRef.GroupIdx))
+				{
+					foreach (IJobStep step in batch.Steps)
+					{
+						if (lastNodeRef == null && step.Id == stepId)
+						{
+							lastNodeRef = new NodeRef(batch.GroupIdx, step.NodeIdx);
+						}
+						else if (lastNodeRef != null && step.NodeIdx == lastNodeRef.NodeIdx)
+						{
+							return new JobStepRefId(job.Id, batch.Id, step.Id);
+						}
+					}
+				}
+			}
+			return null;
+		}
+
+		/// <summary>
+		/// Gets a particular step currently scheduled to be executed for a job
+		/// </summary>
+		/// <param name="jobId">Unique id for the job</param>
+		/// <param name="batchId">Unique id for the batch</param>
+		/// <param name="stepId">Unique id for the step</param>
+		/// <param name="filter">Filter for the properties to return</param>
+		/// <returns>List of nodes to be executed</returns>
+		[HttpGet]
+		[Route("/api/v1/jobs/{jobId}/batches/{batchId}/steps/{stepId}")]
+		[ProducesResponseType(typeof(GetStepResponse), 200)]
+		public async Task<ActionResult<object>> GetStepAsync(JobId jobId, SubResourceId batchId, SubResourceId stepId, [FromQuery] PropertyFilter? filter = null)
+		{
+			IJob? job = await _jobService.GetJobAsync(jobId);
+			if (job == null)
+			{
+				return NotFound(jobId);
+			}
+			if (!await _jobService.AuthorizeAsync(job, AclAction.ViewJob, User, null))
+			{
+				return Forbid(AclAction.ViewJob, jobId);
+			}
+
+			foreach (IJobStepBatch batch in job.Batches)
+			{
+				if (batch.Id == batchId)
+				{
+					foreach (IJobStep step in batch.Steps)
+					{
+						if (step.Id == stepId)
+						{
+							GetStepResponse response = await CreateStepResponseAsync(step);
+							return response.ApplyFilter(filter);
+						}
+					}
+					return NotFound(jobId, batchId, stepId);
+				}
+			}
+
+			return NotFound(jobId, batchId);
+		}
+
+		/// <summary>
+		/// Updates notifications for a specific job.
+		/// </summary>
+		/// <param name="jobId">Unique id for the job</param>
+		/// <param name="batchId">Unique id for the batch</param>
+		/// <param name="stepId">Unique id for the step</param>
+		/// <param name="request">The notification request</param>
+		/// <returns>Information about the requested job</returns>
+		[HttpPut]
+		[Route("/api/v1/jobs/{jobId}/batches/{batchId}/steps/{stepId}/notifications")]
+		public async Task<ActionResult> UpdateStepNotificationsAsync(JobId jobId, SubResourceId batchId, SubResourceId stepId, [FromBody] UpdateNotificationsRequest request)
+		{
+			IJob? job = await _jobService.GetJobAsync(jobId);
+			if (job == null)
+			{
+				return NotFound(jobId);
+			}
+			if (!await _jobService.AuthorizeAsync(job, AclAction.CreateSubscription, User, null))
+			{
+				return Forbid(AclAction.CreateSubscription, jobId);
+			}
+
+			if (!job.TryGetBatch(batchId, out IJobStepBatch? batch))
+			{
+				return NotFound(jobId, batchId);
+			}
+
+			if (!batch.TryGetStep(stepId, out IJobStep? step))
+			{
+				return NotFound(jobId, batchId, stepId);
+			}
+
+			ObjectId? triggerId = step.NotificationTriggerId;
+			if (triggerId == null)
+			{
+				triggerId = ObjectId.GenerateNewId();
+				if (await _jobService.UpdateStepAsync(job, batchId, stepId, JobStepState.Unspecified, JobStepOutcome.Unspecified, newNotificationTriggerId: triggerId) == null)
+				{
+					return NotFound(jobId, batchId, stepId);
+				}
+			}
+
+			await _notificationService.UpdateSubscriptionsAsync(triggerId.Value, User, request.Email, request.Slack);
+			return Ok();
+		}
+
+		/// <summary>
+		/// Gets information about a specific job.
+		/// </summary>
+		/// <param name="jobId">Id of the job to find</param>
+		/// <param name="batchId">Unique id for the batch</param>
+		/// <param name="stepId">Unique id for the step</param>
+		/// <returns>Information about the requested job</returns>
+		[HttpGet]
+		[Route("/api/v1/jobs/{jobId}/batches/{batchId}/steps/{stepId}/notifications")]
+		public async Task<ActionResult<GetNotificationResponse>> GetStepNotificationsAsync(JobId jobId, SubResourceId batchId, SubResourceId stepId)
+		{
+			IJob? job = await _jobService.GetJobAsync(jobId);
+			if (job == null)
+			{
+				return NotFound(jobId);
+			}
+
+			IJobStep? step;
+			if (!job.TryGetBatch(batchId, out IJobStepBatch? batch))
+			{
+				return NotFound(jobId, batchId);
+			}
+			if (!batch.TryGetStep(stepId, out step))
+			{
+				return NotFound(jobId, batchId, stepId);
+			}
+			if (!await _jobService.AuthorizeAsync(job, AclAction.CreateSubscription, User, null))
+			{
+				return Forbid(AclAction.CreateSubscription, jobId);
+			}
+
+			INotificationSubscription? subscription;
+			if (step.NotificationTriggerId == null)
+			{
+				subscription = null;
+			}
+			else
+			{
+				subscription = await _notificationService.GetSubscriptionsAsync(step.NotificationTriggerId.Value, User);
+			}
+			return new GetNotificationResponse(subscription);
+		}
+
+		/// <summary>
+		/// Gets a particular step currently scheduled to be executed for a job
+		/// </summary>
+		/// <param name="jobId">Unique id for the job</param>
+		/// <param name="batchId">Unique id for the batch</param>
+		/// <param name="stepId">Unique id for the step</param>
+		/// <param name="name"></param>
+		/// <returns>List of nodes to be executed</returns>
+		[HttpGet]
+		[Route("/api/v1/jobs/{jobId}/batches/{batchId}/steps/{stepId}/artifacts/{*name}")]
+		public async Task<ActionResult> GetArtifactAsync(JobId jobId, SubResourceId batchId, SubResourceId stepId, string name)
+		{
+			IJob? job = await _jobService.GetJobAsync(jobId);
+			if (job == null)
+			{
+				return NotFound(jobId);
+			}
+			if (!await _jobService.AuthorizeAsync(job, AclAction.ViewJob, User, null))
+			{
+				return Forbid(AclAction.ViewJob, jobId);
+			}
+			if (!job.TryGetBatch(batchId, out IJobStepBatch? batch))
+			{
+				return NotFound(jobId, batchId);
+			}
+			if (!batch.TryGetStep(stepId, out _))
+			{
+				return NotFound(jobId, batchId, stepId);
+			}
+
+			List<IArtifact> artifacts = await _artifactCollection.GetArtifactsAsync(jobId, stepId, name);
+			if (artifacts.Count == 0)
+			{
+				return NotFound();
+			}
+
+			IArtifact artifact = artifacts[0];
+			return new FileStreamResult(await _artifactCollection.OpenArtifactReadStreamAsync(artifact), artifact.MimeType);
+		}
+
+		/// <summary>
+		/// Gets a particular step currently scheduled to be executed for a job
+		/// </summary>
+		/// <param name="jobId">Unique id for the job</param>
+		/// <param name="batchId">Unique id for the batch</param>
+		/// <param name="stepId">Unique id for the step</param>
+		/// <returns>List of nodes to be executed</returns>
+		[HttpGet]
+		[Route("/api/v1/jobs/{jobId}/batches/{batchId}/steps/{stepId}/trace")]
+		public async Task<ActionResult> GetStepTraceAsync(JobId jobId, SubResourceId batchId, SubResourceId stepId)
+		{
+			IJob? job = await _jobService.GetJobAsync(jobId);
+			if (job == null)
+			{
+				return NotFound(jobId);
+			}
+			if (!await _jobService.AuthorizeAsync(job, AclAction.ViewJob, User, null))
+			{
+				return Forbid(AclAction.ViewJob, jobId);
+			}
+			if (!job.TryGetBatch(batchId, out IJobStepBatch? batch))
+			{
+				return NotFound(jobId, batchId);
+			}
+			if (!batch.TryGetStep(stepId, out _))
+			{
+				return NotFound(jobId, batchId, stepId);
+			}
+
+			List<IArtifact> artifacts = await _artifactCollection.GetArtifactsAsync(jobId, stepId, null);
+			foreach (IArtifact artifact in artifacts)
+			{
+				if (artifact.Name.Equals("trace.json", StringComparison.OrdinalIgnoreCase))
+				{
+					return new FileStreamResult(await _artifactCollection.OpenArtifactReadStreamAsync(artifact), "text/json");
+				}
+			}
+			return NotFound();
+		}
+
+		/// <summary>
+		/// Updates notifications for a specific label.
+		/// </summary>
+		/// <param name="jobId">Unique id for the job</param>
+		/// <param name="labelIndex">Index for the label</param>
+		/// <param name="request">The notification request</param>
+		[HttpPut]
+		[Route("/api/v1/jobs/{jobId}/labels/{labelIndex}/notifications")]
+		public async Task<ActionResult> UpdateLabelNotificationsAsync(JobId jobId, int labelIndex, [FromBody] UpdateNotificationsRequest request)
+		{
+			ObjectId triggerId;
+			for (; ; )
+			{
+				IJob? job = await _jobService.GetJobAsync(jobId);
+				if (job == null)
+				{
+					return NotFound(jobId);
+				}
+				if (!await _jobService.AuthorizeAsync(job, AclAction.CreateSubscription, User, null))
+				{
+					return Forbid(AclAction.CreateSubscription, jobId);
+				}
+
+				ObjectId newTriggerId;
+				if (job.LabelIdxToTriggerId.TryGetValue(labelIndex, out newTriggerId))
+				{
+					triggerId = newTriggerId;
+					break;
+				}
+
+				newTriggerId = ObjectId.GenerateNewId();
+
+				IJob? newJob = await _jobService.UpdateJobAsync(job, labelIdxToTriggerId: new KeyValuePair<int, ObjectId>(labelIndex, newTriggerId));
+				if (newJob != null)
+				{
+					triggerId = newTriggerId;
+					break;
+				}
+			}
+
+			await _notificationService.UpdateSubscriptionsAsync(triggerId, User, request.Email, request.Slack);
+			return Ok();
+		}
+
+		/// <summary>
+		/// Gets notification info about a specific label in a job.
+		/// </summary>
+		/// <param name="jobId">Id of the job to find</param>
+		/// <param name="labelIndex">Index for the label</param>
+		/// <returns>Notification info for the requested label in the job</returns>
+		[HttpGet]
+		[Route("/api/v1/jobs/{jobId}/labels/{labelIndex}/notifications")]
+		public async Task<ActionResult<GetNotificationResponse>> GetLabelNotificationsAsync(JobId jobId, int labelIndex)
+		{
+			IJob? job = await _jobService.GetJobAsync(jobId);
+			if (job == null)
+			{
+				return NotFound(jobId);
+			}
+
+			INotificationSubscription? subscription;
+			if (!job.LabelIdxToTriggerId.ContainsKey(labelIndex))
+			{
+				subscription = null;
+			}
+			else
+			{
+				subscription = await _notificationService.GetSubscriptionsAsync(job.LabelIdxToTriggerId[labelIndex], User);
+			}
+			return new GetNotificationResponse(subscription);
+		}
+	}
+}
diff --git a/Engine/Source/Programs/Horde/Horde.Build/Jobs/Schedules/ISchedule.cs b/Engine/Source/Programs/Horde/Horde.Build/Jobs/Schedules/ISchedule.cs
new file mode 100644
--- /dev/null
+++ b/Engine/Source/Programs/Horde/Horde.Build/Jobs/Schedules/ISchedule.cs
@@ -0,0 +1,331 @@
+// Copyright Epic Games, Inc. All Rights Reserved.
+
+using System;
+using System.Collections.Generic;
+using Horde.Build.Utilities;
+using MongoDB.Bson.Serialization.Attributes;
+
+namespace Horde.Build.Models
+{
+	using JobId = ObjectId<IJob>;
+	using TemplateRefId = StringId<TemplateRef>;
+
+	/// <summary>
+	/// Specifies a pattern of times that this schedule should run. Each schedule may have multiple patterns.
+	/// </summary>
+	public class SchedulePattern
+	{
+		/// <summary>
+		/// Which days of the week the schedule should run
+		/// </summary>
+		public List<DayOfWeek>? DaysOfWeek { get; set; }
+
+		/// <summary>
+		/// Time during the day for the first schedule to trigger. Measured in minutes from midnight.
+		/// </summary>
+		public int MinTime { get; set; }
+
+		/// <summary>
+		/// Time during the day for the last schedule to trigger. Measured in minutes from midnight.
+		/// </summary>
+		public int? MaxTime { get; set; }
+
+		/// <summary>
+		/// Interval between each schedule triggering
+		/// </summary>
+		public int? Interval { get; set; }
+
+		/// <summary>
+		/// Private constructor for serialization
+		/// </summary>
+		private SchedulePattern()
+		{
+		}
+
+		/// <summary>
+		/// Constructor
+		/// </summary>
+		/// <param name="daysOfWeek">Which days of the week the schedule should run</param>
+		/// <param name="minTime">Time during the day for the first schedule to trigger. Measured in minutes from midnight.</param>
+		/// <param name="maxTime">Time during the day for the last schedule to trigger. Measured in minutes from midnight.</param>
+		/// <param name="interval">Interval between each schedule triggering</param>
+		public SchedulePattern(List<DayOfWeek>? daysOfWeek, int minTime, int? maxTime, int? interval)
+		{
+			DaysOfWeek = daysOfWeek;
+			MinTime = minTime;
+			MaxTime = maxTime;
+			Interval = interval;
+		}
+
+		/// <summary>
+		/// Calculates the trigger index based on the given time in minutes
+		/// </summary>
+		/// <param name="lastTimeUtc">Time for the last trigger</param>
+		/// <param name="timeZone">The timezone for running the schedule</param>
+		/// <returns>Index of the trigger</returns>
+		public DateTime GetNextTriggerTimeUtc(DateTime lastTimeUtc, TimeZoneInfo timeZone)
+		{
+			// Convert last time into the correct timezone for running the scheule
+			DateTimeOffset lastTime = TimeZoneInfo.ConvertTime((DateTimeOffset)lastTimeUtc, timeZone);
+
+			// Get the base time (ie. the start of this day) for anchoring the schedule
+			DateTimeOffset baseTime = new DateTimeOffset(lastTime.Year, lastTime.Month, lastTime.Day, 0, 0, 0, lastTime.Offset);
+			for (; ; )
+			{
+				if (DaysOfWeek == null || DaysOfWeek.Contains(baseTime.DayOfWeek))
+				{
+					// Get the last time in minutes from the start of this day
+					int lastTimeMinutes = (int)(lastTime - baseTime).TotalMinutes;
+
+					// Get the time of the first trigger of this day. If the last time is less than this, this is the next trigger.
+					if (lastTimeMinutes < MinTime)
+					{
+						return baseTime.AddMinutes(MinTime).UtcDateTime;
+					}
+
+					// Otherwise, get the time for the last trigger in the day.
+					if (Interval.HasValue && Interval.Value > 0)
+					{
+						int actualMaxTime = MaxTime ?? ((24 * 60) - 1);
+						if (lastTimeMinutes < actualMaxTime)
+						{
+							int lastIndex = (lastTimeMinutes - MinTime) / Interval.Value;
+							int nextIndex = lastIndex + 1;
+
+							int nextTimeMinutes = MinTime + (nextIndex * Interval.Value);
+							if (nextTimeMinutes <= actualMaxTime)
+							{
+								return baseTime.AddMinutes(nextTimeMinutes).UtcDateTime;
+							}
+						}
+					}
+				}
+				baseTime = baseTime.AddDays(1.0);
+			}
+		}
+	}
+
+	/// <summary>
+	/// Gate allowing a schedule to trigger.
+	/// </summary>
+	public class ScheduleGate
+	{
+		/// <summary>
+		/// The template containing the dependency
+		/// </summary>
+		public TemplateRefId TemplateRefId { get; set; }
+
+		/// <summary>
+		/// Target to wait for
+		/// </summary>
+		public string Target { get; set; }
+
+		/// <summary>
+		/// Private constructor for serializtaion
+		/// </summary>
+		[BsonConstructor]
+		private ScheduleGate()
+		{
+			Target = String.Empty;
+		}
+
+		/// <summary>
+		/// Constructor
+		/// </summary>
+		/// <param name="templateRefId">The template containing the dependency</param>
+		/// <param name="target">Target to wait for</param>
+		public ScheduleGate(TemplateRefId templateRefId, string target)
+		{
+			TemplateRefId = templateRefId;
+			Target = target;
+		}
+	}
+
+	/// <summary>
+	/// Represents a schedule
+	/// </summary>
+	public class Schedule
+	{
+		/// <summary>
+		/// Whether this schedule is currently enabled
+		/// </summary>
+		public bool Enabled { get; set; }
+
+		/// <summary>
+		/// Maximum number of builds triggered by this schedule active at once. Set to zero for unlimited.
+		/// </summary>
+		public int MaxActive { get; set; }
+
+		/// <summary>
+		/// Maximum number of changes the schedule can fall behind head revision. If greater than zero, builds will be triggered for every submitted changelist until the backlog is this size.
+		/// </summary>
+		public int MaxChanges { get; set; }
+
+		/// <summary>
+		/// Whether the build requires a change to be submitted
+		/// </summary>
+		[BsonIgnoreIfDefault, BsonDefaultValue(true)]
+		public bool RequireSubmittedChange { get; set; } = true;
+
+		/// <summary>
+		/// Reference to another job/target which much succeed for this schedule to trigger
+		/// </summary>
+		public ScheduleGate? Gate { get; set; }
+
+		/// <summary>
+		/// Whether this build requires a code change to trigger
+		/// </summary>
+		public List<ChangeContentFlags>? Filter { get; set; }
+
+		/// <summary>
+		/// File paths which should trigger this schedule
+		/// </summary>
+		public List<string>? Files { get; set; }
+
+		/// <summary>
+		/// Parameters for the template
+		/// </summary>
+		public Dictionary<string, string> TemplateParameters { get; set; }
+
+		/// <summary>
+		/// List of patterns for this schedule
+		/// </summary>
+		public List<SchedulePattern> Patterns { get; set; }
+
+		/// <summary>
+		/// Last changelist number that this was triggered for
+		/// </summary>
+		public int LastTriggerChange { get; set; }
+
+		/// <summary>
+		/// Last time that the schedule was triggered
+		/// </summary>
+		[BsonIgnoreIfNull]
+		public DateTimeOffset? LastTriggerTime { get; set; }
+
+		/// <summary>
+		/// Gets the last trigger time, in UTC
+		/// </summary>
+		[System.Diagnostics.CodeAnalysis.SuppressMessage("Naming", "CA1721:Property names should not match get methods", Justification = "<Pending>")]
+		public DateTime LastTriggerTimeUtc { get; set; }
+
+		/// <summary>
+		/// List of jobs that are currently active
+		/// </summary>
+		public List<JobId> ActiveJobs { get; set; } = new List<JobId>();
+
+		/// <summary>
+		/// Private constructor for serializtaion
+		/// </summary>
+		[BsonConstructor]
+		private Schedule()
+		{
+			TemplateParameters = new Dictionary<string, string>();
+			Patterns = new List<SchedulePattern>();
+		}
+
+		/// <summary>
+		/// Constructor
+		/// </summary>
+		/// <param name="currentTimeUtc">The current time. This will be used to seed the last trigger time.</param>
+		/// <param name="enabled">Whether the schedule is currently enabled</param>
+		/// <param name="maxActive">Maximum number of builds that may be active at once</param>
+		/// <param name="maxChanges">Maximum number of changes the schedule can fall behind head revision</param>
+		/// <param name="requireSubmittedChange">Whether a change has to be submitted for the schedule to trigger</param>
+		/// <param name="gate">Reference to another job/target which much succeed for this schedule to trigger</param>
+		/// <param name="filter">Filter for changes to consider</param>
+		/// <param name="files">Files that should trigger the schedule</param>
+		/// <param name="templateParameters">Parameters for the template to run</param>
+		/// <param name="patterns">List of patterns for the schedule</param>
+		public Schedule(DateTime currentTimeUtc, bool enabled = true, int maxActive = 0, int maxChanges = 0, bool requireSubmittedChange = true, ScheduleGate? gate = null, List<ChangeContentFlags>? filter = null, List<string>? files = null, Dictionary<string, string>? templateParameters = null, List<SchedulePattern>? patterns = null)
+		{
+			Enabled = enabled;
+			MaxActive = maxActive;
+			MaxChanges = maxChanges;
+			RequireSubmittedChange = requireSubmittedChange;
+			Gate = gate;
+			Filter = filter;
+			Files = files;
+			TemplateParameters = templateParameters ?? new Dictionary<string, string>();
+			Patterns = patterns ?? new List<SchedulePattern>();
+			LastTriggerTimeUtc = currentTimeUtc;
+		}
+
+		/// <summary>
+		/// Copies the state fields from another schedule object
+		/// </summary>
+		/// <param name="other">The schedule object to copy from</param>
+		public void CopyState(Schedule other)
+		{
+			LastTriggerChange = other.LastTriggerChange;
+			LastTriggerTime = other.LastTriggerTime;
+			ActiveJobs.AddRange(other.ActiveJobs);
+		}
+
+		/// <summary>
+		/// Gets the flags to filter changes by
+		/// </summary>
+		/// <returns>Set of filter flags</returns>
+		public ChangeContentFlags? GetFilterFlags()
+		{
+			if(Filter == null)
+			{
+				return null;
+			}
+
+			ChangeContentFlags flags = 0;
+			foreach (ChangeContentFlags flag in Filter)
+			{
+				flags |= flag;
+			}
+			return flags;
+		}
+
+		/// <summary>
+		/// Gets the last time that the schedule triggered
+		/// </summary>
+		/// <returns>Last trigger time</returns>
+		public DateTime GetLastTriggerTimeUtc()
+		{
+			if (LastTriggerTime != null)
+			{
+				LastTriggerTimeUtc = LastTriggerTime.Value.UtcDateTime;
+				LastTriggerTime = null;
+			}
+			return LastTriggerTimeUtc;
+		}
+
+		/// <summary>
+		/// Get the next time that the schedule will trigger
+		/// </summary>
+		/// <param name="timeZone">Timezone to evaluate the trigger</param>
+		/// <returns>Next time at which the schedule will trigger</returns>
+		public DateTime? GetNextTriggerTimeUtc(TimeZoneInfo timeZone)
+		{
+			return GetNextTriggerTimeUtc(GetLastTriggerTimeUtc(), timeZone);
+		}
+
+		/// <summary>
+		/// Get the next time that the schedule will trigger
+		/// </summary>
+		/// <param name="lastTimeUtc">Last time at which the schedule triggered</param>
+		/// <param name="timeZone">Timezone to evaluate the trigger</param>
+		/// <returns>Next time at which the schedule will trigger</returns>
+		public DateTime? GetNextTriggerTimeUtc(DateTime lastTimeUtc, TimeZoneInfo timeZone)
+		{
+			DateTime? nextTriggerTimeUtc = null;
+			if (Enabled)
+			{
+				foreach (SchedulePattern pattern in Patterns)
+				{
+					DateTime patternTriggerTime = pattern.GetNextTriggerTimeUtc(lastTimeUtc, timeZone);
+					if (nextTriggerTimeUtc == null || patternTriggerTime < nextTriggerTimeUtc)
+					{
+						nextTriggerTimeUtc = patternTriggerTime;
+					}
+				}
+			}
+			return nextTriggerTimeUtc;
+		}
+	}
+}
diff --git a/Engine/Source/Programs/Horde/Horde.Build/Jobs/Schedules/ScheduleService.cs b/Engine/Source/Programs/Horde/Horde.Build/Jobs/Schedules/ScheduleService.cs
new file mode 100644
--- /dev/null
+++ b/Engine/Source/Programs/Horde/Horde.Build/Jobs/Schedules/ScheduleService.cs
@@ -0,0 +1,598 @@
+// Copyright Epic Games, Inc. All Rights Reserved.
+
+using System;
+using System.Collections.Generic;
+using System.Diagnostics;
+using System.Linq;
+using System.Text.RegularExpressions;
+using System.Threading;
+using System.Threading.Tasks;
+using EpicGames.Core;
+using EpicGames.Redis;
+using EpicGames.Redis.Utility;
+using EpicGames.Serialization;
+using Horde.Build.Api;
+using Horde.Build.Collections;
+using Horde.Build.Models;
+using Horde.Build.Server;
+using Horde.Build.Utilities;
+using HordeCommon;
+using Microsoft.Extensions.Hosting;
+using Microsoft.Extensions.Logging;
+using Microsoft.Extensions.Options;
+using MongoDB.Driver;
+using OpenTracing;
+using OpenTracing.Util;
+using StackExchange.Redis;
+using TimeZoneConverter;
+
+namespace Horde.Build.Services
+{
+	using JobId = ObjectId<IJob>;
+	using StreamId = StringId<IStream>;
+	using TemplateRefId = StringId<TemplateRef>;
+
+	/// <summary>
+	/// Manipulates schedule instances
+	/// </summary>
+	public sealed class ScheduleService : IHostedService, IDisposable
+	{
+		[RedisConverter(typeof(RedisCbConverter<QueueItem>))]
+		class QueueItem
+		{
+			[CbField("sid")]
+			public StreamId StreamId { get; set; }
+
+			[CbField("tid")]
+			public TemplateRefId TemplateId { get; set; }
+
+			public QueueItem()
+			{
+			}
+
+			public QueueItem(StreamId streamId, TemplateRefId templateId)
+			{
+				StreamId = streamId;
+				TemplateId = templateId;
+			}
+
+			public static DateTime GetTimeFromScore(double score) => DateTime.UnixEpoch + TimeSpan.FromSeconds(score);
+			public static double GetScoreFromTime(DateTime time) => (time.ToUniversalTime() - DateTime.UnixEpoch).TotalSeconds;
+		}
+
+		class PerforceHistory
+		{
+			readonly IPerforceService _perforce;
+			readonly string _clusterName;
+			readonly string _streamName;
+			int _maxResults;
+			readonly List<ChangeDetails> _changes = new List<ChangeDetails>();
+			int _nextIndex;
+
+			public PerforceHistory(IPerforceService perforce, string clusterName, string streamName)
+			{
+				_perforce = perforce;
+				_clusterName = clusterName;
+				_streamName = streamName;
+				_maxResults = 10;
+			}
+
+			async ValueTask<ChangeDetails?> GetChangeAtIndexAsync(int index)
+			{
+				while (index >= _changes.Count)
+				{
+					List<ChangeSummary> newChanges = await _perforce.GetChangesAsync(_clusterName, _streamName, null, null, _maxResults, null);
+
+					int numResults = newChanges.Count;
+					if (_changes.Count > 0)
+					{
+						newChanges.RemoveAll(x => x.Number >= _changes[^1].Number);
+					}
+					if (newChanges.Count == 0 && numResults < _maxResults)
+					{
+						return null;
+					}
+					if(newChanges.Count > 0)
+					{
+						_changes.AddRange((await _perforce.GetChangeDetailsAsync(_clusterName, _streamName, newChanges.ConvertAll(x => x.Number), null)).OrderByDescending(x => x.Number));
+					}
+					_maxResults += 10;
+				}
+				return _changes[index];
+			}
+
+			public async ValueTask<ChangeDetails?> GetNextChangeAsync()
+			{
+				ChangeDetails? details = await GetChangeAtIndexAsync(_nextIndex);
+				if (details != null)
+				{
+					_nextIndex++;
+				}
+				return details;
+			}
+
+			public async ValueTask<int> GetCodeChange(int change)
+			{
+				int index = _changes.BinarySearch(x => -x.Number, -change);
+				if (index < 0)
+				{
+					index = ~index;
+				}
+
+				for (; ; )
+				{
+					ChangeDetails? details = await GetChangeAtIndexAsync(index);
+					if (details == null)
+					{
+						return 0;
+					}
+					if ((details.GetContentFlags() & ChangeContentFlags.ContainsCode) != 0)
+					{
+						return details.Number;
+					}
+					index++;
+				}
+			}
+		}
+
+		readonly IGraphCollection _graphs;
+		readonly IPerforceService _perforce;
+		readonly IJobCollection _jobCollection;
+		readonly JobService _jobService;
+		readonly StreamService _streamService;
+		readonly ITemplateCollection _templateCollection;
+		readonly TimeZoneInfo _timeZone;
+		readonly IClock _clock;
+		readonly ILogger _logger;
+		readonly RedisService _redis;
+		readonly RedisKey _baseLockKey;
+		readonly RedisKey _tickLockKey; // Lock to tick the queue
+		readonly RedisSortedSet<QueueItem> _queue; // Items to tick, ordered by time
+		readonly ITicker _ticker;
+
+		/// <summary>
+		/// Constructor
+		/// </summary>
+		public ScheduleService(RedisService redis, IGraphCollection graphs, IPerforceService perforce, IJobCollection jobCollection, JobService jobService, StreamService streamService, ITemplateCollection templateCollection, MongoService mongoService, IClock clock, IOptionsMonitor<ServerSettings> settings, ILogger<ScheduleService> logger)
+		{
+			_graphs = graphs;
+			_perforce = perforce;
+			_jobCollection = jobCollection;
+			_jobService = jobService;
+			_streamService = streamService;
+			_templateCollection = templateCollection;
+			_clock = clock;
+
+			string? scheduleTimeZone = settings.CurrentValue.ScheduleTimeZone;
+			_timeZone = (scheduleTimeZone == null) ? TimeZoneInfo.Local : TZConvert.GetTimeZoneInfo(scheduleTimeZone);
+
+			_logger = logger;
+
+			_redis = redis;
+			_baseLockKey = "scheduler/locks";
+			_tickLockKey = _baseLockKey.Append("/tick");
+			_queue = new RedisSortedSet<QueueItem>(redis.Database, "scheduler/queue");
+			if (mongoService.ReadOnlyMode)
+			{
+				_ticker = new NullTicker();
+			}
+			else
+			{
+				_ticker = clock.AddTicker<ScheduleService>(TimeSpan.FromMinutes(1.0), TickAsync, logger);
+			}
+		}
+
+		/// <inheritdoc/>
+		public async Task StartAsync(CancellationToken cancellationToken)
+		{
+			await _ticker.StartAsync();
+		}
+
+		/// <inheritdoc/>
+		public async Task StopAsync(CancellationToken cancellationToken)
+		{
+			await _ticker.StopAsync();
+		}
+
+		/// <inheritdoc/>
+		public void Dispose()
+		{
+			_ticker.Dispose();
+		}
+
+		async ValueTask TickAsync(CancellationToken cancellationToken)
+		{
+			DateTime utcNow = _clock.UtcNow;
+
+			// Update the current queue
+			await using (RedisLock sharedLock = new RedisLock(_redis.Database, _tickLockKey))
+			{
+				if (await sharedLock.AcquireAsync(TimeSpan.FromMinutes(1.0), false))
+				{
+					await UpdateQueueAsync(utcNow);
+				}
+			}
+
+			// Keep updating schedules
+			while (!cancellationToken.IsCancellationRequested)
+			{
+				// Get the item with the lowest score (ie. the one that hasn't been updated in the longest time)
+				QueueItem? item = await PopQueueItemAsync();
+				if (item == null)
+				{
+					break;
+				}
+
+				// Acquire the lock for this schedule and update it
+				await using (RedisLock sharedLock = new RedisLock<QueueItem>(_redis.Database, _baseLockKey, item))
+				{
+					if (await sharedLock.AcquireAsync(TimeSpan.FromMinutes(1.0)))
+					{
+						try
+						{
+							await TriggerAsync(item.StreamId, item.TemplateId, utcNow, cancellationToken);
+						}
+						catch (OperationCanceledException)
+						{
+							throw;
+						}
+						catch (Exception ex)
+						{
+							_logger.LogError(ex, "Error while updating schedule for {StreamId}/{TemplateId}", item.StreamId, item.TemplateId);
+						}
+					}
+				}
+			}
+		}
+
+		async Task<QueueItem?> PopQueueItemAsync()
+		{
+			for (; ; )
+			{
+				QueueItem[] items = await _queue.RangeByRankAsync(0, 0);
+				if (items.Length == 0)
+				{
+					return null;
+				}
+				if (await _queue.RemoveAsync(items[0]))
+				{
+					return items[0];
+				}
+			}
+		}
+
+		internal async Task ResetAsync()
+		{
+			await _redis.Database.KeyDeleteAsync(_queue.Key);
+			await _redis.Database.KeyDeleteAsync(_tickLockKey);
+		}
+
+		internal async Task TickForTestingAsync()
+		{
+			await UpdateQueueAsync(_clock.UtcNow);
+			await TickAsync(CancellationToken.None);
+		}
+
+		/// <summary>
+		/// Get the current set of streams and ensure there's an entry for each item
+		/// </summary>
+		public async Task UpdateQueueAsync(DateTime utcNow)
+		{
+			List<SortedSetEntry<QueueItem>> queueItems = new List<SortedSetEntry<QueueItem>>();
+
+			List<IStream> streams = await _streamService.GetStreamsAsync();
+			foreach (IStream stream in streams)
+			{
+				foreach ((TemplateRefId templateId, TemplateRef templateRef) in stream.Templates)
+				{
+					if (templateRef.Schedule != null)
+					{
+						DateTime? nextTriggerTimeUtc = templateRef.Schedule.GetNextTriggerTimeUtc(_timeZone);
+						if (nextTriggerTimeUtc != null)
+						{
+							if (utcNow > nextTriggerTimeUtc.Value)
+							{
+								double score = QueueItem.GetScoreFromTime(nextTriggerTimeUtc.Value);
+								queueItems.Add(new SortedSetEntry<QueueItem>(new QueueItem(stream.Id, templateId), score));
+
+								await _streamService.UpdateScheduleTriggerAsync(stream, templateId, utcNow);
+							}
+						}
+					}
+				}
+			}
+
+			await _queue.AddAsync(queueItems.ToArray());
+		}
+
+		/// <summary>
+		/// Trigger a schedule to run
+		/// </summary>
+		/// <param name="streamId">Stream for the schedule</param>
+		/// <param name="templateId"></param>
+		/// <param name="utcNow"></param>
+		/// <param name="cancellationToken"></param>
+		/// <returns>Async task</returns>
+		internal async Task<bool> TriggerAsync(StreamId streamId, TemplateRefId templateId, DateTime utcNow, CancellationToken cancellationToken)
+		{
+			IStream? stream = await _streamService.GetStreamAsync(streamId);
+			if (stream == null || !stream.Templates.TryGetValue(templateId, out TemplateRef? templateRef))
+			{
+				return false;
+			}
+
+			Schedule? schedule = templateRef.Schedule;
+			if (schedule == null)
+			{
+				return false;
+			}
+
+			using IScope scope = GlobalTracer.Instance.BuildSpan("ScheduleService.TriggerAsync").StartActive();
+			scope.Span.SetTag("StreamId", stream.Id);
+			scope.Span.SetTag("TemplateId", templateId);
+
+			Stopwatch stopwatch = Stopwatch.StartNew();
+			_logger.LogInformation("Updating schedule for {StreamId} template {TemplateId}", stream.Id, templateId);
+
+			// Get a list of jobs that we need to remove
+			List<JobId> removeJobIds = new List<JobId>();
+			foreach (JobId activeJobId in schedule.ActiveJobs)
+			{
+				IJob? job = await _jobService.GetJobAsync(activeJobId);
+				if (job == null || job.Batches.All(x => x.State == JobStepBatchState.Complete))
+				{
+					_logger.LogInformation("Removing active job {JobId}", activeJobId);
+					removeJobIds.Add(activeJobId);
+				}
+			}
+			await _streamService.UpdateScheduleTriggerAsync(stream, templateId, removeJobs: removeJobIds);
+
+			// If the stream is paused, bail out
+			if (stream.IsPaused(utcNow))
+			{
+				_logger.LogDebug("Skipping schedule update for stream {StreamId}. It has been paused until {PausedUntil} with comment '{PauseComment}'.", stream.Id, stream.PausedUntil, stream.PauseComment);
+				return false;
+			}
+
+			// Trigger this schedule
+			try
+			{
+				await TriggerAsync(stream, templateId, templateRef, schedule, schedule.ActiveJobs.Count - removeJobIds.Count, utcNow, cancellationToken);
+			}
+			catch (Exception ex)
+			{
+				_logger.LogError(ex, "Failed to start schedule {StreamId}/{TemplateId}", stream.Id, templateId);
+			}
+
+			// Print some timing info
+			stopwatch.Stop();
+			_logger.LogInformation("Updated schedule for {StreamId} template {TemplateId} in {TimeSeconds}ms", stream.Id, templateId, (long)stopwatch.Elapsed.TotalMilliseconds);
+			return true;
+		}
+
+		/// <summary>
+		/// Trigger a schedule to run
+		/// </summary>
+		/// <param name="stream">Stream for the schedule</param>
+		/// <param name="templateId"></param>
+		/// <param name="templateRef"></param>
+		/// <param name="schedule"></param>
+		/// <param name="numActiveJobs"></param>
+		/// <param name="utcNow">The current time</param>
+		/// <param name="cancellationToken"></param>
+		/// <returns>Async task</returns>
+		private async Task TriggerAsync(IStream stream, TemplateRefId templateId, TemplateRef templateRef, Schedule schedule, int numActiveJobs, DateTime utcNow, CancellationToken cancellationToken)
+		{
+			// Check we're not already at the maximum number of allowed jobs
+			if (schedule.MaxActive != 0 && numActiveJobs >= schedule.MaxActive)
+			{
+				_logger.LogInformation("Skipping trigger of {StreamId} template {TemplateId} - already have maximum number of jobs running ({NumJobs})", stream.Id, templateId, schedule.MaxActive);
+				foreach (JobId jobId in schedule.ActiveJobs)
+				{
+					_logger.LogInformation("Active job for {StreamId} template {TemplateId}: {JobId}", stream.Id, templateId, jobId);
+				}
+				return;
+			}
+
+			// Minimum changelist number, inclusive
+			int minChangeNumber = schedule.LastTriggerChange;
+			if (minChangeNumber > 0 && !schedule.RequireSubmittedChange)
+			{
+				minChangeNumber--;
+			}
+
+			// Maximum changelist number, exclusive
+			int? maxChangeNumber = null;
+
+			// Get the maximum number of changes to trigger
+			int maxNewChanges = 1;
+			if (schedule.MaxChanges != 0)
+			{
+				maxNewChanges = schedule.MaxChanges;
+			}
+			if (schedule.MaxActive != 0)
+			{
+				maxNewChanges = Math.Min(maxNewChanges, schedule.MaxActive - numActiveJobs);
+			}
+
+			// Create a timer to limit the amount we look back through P4 history
+			Stopwatch timer = Stopwatch.StartNew();
+			ChangeContentFlags? filterFlags = schedule.GetFilterFlags();
+
+			// Create a file filter
+			FileFilter? fileFilter = null;
+			if (schedule.Files != null)
+			{
+				fileFilter = new FileFilter(schedule.Files);
+			}
+
+			// Cache the Perforce history as we're iterating through changes to improve query performance
+			PerforceHistory history = new PerforceHistory(_perforce, stream.ClusterName, stream.Name);
+			
+			// Start as many jobs as possible
+			List<(int Change, int CodeChange)> triggerChanges = new List<(int, int)>();
+			while (triggerChanges.Count < maxNewChanges)
+			{
+				cancellationToken.ThrowIfCancellationRequested();
+
+				// Get the next valid change
+				ChangeDetails? changeDetails = null;
+				if (schedule.Gate != null)
+				{
+					changeDetails = await GetNextChangeForGateAsync(stream, templateId, schedule.Gate, minChangeNumber, maxChangeNumber, cancellationToken);
+				}
+				else
+				{
+					changeDetails = await history.GetNextChangeAsync();
+				}
+
+				// Quit if we didn't find anything
+				if (changeDetails == null)
+				{
+					break;
+				}
+				if (changeDetails.Number < minChangeNumber)
+				{
+					break;
+				}
+				if (changeDetails.Number == minChangeNumber && (schedule.RequireSubmittedChange || triggerChanges.Count > 0))
+				{
+					break;
+				}
+
+				// Adjust the changelist for the desired filter
+				int change = changeDetails.Number;
+				if (ShouldBuildChange(changeDetails, filterFlags, fileFilter))
+				{
+					int codeChange = await history.GetCodeChange(change);
+					if (codeChange == -1)
+					{
+						_logger.LogWarning("Unable to find code change for CL {Change}", change);
+						codeChange = change;
+					}
+					triggerChanges.Add((change, codeChange));
+				}
+
+				// Check we haven't exceeded the time limit
+				if (timer.Elapsed > TimeSpan.FromMinutes(2.0))
+				{
+					_logger.LogError("Querying for changes to trigger for {StreamId} template {TemplateId} has taken {Time}. Aborting.", stream.Id, templateId, timer.Elapsed);
+					break;
+				}
+				
+				// Update the remaining range of changes to check for
+				maxChangeNumber = change - 1;
+				if(maxChangeNumber < minChangeNumber)
+				{
+					break;
+				}
+			}
+
+			// Early out if there's nothing to do
+			if (triggerChanges.Count == 0)
+			{
+				_logger.LogInformation("Skipping trigger of {StreamName} template {TemplateId} - no candidate changes after CL {LastTriggerChange}", stream.Id, templateId, schedule.LastTriggerChange);
+				return;
+			}
+
+			// Get the matching template
+			ITemplate? template = await _templateCollection.GetAsync(templateRef.Hash);
+			if (template == null)
+			{
+				_logger.LogWarning("Unable to find template '{TemplateHash}' for '{TemplateId}'", templateRef.Hash, templateId);
+				return;
+			}
+
+			// Register the graph for it
+			IGraph graph = await _graphs.AddAsync(template);
+
+			// We may need to submit a new change for any new jobs. This only makes sense if there's one change.
+			if (template.SubmitNewChange != null)
+			{
+				int newChange = await _perforce.CreateNewChangeForTemplateAsync(stream, template);
+				int newCodeChange = await _perforce.GetCodeChangeAsync(stream.ClusterName, stream.Name, newChange);
+				triggerChanges = new List<(int, int)> { (newChange, newCodeChange) };
+			}
+
+			// Try to start all the new jobs
+			_logger.LogInformation("Starting {NumJobs} new jobs for {StreamId} template {TemplateId} (active: {NumActive}, max new: {MaxNewJobs})", triggerChanges.Count, stream.Id, templateId, numActiveJobs, maxNewChanges);
+			foreach ((int change, int codeChange) in triggerChanges.OrderBy(x => x.Change))
+			{
+				cancellationToken.ThrowIfCancellationRequested();
+				List<string> defaultArguments = template.GetDefaultArguments();
+				IJob newJob = await _jobService.CreateJobAsync(null, stream, templateId, template.Id, graph, template.Name, change, codeChange, null, null, null, template.Priority, null, null, template.PromoteIssuesByDefault, templateRef.ChainedJobs, templateRef.ShowUgsBadges, templateRef.ShowUgsAlerts, templateRef.NotificationChannel, templateRef.NotificationChannelFilter, defaultArguments);
+				_logger.LogInformation("Started new job for {StreamId} template {TemplateId} at CL {Change} (Code CL {CodeChange}): {JobId}", stream.Id, templateId, change, codeChange, newJob.Id);
+				await _streamService.UpdateScheduleTriggerAsync(stream, templateId, utcNow, change, new List<JobId> { newJob.Id }, new List<JobId>());
+			}
+		}
+
+		/// <summary>
+		/// Tests whether a schedule should build a particular change, based on its requested change filters
+		/// </summary>
+		/// <param name="details">The change details</param>
+		/// <param name="filterFlags"></param>
+		/// <param name="fileFilter">Filter for the files to trigger a build</param>
+		/// <returns></returns>
+		private bool ShouldBuildChange(ChangeDetails details, ChangeContentFlags? filterFlags, FileFilter? fileFilter)
+		{
+			if (Regex.IsMatch(details.Description, @"^\s*#\s*skipci", RegexOptions.Multiline))
+			{
+				return false;
+			}
+			if (filterFlags != null && filterFlags.Value != 0)
+			{
+				if ((details.GetContentFlags() & filterFlags.Value) == 0)
+				{
+					_logger.LogDebug("Not building change {Change} ({ChangeFlags}) due to filter flags ({FilterFlags})", details.Number, details.GetContentFlags().ToString(), filterFlags.Value.ToString());
+					return false;
+				}
+			}
+			if (fileFilter != null)
+			{
+				if (!details.Files.Any(x => fileFilter.Matches(x.Path)))
+				{
+					_logger.LogDebug("Not building change {Change} due to file filter", details.Number);
+					return false;
+				}
+			}
+			return true;
+		}
+
+		/// <summary>
+		/// Gets the next change to build for a schedule on a gate
+		/// </summary>
+		/// <returns></returns>
+		private async Task<ChangeDetails?> GetNextChangeForGateAsync(IStream stream, TemplateRefId templateRefId, ScheduleGate gate, int? minChange, int? maxChange, CancellationToken cancellationToken)
+		{
+			for (; ; )
+			{
+				cancellationToken.ThrowIfCancellationRequested();
+
+				List<IJob> jobs = await _jobCollection.FindAsync(streamId: stream.Id, templates: new[] { gate.TemplateRefId }, minChange: minChange, maxChange: maxChange, count: 1);
+				if (jobs.Count == 0)
+				{
+					return null;
+				}
+
+				IJob job = jobs[0];
+
+				IGraph? graph = await _graphs.GetAsync(job.GraphHash);
+				if (graph != null)
+				{
+					(JobStepState, JobStepOutcome)? state = job.GetTargetState(graph, gate.Target);
+					if (state != null && state.Value.Item1 == JobStepState.Completed)
+					{
+						JobStepOutcome outcome = state.Value.Item2;
+						if (outcome == JobStepOutcome.Success || outcome == JobStepOutcome.Warnings)
+						{
+							return await _perforce.GetChangeDetailsAsync(stream.ClusterName, stream.Name, job.Change);
+						}
+						_logger.LogInformation("Skipping trigger of {StreamName} template {TemplateId} - last {OtherTemplateRefId} job ({JobId}) ended with errors", stream.Id, templateRefId, gate.TemplateRefId, job.Id);
+					}
+				}
+
+				maxChange = job.Change - 1;
+			}
+		}
+	}
+}
diff --git a/Engine/Source/Programs/Horde/Horde.Build/Jobs/TemplateCollection.cs b/Engine/Source/Programs/Horde/Horde.Build/Jobs/TemplateCollection.cs
new file mode 100644
--- /dev/null
+++ b/Engine/Source/Programs/Horde/Horde.Build/Jobs/TemplateCollection.cs
@@ -0,0 +1,152 @@
+// Copyright Epic Games, Inc. All Rights Reserved.
+
+using System;
+using System.Collections.Generic;
+using System.Linq;
+using System.Threading.Tasks;
+using EpicGames.Core;
+using Horde.Build.Models;
+using Horde.Build.Server;
+using Horde.Build.Utilities;
+using HordeCommon;
+using Microsoft.Extensions.Caching.Memory;
+using MongoDB.Bson;
+using MongoDB.Bson.Serialization;
+using MongoDB.Bson.Serialization.Attributes;
+using MongoDB.Driver;
+
+namespace Horde.Build.Collections.Impl
+{
+	/// <summary>
+	/// Collection of template documents
+	/// </summary>
+	public sealed class TemplateCollection : ITemplateCollection, IDisposable
+	{
+		/// <summary>
+		/// Document describing a job template. These objects are considered immutable once created and uniquely referenced by hash, in order to de-duplicate across all job runs.
+		/// </summary>
+		class TemplateDocument : ITemplate
+		{
+			[BsonRequired, BsonId]
+			public ContentHash Id { get; private set; } = ContentHash.Empty;
+
+			[BsonRequired]
+			public string Name { get; private set; }
+
+			public Priority? Priority { get; private set; }
+			public bool AllowPreflights { get; set; } = true;
+			public bool UpdateIssues { get; set; } = false;
+			public bool PromoteIssuesByDefault { get; set; } = false;
+			public string? InitialAgentType { get; set; }
+
+			[BsonIgnoreIfNull]
+			public string? SubmitNewChange { get; set; }
+
+			[BsonIgnoreIfNull]
+			public string? SubmitDescription { get; set; }
+
+			public List<string> Arguments { get; private set; } = new List<string>();
+			public List<Parameter> Parameters { get; private set; } = new List<Parameter>();
+
+			IReadOnlyList<string> ITemplate.Arguments => Arguments;
+			IReadOnlyList<Parameter> ITemplate.Parameters => Parameters;
+
+			[BsonConstructor]
+			private TemplateDocument()
+			{
+				Name = null!;
+			}
+
+			public TemplateDocument(string name, Priority? priority, bool bAllowPreflights, bool updateIssues, bool promoteIssuesByDefault, string? initialAgentType, string? submitNewChange, string? submitDescription, List<string>? arguments, List<Parameter>? parameters)
+			{
+				Name = name;
+				Priority = priority;
+				AllowPreflights = bAllowPreflights;
+				UpdateIssues = updateIssues;
+				PromoteIssuesByDefault = promoteIssuesByDefault;
+				InitialAgentType = initialAgentType;
+				SubmitNewChange = submitNewChange;
+				SubmitDescription = submitDescription;
+				Arguments = arguments ?? new List<string>();
+				Parameters = parameters ?? new List<Parameter>();
+
+				// Compute the hash once all other fields have been set
+				Id = ContentHash.SHA1(BsonExtensionMethods.ToBson(this));
+			}
+		}
+
+		/// <summary>
+		/// Template documents
+		/// </summary>
+		readonly IMongoCollection<TemplateDocument> _templates;
+
+		/// <summary>
+		/// Cache of template documents
+		/// </summary>
+		readonly MemoryCache _templateCache;
+
+		/// <summary>
+		/// Constructor
+		/// </summary>
+		/// <param name="mongoService">The database service singleton</param>
+		public TemplateCollection(MongoService mongoService)
+		{
+			// Ensure discriminator cannot be registered twice (throws exception). Can otherwise happen during unit tests.
+			if (BsonSerializer.LookupDiscriminatorConvention(typeof(JobsTabColumn)) == null)
+			{
+				BsonSerializer.RegisterDiscriminatorConvention(typeof(JobsTabColumn), new DefaultDiscriminatorConvention(typeof(JobsTabColumn), typeof(JobsTabLabelColumn)));	
+			}
+			
+			_templates = mongoService.GetCollection<TemplateDocument>("Templates");
+
+			MemoryCacheOptions options = new MemoryCacheOptions();
+			_templateCache = new MemoryCache(options);
+		}
+
+		/// <inheritdoc/>
+		public void Dispose()
+		{
+			_templateCache.Dispose();
+		}
+
+		/// <inheritdoc/>
+		public async Task<ITemplate> AddAsync(string name, Priority? priority, bool bAllowPreflights, bool bUpdateIssues, bool bPromoteIssuesByDefault, string? initialAgentType, string? submitNewChange, string? submitDescription, List<string>? arguments, List<Parameter>? parameters)
+		{
+			TemplateDocument template = new TemplateDocument(name, priority, bAllowPreflights, bUpdateIssues, bPromoteIssuesByDefault, initialAgentType, submitNewChange, submitDescription, arguments, parameters);
+			if (await GetAsync(template.Id) == null)
+			{
+				await _templates.ReplaceOneAsync(x => x.Id == template.Id, template, new ReplaceOptions { IsUpsert = true });
+			}
+			return template;
+		}
+
+		/// <inheritdoc/>
+		public async Task<List<ITemplate>> FindAllAsync()
+		{
+			List<TemplateDocument> results = await _templates.Find(FilterDefinition<TemplateDocument>.Empty).ToListAsync();
+			return results.ConvertAll<ITemplate>(x => x);
+		}
+
+		/// <inheritdoc/>
+		public async Task<ITemplate?> GetAsync(ContentHash templateId)
+		{
+			object? result;
+			if (_templateCache.TryGetValue(templateId, out result))
+			{
+				return (ITemplate?)result;
+			}
+
+			ITemplate? template = await _templates.Find<TemplateDocument>(x => x.Id == templateId).FirstOrDefaultAsync();
+			if (template != null)
+			{
+				using (ICacheEntry entry = _templateCache.CreateEntry(templateId))
+				{
+					entry.SetSlidingExpiration(TimeSpan.FromMinutes(30.0));
+					entry.SetValue(template);
+				}
+			}
+
+			return template;
+		}
+	}
+}
