Horde: Finish re-org of Horde solution started a while ago.

diff --git a/Engine/Source/Programs/Horde/Horde.Build/Jobs/Templates/ITemplate.cs b/Engine/Source/Programs/Horde/Horde.Build/Jobs/Templates/ITemplate.cs
new file mode 100644
--- /dev/null
+++ b/Engine/Source/Programs/Horde/Horde.Build/Jobs/Templates/ITemplate.cs
@@ -0,0 +1,501 @@
+// Copyright Epic Games, Inc. All Rights Reserved.
+
+using System;
+using System.Collections.Generic;
+using EpicGames.Core;
+using Horde.Build.Api;
+using HordeCommon;
+using MongoDB.Bson.Serialization.Attributes;
+
+namespace Horde.Build.Models
+{
+	/// <summary>
+	/// Base class for parameters used to configure templates via the new build dialog
+	/// </summary>
+	[BsonDiscriminator(RootClass = true)]
+	[BsonKnownTypes(typeof(GroupParameter), typeof(TextParameter), typeof(ListParameter), typeof(BoolParameter))]
+	public abstract class Parameter
+	{
+		/// <summary>
+		/// Gets the default arguments for this parameter and its children
+		/// </summary>
+		/// <param name="defaultArguments">List of default arguments</param>
+		public abstract void GetDefaultArguments(List<string> defaultArguments);
+
+		/// <summary>
+		/// Convert this parameter to data for serialization
+		/// </summary>
+		/// <returns>Serializable parameter data</returns>
+		public abstract ParameterData ToData();
+	}
+
+	/// <summary>
+	/// Used to group a number of other parameters
+	/// </summary>
+	public class GroupParameter : Parameter
+	{
+		/// <summary>
+		/// Label to display next to this parameter
+		/// </summary>
+		public string Label { get; set; }
+
+		/// <summary>
+		/// How to display this group
+		/// </summary>
+		public GroupParameterStyle Style { get; set; }
+
+		/// <summary>
+		/// List of child parameters
+		/// </summary>
+		public List<Parameter> Children { get; set; }
+
+		/// <summary>
+		/// Private constructor for serialization
+		/// </summary>
+		public GroupParameter()
+		{
+			Label = null!;
+			Children = null!;
+		}
+
+		/// <summary>
+		/// Constructor
+		/// </summary>
+		/// <param name="label">Name of the group</param>
+		/// <param name="style">How to display this group</param>
+		/// <param name="children">List of child parameters</param>
+		public GroupParameter(string label, GroupParameterStyle style, List<Parameter> children)
+		{
+			Label = label;
+			Style = style;
+			Children = children;
+		}
+
+		/// <summary>
+		/// Gets the default arguments for this parameter and its children
+		/// </summary>
+		/// <param name="defaultArguments">List of default arguments</param>
+		public override void GetDefaultArguments(List<string> defaultArguments)
+		{
+			foreach (Parameter child in Children)
+			{
+				child.GetDefaultArguments(defaultArguments);
+			}
+		}
+
+		/// <summary>
+		/// Convert this parameter to data for serialization
+		/// </summary>
+		/// <returns>Serializable parameter data</returns>
+		public override ParameterData ToData()
+		{
+			return new GroupParameterData(Label, Style, Children.ConvertAll(x => x.ToData()));
+		}
+	}
+
+	/// <summary>
+	/// Free-form text entry parameter
+	/// </summary>
+	public sealed class TextParameter : Parameter
+	{
+		/// <summary>
+		/// Label to display next to this parameter. Should default to the parameter name.
+		/// </summary>
+		public string Label { get; set; }
+
+		/// <summary>
+		/// Argument to add (will have the value of this field appended)
+		/// </summary>
+		public string Argument { get; set; }
+
+		/// <summary>
+		/// Default value for this argument
+		/// </summary>
+		public string Default { get; set; }
+
+		/// <summary>
+		/// Hint text to display when the field is empty
+		/// </summary>
+		public string? Hint { get; set; }
+
+		/// <summary>
+		/// Regex used to validate values entered into this text field.
+		/// </summary>
+		[BsonIgnoreIfNull]
+		public string? Validation { get; set; }
+
+		/// <summary>
+		/// Message displayed to explain valid values if validation fails.
+		/// </summary>
+		[BsonIgnoreIfNull]
+		public string? Description { get; set; }
+
+		/// <summary>
+		/// Tool tip text to display
+		/// </summary>
+		public string? ToolTip { get; set; }
+
+		/// <summary>
+		/// Private constructor for serialization
+		/// </summary>
+		public TextParameter()
+		{
+			Label = null!;
+			Argument = null!;
+			Default = String.Empty;
+		}
+
+		/// <summary>
+		/// Constructor
+		/// </summary>
+		/// <param name="label">Parameter to pass this value to the BuildGraph script with</param>
+		/// <param name="argument">Label to show next to the parameter</param>
+		/// <param name="defaultValue">Default value for this argument</param>
+		/// <param name="hint">Hint text to display</param>
+		/// <param name="validation">Regex used to validate entries</param>
+		/// <param name="description">Message displayed for invalid values</param>
+		/// <param name="toolTip">Tool tip text to display</param>
+		public TextParameter(string label, string argument, string defaultValue, string? hint, string? validation, string? description, string? toolTip)
+		{
+			Label = label;
+			Argument = argument;
+			Default = defaultValue;
+			Hint = hint;
+			Validation = validation;
+			Description = description;
+			ToolTip = toolTip;
+		}
+
+		/// <summary>
+		/// Gets the default arguments for this parameter and its children
+		/// </summary>
+		/// <param name="defaultArguments">List of default arguments</param>
+		public override void GetDefaultArguments(List<string> defaultArguments)
+		{
+			defaultArguments.Add(Argument + Default);
+		}
+
+		/// <summary>
+		/// Convert this parameter to data for serialization
+		/// </summary>
+		/// <returns>Serializable parameter data</returns>
+		public override ParameterData ToData()
+		{
+			return new TextParameterData(Label, Argument, Default, Hint, Validation, Description, ToolTip);
+		}
+	}
+
+	/// <summary>
+	/// Possible option for a list parameter
+	/// </summary>
+	public class ListParameterItem
+	{
+		/// <summary>
+		/// Group to display this entry in
+		/// </summary>
+		[BsonIgnoreIfNull]
+		public string? Group { get; set; }
+
+		/// <summary>
+		/// Text to display for this option.
+		/// </summary>
+		public string Text { get; set; }
+
+		/// <summary>
+		/// Argument to add with this parameter.
+		/// </summary>
+		[BsonIgnoreIfNull]
+		public string? ArgumentIfEnabled { get; set; }
+
+		/// <summary>
+		/// Argument to add with this parameter.
+		/// </summary>
+		[BsonIgnoreIfNull]
+		public string? ArgumentIfDisabled { get; set; }
+
+		/// <summary>
+		/// Whether this item is selected by default
+		/// </summary>
+		[BsonIgnoreIfDefault, BsonDefaultValue(false)]
+		public bool Default { get; set; }
+
+		/// <summary>
+		/// Default constructor for JSON serializer
+		/// </summary>
+		public ListParameterItem()
+		{
+			Text = String.Empty;
+		}
+
+		/// <summary>
+		/// Constructor
+		/// </summary>
+		/// <param name="group">The group to put this parameter in</param>
+		/// <param name="text">Text to display for this option</param>
+		/// <param name="argumentIfEnabled">Argument to add if this option is enabled</param>
+		/// <param name="argumentIfDisabled">Argument to add if this option is disabled</param>
+		/// <param name="defaultValue">Whether this item is selected by default</param>
+		public ListParameterItem(string? group, string text, string? argumentIfEnabled, string? argumentIfDisabled, bool defaultValue)
+		{
+			Group = group;
+			Text = text;
+			ArgumentIfEnabled = argumentIfEnabled;
+			ArgumentIfDisabled = argumentIfDisabled;
+			Default = defaultValue;
+		}
+
+		/// <summary>
+		/// Convert this parameter to data for serialization
+		/// </summary>
+		/// <returns>Serializable parameter data</returns>
+		public ListParameterItemData ToData()
+		{
+			return new ListParameterItemData(Group, Text, ArgumentIfEnabled, ArgumentIfDisabled, Default);
+		}
+	}
+
+	/// <summary>
+	/// Allows the user to select a value from a constrained list of choices
+	/// </summary>
+	public class ListParameter : Parameter
+	{
+		/// <summary>
+		/// Label to display next to this parameter.
+		/// </summary>
+		public string Label { get; set; }
+
+		/// <summary>
+		/// Style of picker parameter to use
+		/// </summary>
+		public ListParameterStyle Style { get; set; }
+
+		/// <summary>
+		/// List of values to display in the list
+		/// </summary>
+		public List<ListParameterItem> Items { get; set; }
+
+		/// <summary>
+		/// Tool tip text to display
+		/// </summary>
+		public string? ToolTip { get; set; }
+
+		/// <summary>
+		/// Private constructor for serialization
+		/// </summary>
+		public ListParameter()
+		{
+			Label = null!;
+			Items = null!;
+		}
+
+		/// <summary>
+		/// List of possible values
+		/// </summary>
+		/// <param name="label">Label to show next to this parameter</param>
+		/// <param name="style">Style of list parameter to use</param>
+		/// <param name="entries">Entries for this list</param>
+		/// <param name="toolTip">Tool tip text to display</param>
+		public ListParameter(string label, ListParameterStyle style, List<ListParameterItem> entries, string? toolTip)
+		{
+			Label = label;
+			Style = style;
+			Items = entries;
+			ToolTip = toolTip;
+		}
+
+		/// <summary>
+		/// Gets the default arguments for this parameter and its children
+		/// </summary>
+		/// <param name="defaultArguments">List of default arguments</param>
+		public override void GetDefaultArguments(List<string> defaultArguments)
+		{
+			foreach(ListParameterItem item in Items)
+			{
+				if (item.Default)
+				{
+					if (item.ArgumentIfEnabled != null)
+					{
+						defaultArguments.Add(item.ArgumentIfEnabled);
+					}
+				}
+				else
+				{
+					if (item.ArgumentIfDisabled != null)
+					{
+						defaultArguments.Add(item.ArgumentIfDisabled);
+					}
+				}
+			}
+		}
+
+		/// <summary>
+		/// Convert this parameter to data for serialization
+		/// </summary>
+		/// <returns>Serializable parameter data</returns>
+		public override ParameterData ToData()
+		{
+			return new ListParameterData(Label, Style, Items.ConvertAll(x => x.ToData()), ToolTip);
+		}
+	}
+
+	/// <summary>
+	/// Allows the user to toggle an option on or off
+	/// </summary>
+	public class BoolParameter : Parameter
+	{
+		/// <summary>
+		/// Label to display next to this parameter.
+		/// </summary>
+		public string Label { get; set; }
+
+		/// <summary>
+		/// Value if enabled
+		/// </summary>
+		[BsonIgnoreIfNull]
+		public string? ArgumentIfEnabled { get; set; }
+
+		/// <summary>
+		/// Value if disabled
+		/// </summary>
+		[BsonIgnoreIfNull]
+		public string? ArgumentIfDisabled { get; set; }
+
+		/// <summary>
+		/// Whether this option should be enabled by default
+		/// </summary>
+		[BsonIgnoreIfDefault, BsonDefaultValue(false)]
+		public bool Default { get; set; }
+
+		/// <summary>
+		/// Tool tip text to display
+		/// </summary>
+		public string? ToolTip { get; set; }
+
+		/// <summary>
+		/// Private constructor for serialization
+		/// </summary>
+		public BoolParameter()
+		{
+			Label = String.Empty;
+		}
+
+		/// <summary>
+		/// Constructor
+		/// </summary>
+		/// <param name="label">Label to display for this parameter</param>
+		/// <param name="argumentIfEnabled">Value if enabled</param>
+		/// <param name="argumentIfDisabled">Value if disabled</param>
+		/// <param name="defaultValue">Default value for this argument</param>
+		/// <param name="toolTip">Tool tip text to display</param>
+		public BoolParameter(string label, string? argumentIfEnabled, string? argumentIfDisabled, bool defaultValue, string? toolTip)
+		{
+			Label = label;
+			ArgumentIfEnabled = argumentIfEnabled;
+			ArgumentIfDisabled = argumentIfDisabled;
+			Default = defaultValue;
+			ToolTip = toolTip;
+		}
+
+		/// <summary>
+		/// Gets the default arguments for this parameter and its children
+		/// </summary>
+		/// <param name="defaultArguments">List of default arguments</param>
+		public override void GetDefaultArguments(List<string> defaultArguments)
+		{
+			string? defaultArgument = Default ? ArgumentIfEnabled : ArgumentIfDisabled;
+			if (!String.IsNullOrEmpty(defaultArgument))
+			{
+				defaultArguments.Add(defaultArgument);
+			}
+		}
+		/// <summary>
+		/// Convert this parameter to data for serialization
+		/// </summary>
+		/// <returns><see cref="BoolParameterData"/> instance</returns>
+		public override ParameterData ToData()
+		{
+			return new BoolParameterData(Label, ArgumentIfEnabled, ArgumentIfDisabled, Default, ToolTip);
+		}
+	}
+
+	/// <summary>
+	/// Document describing a job template. These objects are considered immutable once created and uniquely referenced by hash, in order to de-duplicate across all job runs.
+	/// </summary>
+	public interface ITemplate
+	{
+		/// <summary>
+		/// Hash of this template
+		/// </summary>
+		public ContentHash Id { get; }
+
+		/// <summary>
+		/// Name of the template.
+		/// </summary>
+		public string Name { get; }
+
+		/// <summary>
+		/// Priority of this job
+		/// </summary>
+		public Priority? Priority { get; }
+
+		/// <summary>
+		/// Whether to allow preflights for this job type
+		/// </summary>
+		public bool AllowPreflights { get; }
+
+		/// <summary>
+		/// Whether to always issues for jobs using this template
+		/// </summary>
+		public bool UpdateIssues { get; }
+
+		/// <summary>
+		/// Whether to promote issues by default for jobs using this template
+		/// </summary>
+		public bool PromoteIssuesByDefault { get; }
+
+		/// <summary>
+		/// Agent type to use for parsing the job state
+		/// </summary>
+		public string? InitialAgentType { get; }
+
+		/// <summary>
+		/// Path to a file within the stream to submit to generate a new changelist for jobs
+		/// </summary>
+		public string? SubmitNewChange { get; }
+
+		/// <summary>
+		/// Description for new changelists
+		/// </summary>
+		public string? SubmitDescription { get; }
+
+		/// <summary>
+		/// Optional predefined user-defined properties for this job
+		/// </summary>
+		public IReadOnlyList<string> Arguments { get; }
+
+		/// <summary>
+		/// Parameters for this template
+		/// </summary>
+		public IReadOnlyList<Parameter> Parameters { get; }
+	}
+
+	/// <summary>
+	/// Extension methods for templates
+	/// </summary>
+	public static class TemplateExtensions
+	{
+		/// <summary>
+		/// Gets the arguments for default options in this template. Does not include the standard template arguments.
+		/// </summary>
+		/// <returns>List of default arguments</returns>
+		public static List<string> GetDefaultArguments(this ITemplate template)
+		{
+			List<string> defaultArguments = new List<string>(template.Arguments);
+			foreach (Parameter parameter in template.Parameters)
+			{
+				parameter.GetDefaultArguments(defaultArguments);
+			}
+			return defaultArguments;
+		}
+	}
+}
diff --git a/Engine/Source/Programs/Horde/Horde.Build/Jobs/Templates/ITemplateCollection.cs b/Engine/Source/Programs/Horde/Horde.Build/Jobs/Templates/ITemplateCollection.cs
new file mode 100644
--- /dev/null
+++ b/Engine/Source/Programs/Horde/Horde.Build/Jobs/Templates/ITemplateCollection.cs
@@ -0,0 +1,44 @@
+// Copyright Epic Games, Inc. All Rights Reserved.
+
+using System.Collections.Generic;
+using System.Threading.Tasks;
+using EpicGames.Core;
+using Horde.Build.Models;
+using HordeCommon;
+
+namespace Horde.Build.Collections
+{
+	/// <summary>
+	/// Interface for a collection of template documents
+	/// </summary>
+	public interface ITemplateCollection
+	{
+		/// <summary>
+		/// Public constructor
+		/// </summary>
+		/// <param name="name">Name of the template</param>
+		/// <param name="priority">Priority of this template</param>
+		/// <param name="bAllowPreflights">Whether to allow preflights of this job</param>
+		/// <param name="bUpdateIssues"> Whether to update issues for all jobs using this template</param>
+		/// <param name="bPromoteIssuesByDefault">Whether to promote issues by default for all jobs using this template</param>
+		/// <param name="initialAgentType">The agent type to parse the buildgraph script</param>
+		/// <param name="submitNewChange">Path to a file within the stream to submit to generate a new changelist for jobs</param>
+		/// <param name="submitDescription">Description for new changes submitted to the stream</param>
+		/// <param name="arguments">List of arguments which are always specified</param>
+		/// <param name="parameters">List of template parameters</param>
+		Task<ITemplate> AddAsync(string name, Priority? priority = null, bool bAllowPreflights = true, bool bUpdateIssues = false, bool bPromoteIssuesByDefault = false, string? initialAgentType = null, string? submitNewChange = null, string? submitDescription = null, List<string>? arguments = null, List<Parameter>? parameters = null);
+
+		/// <summary>
+		/// Gets all the available templates
+		/// </summary>
+		/// <returns>List of template documents</returns>
+		Task<List<ITemplate>> FindAllAsync();
+
+		/// <summary>
+		/// Gets a template by ID
+		/// </summary>
+		/// <param name="templateId">Unique id of the template</param>
+		/// <returns>The template document</returns>
+		Task<ITemplate?> GetAsync(ContentHash templateId);
+	}
+}
diff --git a/Engine/Source/Programs/Horde/Horde.Build/Jobs/Templates/TemplateMessages.cs b/Engine/Source/Programs/Horde/Horde.Build/Jobs/Templates/TemplateMessages.cs
new file mode 100644
--- /dev/null
+++ b/Engine/Source/Programs/Horde/Horde.Build/Jobs/Templates/TemplateMessages.cs
@@ -0,0 +1,547 @@
+// Copyright Epic Games, Inc. All Rights Reserved.
+
+using System;
+using System.Collections.Generic;
+using System.ComponentModel.DataAnnotations;
+using EpicGames.Core;
+using Horde.Build.Models;
+using HordeCommon;
+
+namespace Horde.Build.Api
+{
+	/// <summary>
+	/// Base class for template parameters
+	/// </summary>
+	[JsonKnownTypes(typeof(GroupParameterData), typeof(TextParameterData), typeof(ListParameterData), typeof(BoolParameterData))]
+	public abstract class ParameterData
+	{
+		/// <summary>
+		/// Convert to a parameter object
+		/// </summary>
+		/// <returns><see cref="Parameter"/> object</returns>
+		public abstract Parameter ToModel();
+	}
+
+	/// <summary>
+	/// Describes how to render a group parameter
+	/// </summary>
+	public enum GroupParameterStyle
+	{
+		/// <summary>
+		/// Separate tab on the form
+		/// </summary>
+		Tab,
+
+		/// <summary>
+		/// Section with heading
+		/// </summary>
+		Section,
+	}
+
+	/// <summary>
+	/// Used to group a number of other parameters
+	/// </summary>
+	[JsonDiscriminator("Group")]
+	public class GroupParameterData : ParameterData
+	{
+		/// <summary>
+		/// Label to display next to this parameter
+		/// </summary>
+		public string Label { get; set; }
+
+		/// <summary>
+		/// How to display this group
+		/// </summary>
+		public GroupParameterStyle Style { get; set; }
+
+		/// <summary>
+		/// List of child parameters
+		/// </summary>
+		public List<ParameterData> Children { get; set; }
+
+		/// <summary>
+		/// Private constructor for serialization
+		/// </summary>
+		public GroupParameterData()
+		{
+			Label = null!;
+			Children = null!;
+		}
+
+		/// <summary>
+		/// Constructor
+		/// </summary>
+		/// <param name="label">Name of the group</param>
+		/// <param name="style">How to display this group</param>
+		/// <param name="children">List of child parameters</param>
+		public GroupParameterData(string label, GroupParameterStyle style, List<ParameterData> children)
+		{
+			Label = label;
+			Style = style;
+			Children = children;
+		}
+
+		/// <summary>
+		/// Converts this data to a model object
+		/// </summary>
+		/// <returns>New <see cref="GroupParameter"/> object</returns>
+		public override Parameter ToModel()
+		{
+			return new GroupParameter(Label, Style, Children.ConvertAll(x => x.ToModel()));
+		}
+	}
+
+	/// <summary>
+	/// Free-form text entry parameter
+	/// </summary>
+	[JsonDiscriminator("Text")]
+	public class TextParameterData : ParameterData
+	{
+		/// <summary>
+		/// Name of the parameter associated with this parameter.
+		/// </summary>
+		public string Label { get; set; }
+
+		/// <summary>
+		/// Argument to pass to the executor
+		/// </summary>
+		public string Argument { get; set; }
+
+		/// <summary>
+		/// Default value for this argument
+		/// </summary>
+		public string Default { get; set; }
+
+		/// <summary>
+		/// Hint text for this parameter
+		/// </summary>
+		public string? Hint { get; set; }
+
+		/// <summary>
+		/// Regex used to validate this parameter
+		/// </summary>
+		public string? Validation { get; set; }
+
+		/// <summary>
+		/// Message displayed if validation fails, informing user of valid values.
+		/// </summary>
+		public string? ValidationError { get; set; }
+
+		/// <summary>
+		/// Tool-tip text to display
+		/// </summary>
+		public string? ToolTip { get; set; }
+
+		/// <summary>
+		/// Private constructor for serialization
+		/// </summary>
+		public TextParameterData()
+		{
+			Label = null!;
+			Argument = null!;
+			Default = String.Empty;
+		}
+
+		/// <summary>
+		/// Constructor
+		/// </summary>
+		/// <param name="label">Label to show next to the parameter</param>
+		/// <param name="argument">Argument to pass this value with</param>
+		/// <param name="defaultValue">Default value for this parameter</param>
+		/// <param name="hint">Hint text to display for this parameter</param>
+		/// <param name="validation">Regex used to validate entries</param>
+		/// <param name="validationError">Message displayed to explain validation issues</param>
+		/// <param name="toolTip">Tool tip text to display</param>
+		public TextParameterData(string label, string argument, string defaultValue, string? hint, string? validation, string? validationError, string? toolTip)
+		{
+			Label = label;
+			Argument = argument;
+			Default = defaultValue;
+			Hint = hint;
+			Validation = validation;
+			ValidationError = validationError;
+			ToolTip = toolTip;
+		}
+
+		/// <summary>
+		/// Converts this data to a model object
+		/// </summary>
+		/// <returns>New <see cref="TextParameter"/> object</returns>
+		public override Parameter ToModel()
+		{
+			return new TextParameter(Label, Argument, Default, Hint, Validation, ValidationError, ToolTip);
+		}
+	}
+
+	/// <summary>
+	/// Style of list parameter
+	/// </summary>
+	public enum ListParameterStyle
+	{
+		/// <summary>
+		/// Regular drop-down list. One item is always selected.
+		/// </summary>
+		List,
+
+		/// <summary>
+		/// Drop-down list with checkboxes
+		/// </summary>
+		MultiList,
+
+		/// <summary>
+		/// Tag picker from list of options
+		/// </summary>
+		TagPicker,
+	}
+
+	/// <summary>
+	/// Possible option for a list parameter
+	/// </summary>
+	public class ListParameterItemData
+	{
+		/// <summary>
+		/// Optional group heading to display this entry under, if the picker style supports it.
+		/// </summary>
+		public string? Group { get; set; }
+
+		/// <summary>
+		/// Name of the parameter associated with this list.
+		/// </summary>
+		public string Text { get; set; }
+
+		/// <summary>
+		/// Argument to pass with this parameter.
+		/// </summary>
+		public string? ArgumentIfEnabled { get; set; }
+
+		/// <summary>
+		/// Argument to pass with this parameter.
+		/// </summary>
+		public string? ArgumentIfDisabled { get; set; }
+
+		/// <summary>
+		/// Whether this item is selected by default
+		/// </summary>
+		public bool Default { get; set; }
+
+		/// <summary>
+		/// Private constructor for serialization
+		/// </summary>
+		public ListParameterItemData()
+		{
+			Text = null!;
+		}
+
+		/// <summary>
+		/// Constructor
+		/// </summary>
+		/// <param name="group">The group to put this parameter in</param>
+		/// <param name="text">Text to display for this option</param>
+		/// <param name="argumentIfEnabled">Argument to pass for this item if it's enabled</param>
+		/// <param name="argumentIfDisabled">Argument to pass for this item if it's enabled</param>
+		/// <param name="defaultValue">Whether this item is selected by default</param>
+		public ListParameterItemData(string? group, string text, string? argumentIfEnabled, string? argumentIfDisabled, bool defaultValue)
+		{
+			Group = group;
+			Text = text;
+			ArgumentIfEnabled = argumentIfEnabled;
+			ArgumentIfDisabled = argumentIfDisabled;
+			Default = defaultValue;
+		}
+
+		/// <summary>
+		/// Converts this data to a model object
+		/// </summary>
+		/// <returns>New <see cref="ListParameterItem"/> object</returns>
+		public ListParameterItem ToModel()
+		{
+			return new ListParameterItem(Group, Text, ArgumentIfEnabled, ArgumentIfDisabled, Default);
+		}
+	}
+
+	/// <summary>
+	/// Allows the user to select a value from a constrained list of choices
+	/// </summary>
+	[JsonDiscriminator("List")]
+	public class ListParameterData : ParameterData
+	{
+		/// <summary>
+		/// Label to display next to this parameter. Defaults to the parameter name.
+		/// </summary>
+		public string Label { get; set; }
+
+		/// <summary>
+		/// The type of list parameter
+		/// </summary>
+		public ListParameterStyle Style { get; set; }
+
+		/// <summary>
+		/// List of values to display in the list
+		/// </summary>
+		public List<ListParameterItemData> Items { get; set; }
+
+		/// <summary>
+		/// Tool tip text to display
+		/// </summary>
+		public string? ToolTip { get; set; }
+
+		/// <summary>
+		/// Private constructor
+		/// </summary>
+		public ListParameterData()
+		{
+			Label = null!;
+			Items = null!;
+		}
+
+		/// <summary>
+		/// List of possible values
+		/// </summary>
+		/// <param name="label">Label to show next to this parameter</param>
+		/// <param name="style">Type of picker to show</param>
+		/// <param name="items">Entries for this list</param>
+		/// <param name="toolTip">Tool tip text to display</param>
+		public ListParameterData(string label, ListParameterStyle style, List<ListParameterItemData> items, string? toolTip)
+		{
+			Label = label;
+			Style = style;
+			Items = items;
+			ToolTip = toolTip;
+		}
+
+		/// <summary>
+		/// Converts this data to a model object
+		/// </summary>
+		/// <returns>New <see cref="ListParameter"/> object</returns>
+		public override Parameter ToModel()
+		{
+			return new ListParameter(Label, Style, Items.ConvertAll(x => x.ToModel()), ToolTip);
+		}
+	}
+
+	/// <summary>
+	/// Allows the user to toggle an option on or off
+	/// </summary>
+	[JsonDiscriminator("Bool")]
+	public class BoolParameterData : ParameterData
+	{
+		/// <summary>
+		/// Name of the parameter associated with this parameter.
+		/// </summary>
+		public string Label { get; set; }
+
+		/// <summary>
+		/// Value if enabled
+		/// </summary>
+		public string? ArgumentIfEnabled { get; set; }
+
+		/// <summary>
+		/// Value if disabled
+		/// </summary>
+		public string? ArgumentIfDisabled { get; set; }
+
+		/// <summary>
+		/// Whether this argument is enabled by default
+		/// </summary>
+		public bool Default { get; set; }
+
+		/// <summary>
+		/// Tool tip text to display
+		/// </summary>
+		public string? ToolTip { get; set; }
+
+		/// <summary>
+		/// Private constructor for serialization
+		/// </summary>
+		public BoolParameterData()
+		{
+			Label = null!;
+		}
+
+		/// <summary>
+		/// Constructor
+		/// </summary>
+		/// <param name="label">Label to show next to this parameter</param>
+		/// <param name="argumentIfEnabled">Value if enabled</param>
+		/// <param name="argumentIfDisabled">Value if disabled</param>
+		/// <param name="defaultValue">Whether this option is enabled by default</param>
+		/// <param name="toolTip">The tool tip text to display</param>
+		public BoolParameterData(string label, string? argumentIfEnabled, string? argumentIfDisabled, bool defaultValue, string? toolTip)
+		{
+			Label = label;
+			ArgumentIfEnabled = argumentIfEnabled;
+			ArgumentIfDisabled = argumentIfDisabled;
+			Default = defaultValue;
+			ToolTip = toolTip;
+		}
+
+		/// <summary>
+		/// Converts this data to a model object
+		/// </summary>
+		/// <returns>New <see cref="BoolParameter"/> object</returns>
+		public override Parameter ToModel()
+		{
+			return new BoolParameter(Label, ArgumentIfEnabled, ArgumentIfDisabled, Default, ToolTip);
+		}
+	}
+
+	/// <summary>
+	/// Parameters to create a new template
+	/// </summary>
+	public class TemplateConfig
+	{
+		/// <summary>
+		/// Name for the new template
+		/// </summary>
+		[Required]
+		public string Name { get; set; } = null!;
+
+		/// <summary>
+		/// Default priority for this job
+		/// </summary>
+		public Priority? Priority { get; set; }
+
+		/// <summary>
+		/// Whether to allow preflights of this template
+		/// </summary>
+		public bool AllowPreflights { get; set; } = true;
+
+		/// <summary>
+		/// Whether issues should be updated for all jobs using this template
+		/// </summary>
+		public bool UpdateIssues { get; set; } = false;
+
+		/// <summary>
+		/// Whether issues should be promoted by default for this template, promoted issues will generate user notifications 
+		/// </summary>
+		public bool PromoteIssuesByDefault { get; set; } = false;
+
+		/// <summary>
+		/// Initial agent type to parse the buildgraph script on
+		/// </summary>
+		public string? InitialAgentType { get; set; }
+
+		/// <summary>
+		/// Path to a file within the stream to submit to generate a new changelist for jobs
+		/// </summary>
+		public string? SubmitNewChange { get; set; }
+
+		/// <summary>
+		/// Description for new changelists
+		/// </summary>
+		public string? SubmitDescription { get; set; }
+
+		/// <summary>
+		/// Fixed arguments for the new job
+		/// </summary>
+		public List<string> Arguments { get; set; } = new List<string>();
+
+		/// <summary>
+		/// Parameters for this template
+		/// </summary>
+		public List<ParameterData> Parameters { get; set; } = new List<ParameterData>();
+	}
+
+	/// <summary>
+	/// Response describing a template
+	/// </summary>
+	public class GetTemplateResponseBase
+	{
+		/// <summary>
+		/// Name of the template
+		/// </summary>
+		public string Name { get; set; }
+
+		/// <summary>
+		/// Default priority for this job
+		/// </summary>
+		public Priority? Priority { get; set; }
+
+		/// <summary>
+		/// Whether to allow preflights of this template
+		/// </summary>
+		public bool AllowPreflights { get; set; }
+
+		/// <summary>
+		/// Whether to always update issues on jobs using this template
+		/// </summary>
+		public bool UpdateIssues { get; set; }
+
+		/// <summary>
+		/// The initial agent type to parse the BuildGraph script on
+		/// </summary>
+		public string? InitialAgentType { get; set; }
+
+		/// <summary>
+		/// Path to a file within the stream to submit to generate a new changelist for jobs
+		/// </summary>
+		public string? SubmitNewChange { get; }
+
+		/// <summary>
+		/// Parameters for the job.
+		/// </summary>
+		public List<string> Arguments { get; set; }
+
+		/// <summary>
+		/// List of parameters for this template
+		/// </summary>
+		public List<ParameterData> Parameters { get; set; }
+
+		/// <summary>
+		/// Parameterless constructor for serialization
+		/// </summary>
+		protected GetTemplateResponseBase()
+		{
+			Name = null!;
+			AllowPreflights = true;
+			Arguments = new List<string>();
+			Parameters = new List<ParameterData>();
+		}
+
+		/// <summary>
+		/// Constructor
+		/// </summary>
+		/// <param name="template">The template to construct from</param>
+		public GetTemplateResponseBase(ITemplate template)
+		{
+			Name = template.Name;
+			Priority = template.Priority;
+			AllowPreflights = template.AllowPreflights;
+			UpdateIssues = template.UpdateIssues;
+			InitialAgentType = template.InitialAgentType;
+			SubmitNewChange = template.SubmitNewChange;
+			Arguments = new List<string>(template.Arguments);
+			Parameters = template.Parameters.ConvertAll(x => x.ToData());
+		}
+	}
+
+	/// <summary>
+	/// Response describing a template
+	/// </summary>
+	public class GetTemplateResponse : GetTemplateResponseBase
+	{
+		/// <summary>
+		/// Unique id of the template
+		/// </summary>
+		public string Id { get; set; }
+
+		/// <summary>
+		/// Parameterless constructor for serialization
+		/// </summary>
+		protected GetTemplateResponse()
+			: base()
+		{
+			Id = null!;
+		}
+
+		/// <summary>
+		/// Constructor
+		/// </summary>
+		/// <param name="template">The template to construct from</param>
+		public GetTemplateResponse(ITemplate template)
+			: base(template)
+		{
+			Id = template.Id.ToString();
+		}
+	}
+}
diff --git a/Engine/Source/Programs/Horde/Horde.Build/Jobs/TestData/ITestDataCollection.cs b/Engine/Source/Programs/Horde/Horde.Build/Jobs/TestData/ITestDataCollection.cs
new file mode 100644
--- /dev/null
+++ b/Engine/Source/Programs/Horde/Horde.Build/Jobs/TestData/ITestDataCollection.cs
@@ -0,0 +1,57 @@
+// Copyright Epic Games, Inc. All Rights Reserved.
+
+using System.Collections.Generic;
+using System.Threading.Tasks;
+using Horde.Build.Models;
+using Horde.Build.Utilities;
+using MongoDB.Bson;
+
+namespace Horde.Build.Collections
+{
+	using JobId = ObjectId<IJob>;
+	using StreamId = StringId<IStream>;
+
+	/// <summary>
+	/// Collection of test data documents
+	/// </summary>
+	public interface ITestDataCollection
+	{
+		/// <summary>
+		/// Creates a new test data document
+		/// </summary>
+		/// <param name="job">The job containing the step</param>
+		/// <param name="step">The step producing the data</param>
+		/// <param name="key">Key identifying the test</param>
+		/// <param name="value">The data to store</param>
+		/// <returns>The new stream document</returns>
+		Task<ITestData> AddAsync(IJob job, IJobStep step, string key, BsonDocument value);
+
+		/// <summary>
+		/// Gets a stream by ID
+		/// </summary>
+		/// <param name="id">Unique id of the stream</param>
+		/// <returns>The stream document</returns>
+		Task<ITestData?> GetAsync(ObjectId id);
+
+		/// <summary>
+		/// Searches for test data that matches a set of criteria
+		/// </summary>
+		/// <param name="streamId">The stream id</param>
+		/// <param name="minChange">The minimum changelist number to return (inclusive)</param>
+		/// <param name="maxChange">The maximum changelist number to return (inclusive)</param>
+		/// <param name="jobId">The job id</param>
+		/// <param name="stepId">The unique step id</param>
+		/// <param name="key">Key identifying the result to return</param>
+		/// <param name="index">Offset within the results to return</param>
+		/// <param name="count">Number of results to return</param>
+		/// <returns>The stream document</returns>
+		Task<List<ITestData>> FindAsync(StreamId? streamId, int? minChange, int? maxChange, JobId? jobId, SubResourceId? stepId, string? key = null, int index = 0, int count = 10);
+
+		/// <summary>
+		/// Delete the test data
+		/// </summary>
+		/// <param name="id">Unique id of the test data</param>
+		/// <returns>Async task</returns>
+		Task DeleteAsync(ObjectId id);
+	}
+}
diff --git a/Engine/Source/Programs/Horde/Horde.Build/Jobs/TestData/TestDataController.Messages.cs b/Engine/Source/Programs/Horde/Horde.Build/Jobs/TestData/TestDataController.Messages.cs
new file mode 100644
--- /dev/null
+++ b/Engine/Source/Programs/Horde/Horde.Build/Jobs/TestData/TestDataController.Messages.cs
@@ -0,0 +1,121 @@
+// Copyright Epic Games, Inc. All Rights Reserved.
+
+using System;
+using System.Collections.Generic;
+using System.ComponentModel.DataAnnotations;
+using Horde.Build.Models;
+using MongoDB.Bson.Serialization;
+
+namespace Horde.Build.Api
+{
+	/// <summary>
+	/// Response object describing test data to store
+	/// </summary>
+	public class CreateTestDataRequest
+	{
+		/// <summary>
+		/// The job which produced the data
+		/// </summary>
+		[Required]
+		public string JobId { get; set; } = String.Empty;
+
+		/// <summary>
+		/// The step that ran
+		/// </summary>
+		[Required]
+		public string StepId { get; set; } = String.Empty;
+
+		/// <summary>
+		/// Key used to identify the particular data
+		/// </summary>
+		public string Key { get; set; } = String.Empty;
+
+		/// <summary>
+		/// The data stored for this test
+		/// </summary>
+		[Required]
+		public Dictionary<string, object> Data { get; set; } = new Dictionary<string, object>();
+	}
+
+	/// <summary>
+	/// Response object describing the created document
+	/// </summary>
+	public class CreateTestDataResponse
+	{
+		/// <summary>
+		/// The id for the new document
+		/// </summary>
+		public string Id { get; set; }
+
+		/// <summary>
+		/// Constructor
+		/// </summary>
+		/// <param name="id">Id of the new document</param>
+		public CreateTestDataResponse(string id)
+		{
+			Id = id;
+		}
+	}
+
+	/// <summary>
+	/// Response object describing test results
+	/// </summary>
+	public class GetTestDataResponse
+	{
+		/// <summary>
+		/// Unique id of the test data
+		/// </summary>
+		public string Id { get; set; }
+
+		/// <summary>
+		/// Stream that generated the test data
+		/// </summary>
+		public string StreamId { get; set; }
+
+		/// <summary>
+		/// The template reference id
+		/// </summary>
+		public string TemplateRefId { get; set; }
+
+		/// <summary>
+		/// The job which produced the data
+		/// </summary>
+		public string JobId { get; set; }
+
+		/// <summary>
+		/// The step that ran
+		/// </summary>
+		public string StepId { get; set; }
+
+		/// <summary>
+		/// The changelist number that contained the data
+		/// </summary>
+		public int Change { get; set; }
+
+		/// <summary>
+		/// Key used to identify the particular data
+		/// </summary>
+		public string Key { get; set; }
+
+		/// <summary>
+		/// The data stored for this test
+		/// </summary>
+		public Dictionary<string, object> Data { get; set; }
+
+		/// <summary>
+		/// Constructor
+		/// </summary>
+		/// <param name="testData">Test data to construct from</param>
+		internal GetTestDataResponse(ITestData testData)
+		{
+			Id = testData.Id.ToString();
+			StreamId = testData.StreamId.ToString();
+			TemplateRefId = testData.TemplateRefId.ToString();
+			JobId = testData.JobId.ToString();
+			StepId = testData.StepId.ToString();
+			Change = testData.Change;
+			Key = testData.Key;
+			Data = BsonSerializer.Deserialize<Dictionary<string, object>>(testData.Data);
+		}
+	}
+}
diff --git a/Engine/Source/Programs/Horde/Horde.Build/Jobs/TestData/TestDataController.cs b/Engine/Source/Programs/Horde/Horde.Build/Jobs/TestData/TestDataController.cs
new file mode 100644
--- /dev/null
+++ b/Engine/Source/Programs/Horde/Horde.Build/Jobs/TestData/TestDataController.cs
@@ -0,0 +1,141 @@
+﻿// Copyright Epic Games, Inc. All Rights Reserved.
+
+using System.Collections.Generic;
+using System.Threading.Tasks;
+using Horde.Build.Acls;
+using Horde.Build.Api;
+using Horde.Build.Collections;
+using Horde.Build.Models;
+using Horde.Build.Services;
+using Horde.Build.Utilities;
+using Microsoft.AspNetCore.Authorization;
+using Microsoft.AspNetCore.Mvc;
+using MongoDB.Bson;
+
+namespace Horde.Build.Controllers
+{
+	using JobId = ObjectId<IJob>;
+	using StreamId = StringId<IStream>;
+
+	/// <summary>
+	/// Controller for the /api/v1/testdata endpoint
+	/// </summary>
+	[ApiController]
+	[Authorize]
+	[Route("[controller]")]
+	public class TestDataController : ControllerBase
+	{
+		/// <summary>
+		/// Collection of job documents
+		/// </summary>
+		private readonly JobService _jobService;
+
+		/// <summary>
+		/// Collection of test data documents
+		/// </summary>
+		private readonly ITestDataCollection _testDataCollection;
+
+		/// <summary>
+		/// Constructor
+		/// </summary>
+		/// <param name="jobService">The job service singleton</param>
+		/// <param name="testDataCollection">Collection of test data documents</param>
+		public TestDataController(JobService jobService, ITestDataCollection testDataCollection)
+		{
+			_jobService = jobService;
+			_testDataCollection = testDataCollection;
+		}
+
+		/// <summary>
+		/// Creates a new TestData document
+		/// </summary>
+		/// <returns>The stream document</returns>
+		[HttpPost]
+		[Route("/api/v1/testdata")]
+		public async Task<ActionResult<CreateTestDataResponse>> CreateAsync(CreateTestDataRequest request)
+		{
+			IJob? job = await _jobService.GetJobAsync(new JobId(request.JobId));
+			if (job == null)
+			{
+				return NotFound();
+			}
+			if (!await _jobService.AuthorizeAsync(job, AclAction.UpdateJob, User, null))
+			{
+				return Forbid();
+			}
+
+			IJobStep? jobStep;
+			if (!job.TryGetStep(request.StepId.ToSubResourceId(), out jobStep))
+			{
+				return NotFound();
+			}
+
+			ITestData testData = await _testDataCollection.AddAsync(job, jobStep, request.Key, new BsonDocument(request.Data));
+			return new CreateTestDataResponse(testData.Id.ToString());
+		}
+
+		/// <summary>
+		/// Searches for test data that matches a set of criteria
+		/// </summary>
+		/// <param name="streamId">The stream id</param>
+		/// <param name="minChange">The minimum changelist number to return (inclusive)</param>
+		/// <param name="maxChange">The maximum changelist number to return (inclusive)</param>
+		/// <param name="jobId">The job id</param>
+		/// <param name="jobStepId">The unique step id</param>
+		/// <param name="key">Key identifying the result to return</param>
+		/// <param name="index">Offset within the results to return</param>
+		/// <param name="count">Number of results to return</param>
+		/// <param name="filter">Filter for properties to return</param>
+		/// <returns>The stream document</returns>
+		[HttpGet]
+		[Route("/api/v1/testdata")]
+		[ProducesResponseType(typeof(List<GetTestDataResponse>), 200)]
+		public async Task<ActionResult<List<object>>> FindTestDataAsync([FromQuery] string? streamId = null, [FromQuery] int? minChange = null, [FromQuery] int? maxChange = null, string? jobId = null, string? jobStepId = null, string? key = null, int index = 0, int count = 10, PropertyFilter? filter = null)
+		{
+			StreamId? streamIdValue = null;
+			if(streamId != null)
+			{
+				streamIdValue = new StreamId(streamId);
+			}
+
+			JobPermissionsCache cache = new JobPermissionsCache();
+
+			List<object> results = new List<object>();
+
+			List<ITestData> documents = await _testDataCollection.FindAsync(streamIdValue, minChange, maxChange, jobId?.ToObjectId<IJob>(), jobStepId?.ToSubResourceId(), key, index, count);
+			foreach (ITestData document in documents)
+			{
+				if (await _jobService.AuthorizeAsync(document.JobId, AclAction.ViewJob, User, cache))
+				{
+					results.Add(PropertyFilter.Apply(new GetTestDataResponse(document), filter));
+				}
+			}
+
+			return results;
+		}
+
+		/// <summary>
+		/// Retrieve information about a specific issue
+		/// </summary>
+		/// <param name="testDataId">Id of the document to get information about</param>
+		/// <param name="filter">Filter for the properties to return</param>
+		/// <returns>List of matching agents</returns>
+		[HttpGet]
+		[Route("/api/v1/testdata/{testDataId}")]
+		[ProducesResponseType(typeof(GetTestDataResponse), 200)]
+		public async Task<ActionResult<object>> GetTestDataAsync(string testDataId, [FromQuery] PropertyFilter? filter = null)
+		{
+			ITestData? testData = await _testDataCollection.GetAsync(testDataId.ToObjectId());
+			if (testData == null)
+			{
+				return NotFound();
+			}
+			if (!await _jobService.AuthorizeAsync(testData.JobId, AclAction.ViewJob, User, null))
+			{
+				return Forbid();
+			}
+
+			return PropertyFilter.Apply(new GetTestDataResponse(testData), filter);
+		}
+	}
+}
diff --git a/Engine/Source/Programs/Horde/Horde.Build/Jobs/TestDataCollection.cs b/Engine/Source/Programs/Horde/Horde.Build/Jobs/TestDataCollection.cs
new file mode 100644
--- /dev/null
+++ b/Engine/Source/Programs/Horde/Horde.Build/Jobs/TestDataCollection.cs
@@ -0,0 +1,129 @@
+// Copyright Epic Games, Inc. All Rights Reserved.
+
+using System;
+using System.Collections.Generic;
+using System.Linq;
+using System.Threading.Tasks;
+using Horde.Build.Models;
+using Horde.Build.Server;
+using Horde.Build.Utilities;
+using MongoDB.Bson;
+using MongoDB.Driver;
+
+namespace Horde.Build.Collections.Impl
+{
+	using JobId = ObjectId<IJob>;
+	using StreamId = StringId<IStream>;
+	using TemplateRefId = StringId<TemplateRef>;
+
+	/// <summary>
+	/// Collection of test data documents
+	/// </summary>
+	public class TestDataCollection : ITestDataCollection
+	{
+		/// <summary>
+		/// Information about a test data document
+		/// </summary>
+		class TestDataDocument : ITestData
+		{
+			public ObjectId Id { get; set; }
+			public StreamId StreamId { get; set; }
+			public TemplateRefId TemplateRefId { get; set; }
+			public JobId JobId { get; set; }
+			public SubResourceId StepId { get; set; }
+			public int Change { get; set; }
+			public string Key { get; set; }
+			public BsonDocument Data { get; set; }
+
+			private TestDataDocument()
+			{
+				Key = String.Empty;
+				Data = new BsonDocument();
+			}
+
+			public TestDataDocument(IJob job, IJobStep jobStep, string key, BsonDocument value)
+			{
+				Id = ObjectId.GenerateNewId();
+				StreamId = job.StreamId;
+				TemplateRefId = job.TemplateId;
+				JobId = job.Id;
+				StepId = jobStep.Id;
+				Change = job.Change;
+				Key = key;
+				Data = value;
+			}
+		}
+
+		/// <summary>
+		/// The stream collection
+		/// </summary>
+		readonly IMongoCollection<TestDataDocument> _testDataDocuments;
+
+		/// <summary>
+		/// Constructor
+		/// </summary>
+		/// <param name="mongoService">The database service instance</param>
+		public TestDataCollection(MongoService mongoService)
+		{
+			List<MongoIndex<TestDataDocument>> indexes = new List<MongoIndex<TestDataDocument>>();
+			indexes.Add(keys => keys.Ascending(x => x.StreamId).Ascending(x => x.Change).Ascending(x => x.Key));
+			indexes.Add(keys => keys.Ascending(x => x.JobId).Ascending(x => x.StepId).Ascending(x => x.Key), unique: true);
+			_testDataDocuments = mongoService.GetCollection<TestDataDocument>("TestData", indexes);
+		}
+
+		/// <inheritdoc/>
+		public async Task<ITestData> AddAsync(IJob job, IJobStep step, string key, BsonDocument value)
+		{
+			TestDataDocument newTestData = new TestDataDocument(job, step, key, value);
+			await _testDataDocuments.InsertOneAsync(newTestData);
+			return newTestData;
+		}
+
+		/// <inheritdoc/>
+		public async Task<ITestData?> GetAsync(ObjectId id)
+		{
+			return await _testDataDocuments.Find<TestDataDocument>(x => x.Id == id).FirstOrDefaultAsync();
+		}
+
+		/// <inheritdoc/>
+		public async Task<List<ITestData>> FindAsync(StreamId? streamId, int? minChange, int? maxChange, JobId? jobId, SubResourceId? stepId, string? key = null, int index = 0, int count = 10)
+		{
+			FilterDefinition<TestDataDocument> filter = FilterDefinition<TestDataDocument>.Empty;
+			if (streamId != null)
+			{
+				filter &= Builders<TestDataDocument>.Filter.Eq(x => x.StreamId, streamId.Value);
+				if (minChange != null)
+				{
+					filter &= Builders<TestDataDocument>.Filter.Gte(x => x.Change, minChange.Value);
+				}
+				if (maxChange != null)
+				{
+					filter &= Builders<TestDataDocument>.Filter.Lte(x => x.Change, maxChange.Value);
+				}
+			}
+			if (jobId != null)
+			{
+				filter &= Builders<TestDataDocument>.Filter.Eq(x => x.JobId, jobId.Value);
+				if (stepId != null)
+				{
+					filter &= Builders<TestDataDocument>.Filter.Eq(x => x.StepId, stepId.Value);
+				}
+			}
+			if (key != null)
+			{
+				filter &= Builders<TestDataDocument>.Filter.Eq(x => x.Key, key);
+			}
+
+			SortDefinition<TestDataDocument> sort = Builders<TestDataDocument>.Sort.Ascending(x => x.StreamId).Descending(x => x.Change);
+
+			List<TestDataDocument> results = await _testDataDocuments.Find(filter).Sort(sort).Skip(index).Limit(count).ToListAsync();
+			return results.ConvertAll<ITestData>(x => x);
+		}
+
+		/// <inheritdoc/>
+		public async Task DeleteAsync(ObjectId id)
+		{
+			await _testDataDocuments.DeleteOneAsync<TestDataDocument>(x => x.Id == id);
+		}
+	}
+}
diff --git a/Engine/Source/Programs/Horde/Horde.Build/Jobs/UserMessages.cs b/Engine/Source/Programs/Horde/Horde.Build/Jobs/UserMessages.cs
new file mode 100644
--- /dev/null
+++ b/Engine/Source/Programs/Horde/Horde.Build/Jobs/UserMessages.cs
@@ -0,0 +1,172 @@
+// Copyright Epic Games, Inc. All Rights Reserved.
+
+using System;
+using System.Collections.Generic;
+using System.Linq;
+using System.Text.Json;
+using EpicGames.Core;
+using Horde.Build.Models;
+using MongoDB.Bson;
+
+namespace Horde.Build.Api
+{
+	/// <summary>
+	/// Response describing the current user
+	/// </summary>
+	public class GetUserResponse
+	{
+		/// <summary>
+		/// Id of the user
+		/// </summary>
+		public string Id { get; set; }
+
+		/// <summary>
+		/// Name of the user
+		/// </summary>
+		public string Name { get; set; }
+
+		/// <summary>
+		/// Avatar image URL (24px)
+		/// </summary>
+		public string? Image24 { get; set; }
+
+		/// <summary>
+		/// Avatar image URL (32px)
+		/// </summary>
+		public string? Image32 { get; set; }
+
+		/// <summary>
+		/// Avatar image URL (48px)
+		/// </summary>
+		public string? Image48 { get; set; }
+
+		/// <summary>
+		/// Avatar image URL (72px)
+		/// </summary>
+		public string? Image72 { get; set; }
+
+		/// <summary>
+		/// Email of the user
+		/// </summary>
+		public string? Email { get; set; }
+
+		/// <summary>
+		/// Claims for the user
+		/// </summary>
+		public List<UserClaim>? Claims { get; set; }
+
+		/// <summary>
+		/// Whether to enable experimental features for this user
+		/// </summary>
+		public bool? EnableExperimentalFeatures { get; set; }
+
+		/// <summary>
+		/// Settings for the dashboard
+		/// </summary>
+		public object? DashboardSettings { get; set; }
+
+		/// <summary>
+		/// List of pinned job ids
+		/// </summary>
+		public List<string>? PinnedJobIds { get; set; }
+
+		/// <summary>
+		/// Constructor
+		/// </summary>
+		public GetUserResponse(IUser user, IAvatar? avatar, IUserClaims? claims, IUserSettings? settings)
+		{
+			Id = user.Id.ToString();
+			Name = user.Name;
+			Email = user.Email;
+
+			Image24 = avatar?.Image24;
+			Image32 = avatar?.Image32;
+			Image48 = avatar?.Image48;
+			Image72 = avatar?.Image72;
+						
+			Claims = claims?.Claims.Select(x => new UserClaim(x)).ToList();
+
+			if (settings != null)
+			{
+				EnableExperimentalFeatures = settings.EnableExperimentalFeatures;
+				
+				DashboardSettings = BsonTypeMapper.MapToDotNetValue(settings.DashboardSettings);
+				PinnedJobIds = settings.PinnedJobIds.ConvertAll(x => x.ToString());
+			}
+		}
+	}
+
+	/// <summary>
+	/// Basic information about a user. May be embedded in other responses.
+	/// </summary>
+	public class GetThinUserInfoResponse
+	{
+		/// <summary>
+		/// Id of the user
+		/// </summary>
+		public string Id { get; set; }
+
+		/// <summary>
+		/// Name of the user
+		/// </summary>
+		public string Name { get; set; }
+
+		/// <summary>
+		/// The user's email address
+		/// </summary>
+		public string? Email { get; set; }
+
+		/// <summary>
+		/// The user login [DEPRECATED]
+		/// </summary>
+		internal string? Login { get; set; }
+
+		/// <summary>
+		/// Constructor
+		/// </summary>
+		/// <param name="user"></param>
+		public GetThinUserInfoResponse(IUser? user)
+		{
+			if (user == null)
+			{
+				Id = String.Empty;
+				Name = "(Unknown)";
+				Email = null;
+				Login = null;
+			}
+			else
+			{
+				Id = user.Id.ToString();
+				Name = user.Name;
+				Email = user.Email;
+				Login = user.Login;
+			}
+		}
+	}
+
+	/// <summary>
+	/// Request to update settings for a user
+	/// </summary>
+	public class UpdateUserRequest
+	{
+		/// <summary>
+		/// Whether to enable experimental features for this user
+		/// </summary>
+		public bool? EnableExperimentalFeatures { get; set; }
+
+		/// <summary>
+		/// New dashboard settings
+		/// </summary>
+		public JsonElement? DashboardSettings { get; set; }
+
+		/// <summary>
+		/// Job ids to add to the pinned list
+		/// </summary>
+		public List<string>? AddPinnedJobIds { get; set; }
+
+		/// <summary>
+		/// Jobs ids to remove from the pinned list
+		/// </summary>
+		public List<string>? RemovePinnedJobIds { get; set; }
+	}
+}
diff --git a/Engine/Source/Programs/Horde/Horde.Build/Logs/Builder/LocalLogBuilder.cs b/Engine/Source/Programs/Horde/Horde.Build/Logs/Builder/LocalLogBuilder.cs
new file mode 100644
--- /dev/null
+++ b/Engine/Source/Programs/Horde/Horde.Build/Logs/Builder/LocalLogBuilder.cs
@@ -0,0 +1,238 @@
+// Copyright Epic Games, Inc. All Rights Reserved.
+
+using System;
+using System.Collections.Concurrent;
+using System.Collections.Generic;
+using System.IO;
+using System.Threading.Tasks;
+using Horde.Build.Api;
+using Horde.Build.Models;
+using Horde.Build.Utilities;
+
+namespace Horde.Build.Logs.Builder
+{
+	using LogId = ObjectId<ILogFile>;
+
+	/// <summary>
+	/// In-memory implementation of a log write buffer
+	/// </summary>
+	class LocalLogBuilder : ILogBuilder
+	{
+		/// <summary>
+		/// Stores information about a sub-chunk which is still being written to
+		/// </summary>
+		class PendingSubChunk
+		{
+			/// <summary>
+			/// The type of data stored in this sub chunk
+			/// </summary>
+			public LogType Type { get; }
+
+			/// <summary>
+			/// Data for the sub-chunk
+			/// </summary>
+			public MemoryStream Stream { get; } = new MemoryStream(4096);
+
+			/// <summary>
+			/// Constructor
+			/// </summary>
+			/// <param name="type">Type of data stored in this subchunk</param>
+			public PendingSubChunk(LogType type)
+			{
+				Type = type;
+			}
+
+			/// <summary>
+			/// Converts this pending sub-chunk to a concrete LogSubChunkData object
+			/// </summary>
+			/// <param name="offset">Offset of this sub-chunk within the file</param>
+			/// <param name="lineIndex">The line index</param>
+			/// <returns>New sub-chunk data object</returns>
+			public LogSubChunkData ToSubChunkData(long offset, int lineIndex)
+			{
+				return new LogSubChunkData(Type, offset, lineIndex, new ReadOnlyLogText(Stream.ToArray()));
+			}
+		}
+
+		/// <summary>
+		/// Stores information about a pending log chunk
+		/// </summary>
+		class PendingChunk
+		{
+			/// <summary>
+			/// Time that this chunk was created
+			/// </summary>
+			public DateTime CreateTimeUtc { get; set; } = DateTime.UtcNow;
+
+			/// <summary>
+			/// Array of sub chunks for this chunk
+			/// </summary>
+			public List<PendingSubChunk> _subChunks = new List<PendingSubChunk>();
+
+			/// <summary>
+			/// The next sub-chunk
+			/// </summary>
+			public PendingSubChunk _nextSubChunk;
+
+			/// <summary>
+			/// Total length of this chunk
+			/// </summary>
+			public long _length;
+
+			/// <summary>
+			/// Whether the chunk is complete
+			/// </summary>
+			public bool _complete;
+
+			/// <summary>
+			/// Constructor
+			/// </summary>
+			/// <param name="type">Type of data to store in this file</param>
+			public PendingChunk(LogType type)
+			{
+				_nextSubChunk = new PendingSubChunk(type);
+			}
+
+			/// <summary>
+			/// Complete the current sub-chunk
+			/// </summary>
+			public void CompleteSubChunk()
+			{
+				if (_nextSubChunk.Stream.Length > 0)
+				{
+					_subChunks.Add(_nextSubChunk);
+					_nextSubChunk = new PendingSubChunk(_nextSubChunk.Type);
+				}
+			}
+		}
+
+		/// <summary>
+		/// Current log chunk state
+		/// </summary>
+		readonly ConcurrentDictionary<(LogId, long), PendingChunk> _pendingChunks = new ConcurrentDictionary<(LogId, long), PendingChunk>();
+
+		/// <inheritdoc/>
+		public bool FlushOnShutdown => true;
+
+		/// <inheritdoc/>
+		public Task<bool> AppendAsync(LogId logId, long chunkOffset, long writeOffset, int writeLineIndex, int writeLineCount, ReadOnlyMemory<byte> data, LogType type)
+		{
+			PendingChunk? pendingChunk;
+			while (!_pendingChunks.TryGetValue((logId, chunkOffset), out pendingChunk))
+			{
+				if(writeOffset != chunkOffset)
+				{
+					return Task.FromResult(false);
+				}
+
+				pendingChunk = new PendingChunk(type);
+
+				if (_pendingChunks.TryAdd((logId, chunkOffset), pendingChunk))
+				{
+					break;
+				}
+			}
+
+			lock (pendingChunk)
+			{
+				if (!pendingChunk._complete && chunkOffset + pendingChunk._length == writeOffset)
+				{
+					pendingChunk._nextSubChunk.Stream.Write(data.Span);
+					pendingChunk._length += data.Span.Length;
+					return Task.FromResult(true);
+				}
+			}
+			return Task.FromResult(false);
+		}
+
+		/// <inheritdoc/>
+		public Task CompleteSubChunkAsync(LogId logId, long offset)
+		{
+			PendingChunk? pendingChunk;
+			if (_pendingChunks.TryGetValue((logId, offset), out pendingChunk))
+			{
+				lock (pendingChunk)
+				{
+					pendingChunk.CompleteSubChunk();
+				}
+			}
+			return Task.CompletedTask;
+		}
+
+		/// <inheritdoc/>
+		public Task CompleteChunkAsync(LogId logId, long offset)
+		{
+			PendingChunk? pendingChunk;
+			if (_pendingChunks.TryGetValue((logId, offset), out pendingChunk))
+			{
+				lock (pendingChunk)
+				{
+					if (!pendingChunk._complete)
+					{
+						pendingChunk._complete = true;
+						pendingChunk.CompleteSubChunk();
+					}
+				}
+			}
+			return Task.CompletedTask;
+		}
+
+		/// <inheritdoc/>
+		public Task RemoveChunkAsync(LogId logId, long offset)
+		{
+			_pendingChunks.TryRemove((logId, offset), out _);
+			return Task.CompletedTask;
+		}
+
+		/// <inheritdoc/>
+		public Task<LogChunkData?> GetChunkAsync(LogId logId, long offset, int lineIndex)
+		{
+			PendingChunk? pendingChunk;
+			if (_pendingChunks.TryGetValue((logId, offset), out pendingChunk))
+			{
+				long subChunkOffset = offset;
+				int subChunkLineIndex = lineIndex;
+
+				List<LogSubChunkData> subChunks = new List<LogSubChunkData>();
+				foreach(PendingSubChunk pendingSubChunk in pendingChunk._subChunks)
+				{
+					LogSubChunkData subChunk = pendingSubChunk.ToSubChunkData(subChunkOffset, subChunkLineIndex);
+					subChunkOffset += subChunk.Length;
+					subChunkLineIndex += subChunk.LineCount;
+					subChunks.Add(subChunk);
+				}
+
+				PendingSubChunk nextSubChunk = pendingChunk._nextSubChunk;
+				if (nextSubChunk.Stream.Length > 0)
+				{
+					LogSubChunkData subChunkData = nextSubChunk.ToSubChunkData(subChunkOffset, subChunkLineIndex);
+					subChunks.Add(subChunkData);
+				}
+
+				return Task.FromResult<LogChunkData?>(new LogChunkData(offset, lineIndex, subChunks));
+			}
+
+			return Task.FromResult<LogChunkData?>(null);
+		}
+
+		/// <inheritdoc/>
+		public Task<List<(LogId, long)>> TouchChunksAsync(TimeSpan minAge)
+		{
+			DateTime utcNow = DateTime.UtcNow;
+
+			List<(LogId, long)> chunks = new List<(LogId, long)>();
+			foreach (KeyValuePair<(LogId, long), PendingChunk> pendingChunk in _pendingChunks.ToArray())
+			{
+				lock (pendingChunk.Value)
+				{
+					if (pendingChunk.Value.CreateTimeUtc < utcNow - minAge)
+					{
+						chunks.Add(pendingChunk.Key);
+						pendingChunk.Value.CreateTimeUtc = utcNow;
+					}
+				}
+			}
+			return Task.FromResult(chunks);
+		}
+	}
+}
diff --git a/Engine/Source/Programs/Horde/Horde.Build/Logs/Builder/RedisLogBuilder.cs b/Engine/Source/Programs/Horde/Horde.Build/Logs/Builder/RedisLogBuilder.cs
new file mode 100644
--- /dev/null
+++ b/Engine/Source/Programs/Horde/Horde.Build/Logs/Builder/RedisLogBuilder.cs
@@ -0,0 +1,314 @@
+// Copyright Epic Games, Inc. All Rights Reserved.
+
+using System;
+using System.Collections.Generic;
+using System.Globalization;
+using System.Text.RegularExpressions;
+using System.Threading.Tasks;
+using EpicGames.Core;
+using Horde.Build.Api;
+using Horde.Build.Models;
+using Horde.Build.Utilities;
+using Microsoft.Extensions.Logging;
+using OpenTracing;
+using OpenTracing.Util;
+using StackExchange.Redis;
+
+namespace Horde.Build.Logs.Builder
+{
+	using Condition = StackExchange.Redis.Condition;
+	using LogId = ObjectId<ILogFile>;
+
+	/// <summary>
+	/// Redis-based cache for log file chunks
+	/// </summary>
+	class RedisLogBuilder : ILogBuilder
+	{
+		const string ItemsKey = "log-items";
+
+		struct ChunkKeys
+		{
+			public string Prefix { get; }
+			public string Type => $"{Prefix}-type";
+			public string LineIndex => $"{Prefix}-lineindex";
+			public string Length => $"{Prefix}-length";
+			public string ChunkData => $"{Prefix}-chunk";
+			public string SubChunkData => $"{Prefix}-subchunk";
+			public string Complete => $"{Prefix}-complete";
+
+			public ChunkKeys(LogId logId, long offset)
+			{
+				Prefix = $"log-{logId}-chunk-{offset}-builder";
+			}
+
+			public static bool TryParse(string prefix, out LogId logId, out long offset)
+			{
+				Match match = Regex.Match(prefix, "log-([^-]+)-chunk-([^-]+)-builder");
+				if (!match.Success || !LogId.TryParse(match.Groups[1].Value, out logId) || !Int64.TryParse(match.Groups[2].Value, out offset))
+				{
+					logId = LogId.Empty;
+					offset = 0;
+					return false;
+				}
+				return true;
+			}
+		}
+
+		/// <inheritdoc/>
+		public bool FlushOnShutdown => false;
+
+		/// <summary>
+		/// The Redis database connection pool
+		/// </summary>
+		readonly RedisConnectionPool _redisConnectionPool;
+
+		/// <summary>
+		/// Logger for debug output
+		/// </summary>
+		readonly ILogger _logger;
+
+		/// <summary>
+		/// Constructor
+		/// </summary>
+		/// <param name="redisConnectionPool">The redis database singleton</param>
+		/// <param name="logger">Logger for debug output</param>
+		public RedisLogBuilder(RedisConnectionPool redisConnectionPool, ILogger logger)
+		{
+			_redisConnectionPool = redisConnectionPool;
+			_logger = logger;
+		}
+
+		/// <inheritdoc/>
+		public async Task<bool> AppendAsync(LogId logId, long chunkOffset, long writeOffset, int writeLineIndex, int writeLineCount, ReadOnlyMemory<byte> data, LogType type)
+		{
+			IDatabase redisDb = _redisConnectionPool.GetDatabase();
+			ChunkKeys keys = new ChunkKeys(logId, chunkOffset);
+
+			if(chunkOffset == writeOffset)
+			{
+				using IScope scope = GlobalTracer.Instance.BuildSpan("Redis.CreateChunk").StartActive();
+				scope.Span.SetTag("LogId", logId.ToString());
+				scope.Span.SetTag("Offset", chunkOffset.ToString(CultureInfo.InvariantCulture));
+				scope.Span.SetTag("WriteOffset", writeOffset.ToString(CultureInfo.InvariantCulture));
+
+				ITransaction? createTransaction = redisDb.CreateTransaction();
+				createTransaction.AddCondition(Condition.SortedSetNotContains(ItemsKey, keys.Prefix));
+				_ = createTransaction.SortedSetAddAsync(ItemsKey, keys.Prefix, DateTime.UtcNow.Ticks);
+				_ = createTransaction.StringSetAsync(keys.Type, (int)type);
+				_ = createTransaction.StringSetAsync(keys.Length, data.Length);
+				_ = createTransaction.StringSetAsync(keys.LineIndex, writeLineIndex + writeLineCount);
+				_ = createTransaction.StringSetAsync(keys.SubChunkData, data);
+				if (await createTransaction.ExecuteAsync())
+				{
+					_logger.LogTrace("Created {Size} bytes in {Key}", data.Length, keys.SubChunkData);
+					return true;
+				}
+			}
+
+			{
+				using IScope scope = GlobalTracer.Instance.BuildSpan("Redis.AppendChunk").StartActive();
+				scope.Span.SetTag("LogId", logId.ToString());
+				scope.Span.SetTag("Offset", chunkOffset.ToString(CultureInfo.InvariantCulture));
+				scope.Span.SetTag("WriteOffset", writeOffset.ToString(CultureInfo.InvariantCulture));
+
+				ITransaction? appendTransaction = redisDb.CreateTransaction();
+				appendTransaction.AddCondition(Condition.SortedSetContains(ItemsKey, keys.Prefix));
+				appendTransaction.AddCondition(Condition.KeyNotExists(keys.Complete));
+				appendTransaction.AddCondition(Condition.StringEqual(keys.Type, (int)type));
+				appendTransaction.AddCondition(Condition.StringEqual(keys.Length, (int)(writeOffset - chunkOffset)));
+				_ = appendTransaction.StringAppendAsync(keys.SubChunkData, data);
+				_ = appendTransaction.StringSetAsync(keys.Length, (int)(writeOffset - chunkOffset) + data.Length);
+				_ = appendTransaction.StringSetAsync(keys.LineIndex, writeLineIndex + writeLineCount);
+				if (await appendTransaction.ExecuteAsync())
+				{
+					_logger.LogTrace("Appended {Size} to {Key}", data.Length, keys.SubChunkData);
+					return true;
+				}
+			}
+
+			return false;
+		}
+
+		/// <inheritdoc/>
+		public async Task CompleteSubChunkAsync(LogId logId, long offset)
+		{
+			IDatabase redisDb = _redisConnectionPool.GetDatabase();
+			ChunkKeys keys = new ChunkKeys(logId, offset);
+			for (; ; )
+			{
+				using IScope scope = GlobalTracer.Instance.BuildSpan("Redis.CompleteSubChunk").StartActive();
+				scope.Span.SetTag("LogId", logId.ToString());
+				scope.Span.SetTag("Offset", offset.ToString(CultureInfo.InvariantCulture));
+
+				LogType type = (LogType)(int)await redisDb.StringGetAsync(keys.Type);
+				int length = (int)await redisDb.StringGetAsync(keys.Length);
+				int lineIndex = (int)await redisDb.StringGetAsync(keys.LineIndex);
+				long chunkDataLength = await redisDb.StringLengthAsync(keys.ChunkData);
+
+				RedisValue subChunkTextValue = await redisDb.StringGetAsync(keys.SubChunkData);
+				if (subChunkTextValue.IsNullOrEmpty)
+				{
+					break;
+				}
+
+				ReadOnlyLogText subChunkText = new ReadOnlyLogText(subChunkTextValue);
+				LogSubChunkData subChunkData = new LogSubChunkData(type, offset + length - subChunkText.Data.Length, lineIndex - subChunkText.LineCount, subChunkText);
+
+				byte[] subChunkDataBytes;
+				try
+				{
+					subChunkDataBytes = subChunkData.ToByteArray(_logger);
+				}
+				catch (Exception ex)
+				{
+					throw new LogTextParseException($"Error while parsing sub-chunk for {logId} offset {offset}", ex);
+				}
+
+				ITransaction writeTransaction = redisDb.CreateTransaction();
+
+				writeTransaction.AddCondition(Condition.StringLengthEqual(keys.ChunkData, chunkDataLength));
+				writeTransaction.AddCondition(Condition.StringLengthEqual(keys.SubChunkData, subChunkText.Length));
+				writeTransaction.AddCondition(Condition.StringEqual(keys.Length, length));
+				writeTransaction.AddCondition(Condition.StringEqual(keys.LineIndex, lineIndex));
+				Task<long> newLength = writeTransaction.StringAppendAsync(keys.ChunkData, subChunkDataBytes);
+				_ = writeTransaction.KeyDeleteAsync(keys.SubChunkData);
+
+				if (await writeTransaction.ExecuteAsync())
+				{
+					_logger.LogDebug("Completed sub-chunk for log {LogId} chunk offset {Offset} -> sub-chunk size {SubChunkSize}, chunk size {ChunkSize}", logId, offset, subChunkDataBytes.Length, await newLength);
+					break;
+				}
+			}
+		}
+
+		/// <inheritdoc/>
+		public async Task CompleteChunkAsync(LogId logId, long offset)
+		{
+			IDatabase redisDb = _redisConnectionPool.GetDatabase();
+			using IScope scope = GlobalTracer.Instance.BuildSpan("Redis.CompleteChunk").StartActive();
+			scope.Span.SetTag("LogId", logId.ToString());
+			scope.Span.SetTag("Offset", offset.ToString(CultureInfo.InvariantCulture));
+
+			ChunkKeys keys = new ChunkKeys(logId, offset);
+
+			ITransaction transaction = redisDb.CreateTransaction();
+			transaction.AddCondition(Condition.SortedSetContains(ItemsKey, keys.Prefix));
+			_ = transaction.StringSetAsync(keys.Complete, true);
+			if(!await transaction.ExecuteAsync())
+			{
+				_logger.LogDebug("Log {LogId} chunk offset {Offset} is not in Redis builder", logId, offset);
+				return;
+			}
+
+			await CompleteSubChunkAsync(logId, offset);
+		}
+
+		/// <inheritdoc/>
+		public async Task RemoveChunkAsync(LogId logId, long offset)
+		{
+			IDatabase redisDb = _redisConnectionPool.GetDatabase();
+			using IScope scope = GlobalTracer.Instance.BuildSpan("Redis.RemoveChunk").StartActive();
+			scope.Span.SetTag("LogId", logId.ToString());
+			scope.Span.SetTag("Offset", offset.ToString(CultureInfo.InvariantCulture));
+
+			ChunkKeys keys = new ChunkKeys(logId, offset);
+
+			ITransaction transaction = redisDb.CreateTransaction();
+			_ = transaction.KeyDeleteAsync(keys.Type);
+			_ = transaction.KeyDeleteAsync(keys.LineIndex);
+			_ = transaction.KeyDeleteAsync(keys.Length);
+			_ = transaction.KeyDeleteAsync(keys.ChunkData);
+			_ = transaction.KeyDeleteAsync(keys.SubChunkData);
+			_ = transaction.KeyDeleteAsync(keys.Complete);
+
+			_ = transaction.SortedSetRemoveAsync(ItemsKey, keys.Prefix);
+			await transaction.ExecuteAsync();
+		}
+
+		/// <inheritdoc/>
+		public async Task<LogChunkData?> GetChunkAsync(LogId logId, long offset, int lineIndex)
+		{
+			IDatabase redisDb = _redisConnectionPool.GetDatabase();
+			ChunkKeys keys = new ChunkKeys(logId, offset);
+			for (; ; )
+			{
+				using IScope scope = GlobalTracer.Instance.BuildSpan("Redis.GetChunk").StartActive();
+				scope.Span.SetTag("LogId", logId.ToString());
+				scope.Span.SetTag("Offset", offset.ToString(CultureInfo.InvariantCulture));
+
+				RedisValue chunkDataValue = await redisDb.StringGetAsync(keys.ChunkData);
+
+				ReadOnlyMemory<byte> chunkData = chunkDataValue;
+
+				ITransaction transaction = redisDb.CreateTransaction();
+				transaction.AddCondition(Condition.StringLengthEqual(keys.ChunkData, chunkData.Length));
+				Task<RedisValue> typeTask = transaction.StringGetAsync(keys.Type);
+				Task<RedisValue> lastSubChunkDataTask = transaction.StringGetAsync(keys.SubChunkData);
+
+				if (await transaction.ExecuteAsync())
+				{
+					MemoryReader reader = new MemoryReader(chunkData);
+
+					long subChunkOffset = offset;
+					int subChunkLineIndex = lineIndex;
+
+					List<LogSubChunkData> subChunks = new List<LogSubChunkData>();
+					while (reader.Offset < reader.Length)
+					{
+						LogSubChunkData subChunkData = reader.ReadLogSubChunkData(subChunkOffset, subChunkLineIndex);
+						subChunkOffset += subChunkData.Length;
+						subChunkLineIndex += subChunkData.LineCount;
+						subChunks.Add(subChunkData);
+					}
+
+					RedisValue type = await typeTask;
+
+					RedisValue lastSubChunkData = await lastSubChunkDataTask;
+					if (lastSubChunkData.Length() > 0)
+					{
+						ReadOnlyLogText subChunkDataText = new ReadOnlyLogText(lastSubChunkData);
+						subChunks.Add(new LogSubChunkData((LogType)(int)type, subChunkOffset, subChunkLineIndex, subChunkDataText));
+					}
+
+					if (subChunks.Count == 0)
+					{
+						break;
+					}
+
+					return new LogChunkData(offset, lineIndex, subChunks);
+				}
+			}
+			return null;
+		}
+
+		/// <inheritdoc/>
+		public async Task<List<(LogId, long)>> TouchChunksAsync(TimeSpan minAge)
+		{
+			IDatabase redisDb = _redisConnectionPool.GetDatabase();
+			
+			// Find all the chunks that are suitable for expiry
+			DateTime utcNow = DateTime.UtcNow;
+			SortedSetEntry[] entries = await redisDb.SortedSetRangeByScoreWithScoresAsync(ItemsKey, stop: (utcNow - minAge).Ticks);
+
+			// Update the score for each element in a transaction. If it succeeds, we can write the chunk. Otherwise another pod has beat us to it.
+			List<(LogId, long)> results = new List<(LogId, long)>();
+			foreach (SortedSetEntry entry in entries)
+			{
+				if (ChunkKeys.TryParse(entry.Element.ToString(), out LogId logId, out long offset))
+				{
+					ITransaction transaction = redisDb.CreateTransaction();
+					transaction.AddCondition(Condition.SortedSetEqual(ItemsKey, entry.Element, entry.Score));
+					_ = transaction.SortedSetAddAsync(ItemsKey, entry.Element, utcNow.Ticks);
+
+					if (!await transaction.ExecuteAsync())
+					{
+						break;
+					}
+
+					results.Add((logId, offset));
+				}
+			}
+			return results;
+		}
+	}
+}
diff --git a/Engine/Source/Programs/Horde/Horde.Build/Logs/ILogBuilder.cs b/Engine/Source/Programs/Horde/Horde.Build/Logs/ILogBuilder.cs
new file mode 100644
--- /dev/null
+++ b/Engine/Source/Programs/Horde/Horde.Build/Logs/ILogBuilder.cs
@@ -0,0 +1,77 @@
+// Copyright Epic Games, Inc. All Rights Reserved.
+
+using System;
+using System.Collections.Generic;
+using System.Threading.Tasks;
+using Horde.Build.Api;
+using Horde.Build.Models;
+using Horde.Build.Utilities;
+
+namespace Horde.Build.Logs
+{
+	using LogId = ObjectId<ILogFile>;
+
+	/// <summary>
+	/// Interface for the log file write cache
+	/// </summary>
+	public interface ILogBuilder
+	{
+		/// <summary>
+		/// Whether the cache should be flushed on shutdown
+		/// </summary>
+		bool FlushOnShutdown { get; }
+
+		/// <summary>
+		/// Append data to a key
+		/// </summary>
+		/// <param name="logId">The log file id</param>
+		/// <param name="chunkOffset">Offset of the chunk within the log file</param>
+		/// <param name="writeOffset">Offset to write to</param>
+		/// <param name="writeLineIndex">Line index being written</param>
+		/// <param name="writeLineCount">Line count being written</param>
+		/// <param name="data">Data to be appended</param>
+		/// <param name="type">Type of data stored in this log file</param>
+		/// <returns>True if the data was appended to the given chunk. False if the chunk has been completed.</returns>
+		Task<bool> AppendAsync(LogId logId, long chunkOffset, long writeOffset, int writeLineIndex, int writeLineCount, ReadOnlyMemory<byte> data, LogType type);
+
+		/// <summary>
+		/// Finish the current sub chunk
+		/// </summary>
+		/// <param name="logId">The log file id</param>
+		/// <param name="offset">Offset of the chunk within the log file</param>
+		/// <returns>Async task</returns>
+		Task CompleteSubChunkAsync(LogId logId, long offset);
+
+		/// <summary>
+		/// Finish the current chunk
+		/// </summary>
+		/// <param name="logId">The log file id</param>
+		/// <param name="offset">Offset of the chunk within the log file</param>
+		/// <returns>Async task</returns>
+		Task CompleteChunkAsync(LogId logId, long offset);
+
+		/// <summary>
+		/// Remove a complete chunk from the builder
+		/// </summary>
+		/// <param name="logId">The log file id</param>
+		/// <param name="offset">Offset of the chunk within the log file</param>
+		/// <returns>Async task</returns>
+		Task RemoveChunkAsync(LogId logId, long offset);
+
+		/// <summary>
+		/// Gets the current chunk for the given log file
+		/// </summary>
+		/// <param name="logId">The log file id</param>
+		/// <param name="offset">Offset of the chunk within the log file</param>
+		/// <param name="lineIndex">Line index of the chunk within the log file</param>
+		/// <returns></returns>
+		Task<LogChunkData?> GetChunkAsync(LogId logId, long offset, int lineIndex);
+
+		/// <summary>
+		/// Touches the timestamps of all the chunks after the given age, and returns them. Used for flushing the builder.
+		/// </summary>
+		/// <param name="minAge">Minimum age of the chunks to enumerate. If specified, only chunks last modified longer than this period will be returned.</param>
+		/// <returns>List of chunks, identified by log id and chunk index</returns>
+		Task<List<(LogId, long)>> TouchChunksAsync(TimeSpan minAge);
+	}
+}
diff --git a/Engine/Source/Programs/Horde/Horde.Build/Logs/ILogStorage.cs b/Engine/Source/Programs/Horde/Horde.Build/Logs/ILogStorage.cs
new file mode 100644
--- /dev/null
+++ b/Engine/Source/Programs/Horde/Horde.Build/Logs/ILogStorage.cs
@@ -0,0 +1,52 @@
+// Copyright Epic Games, Inc. All Rights Reserved.
+
+using System;
+using System.Threading.Tasks;
+using Horde.Build.Models;
+using Horde.Build.Utilities;
+
+namespace Horde.Build.Logs
+{
+	using LogId = ObjectId<ILogFile>;
+
+	/// <summary>
+	/// Caching interface for reading and writing log data.
+	/// </summary>
+	public interface ILogStorage : IDisposable
+	{
+		/// <summary>
+		/// Attempts to read an index for the given log file
+		/// </summary>
+		/// <param name="logId">Unique id of the log file</param>
+		/// <param name="length">Length of the file covered by the index</param>
+		/// <returns>Index for the log file</returns>
+		Task<LogIndexData?> ReadIndexAsync(LogId logId, long length);
+
+		/// <summary>
+		/// Log file to write an index for
+		/// </summary>
+		/// <param name="logId">Unique id of the log file</param>
+		/// <param name="length">Length of the file covered by the index</param>
+		/// <param name="index">The log file index</param>
+		/// <returns>Async task</returns>
+		Task WriteIndexAsync(LogId logId, long length, LogIndexData index);
+
+		/// <summary>
+		/// Retrieves an item from the cache
+		/// </summary>
+		/// <param name="logId">Unique id of the log file</param>
+		/// <param name="offset">Offset of the chunk to read</param>
+		/// <param name="lineIndex">First line of the chunk</param>
+		/// <returns>Data for the given key, or null if it's not present</returns>
+		Task<LogChunkData?> ReadChunkAsync(LogId logId, long offset, int lineIndex);
+
+		/// <summary>
+		/// Writes a chunk to storage
+		/// </summary>
+		/// <param name="logId">Unique id of the log file</param>
+		/// <param name="offset">Offset of the chunk to write</param>
+		/// <param name="chunkData">Information about the chunk data</param>
+		/// <returns>Async task</returns>
+		Task WriteChunkAsync(LogId logId, long offset, LogChunkData chunkData);
+	}
+}
diff --git a/Engine/Source/Programs/Horde/Horde.Build/Logs/LogEventCollection.cs b/Engine/Source/Programs/Horde/Horde.Build/Logs/LogEventCollection.cs
new file mode 100644
--- /dev/null
+++ b/Engine/Source/Programs/Horde/Horde.Build/Logs/LogEventCollection.cs
@@ -0,0 +1,190 @@
+// Copyright Epic Games, Inc. All Rights Reserved.
+
+using System;
+using System.Collections.Generic;
+using System.Linq;
+using System.Threading.Tasks;
+using Horde.Build.Models;
+using Horde.Build.Server;
+using Horde.Build.Utilities;
+using HordeCommon;
+using Microsoft.Extensions.Logging;
+using MongoDB.Bson;
+using MongoDB.Bson.Serialization.Attributes;
+using MongoDB.Driver;
+
+namespace Horde.Build.Collections.Impl
+{
+	using LogId = ObjectId<ILogFile>;
+
+	/// <summary>
+	/// Collection of event documents
+	/// </summary>
+	public class LogEventCollection : ILogEventCollection
+	{
+		class LogEventId
+		{
+			[BsonElement("l")]
+			public LogId LogId { get; set; }
+
+			[BsonElement("n")]
+			public int LineIndex { get; set; }
+		}
+
+		class LogEventDocument : ILogEvent
+		{
+			[BsonId]
+			public LogEventId Id { get; set; }
+
+			[BsonElement("w"), BsonIgnoreIfDefault, BsonDefaultValue(false)]
+			public bool IsWarning { get; set; }
+
+			[BsonElement("c"), BsonIgnoreIfNull]
+			public int? LineCount { get; set; }
+
+			[BsonElement("s")]
+			public ObjectId? SpanId { get; set; }
+
+			LogId ILogEvent.LogId => Id.LogId;
+			int ILogEvent.LineIndex => Id.LineIndex;
+			int ILogEvent.LineCount => LineCount ?? 1;
+			EventSeverity ILogEvent.Severity => IsWarning ? EventSeverity.Warning : EventSeverity.Error;
+
+			public LogEventDocument()
+			{
+				Id = new LogEventId();
+			}
+
+			public LogEventDocument(LogId logId, EventSeverity severity, int lineIndex, int lineCount, ObjectId? spanId)
+			{
+				Id = new LogEventId { LogId = logId, LineIndex = lineIndex };
+				IsWarning = severity == EventSeverity.Warning;
+				LineCount = (lineCount > 1) ? (int?)lineCount : null;
+				SpanId = spanId;
+			}
+
+			public LogEventDocument(NewLogEventData data)
+				: this(data.LogId, data.Severity, data.LineIndex, data.LineCount, data.SpanId)
+			{
+			}
+		}
+
+		class LegacyEventDocument
+		{
+			public ObjectId Id { get; set; }
+			public DateTime Time { get; set; }
+			public EventSeverity Severity { get; set; }
+			public LogId LogId { get; set; }
+			public int LineIndex { get; set; }
+			public int LineCount { get; set; }
+
+			public string? Message { get; set; }
+
+			[BsonIgnoreIfNull, BsonElement("IssueId2")]
+			public int? IssueId { get; set; }
+
+			public BsonDocument? Data { get; set; }
+
+			public int UpgradeVersion { get; set; }
+		}
+
+		/// <summary>
+		/// Collection of event documents
+		/// </summary>
+		readonly IMongoCollection<LogEventDocument> _logEvents;
+
+		/// <summary>
+		/// Collection of legacy event documents
+		/// </summary>
+		readonly IMongoCollection<LegacyEventDocument> _legacyEvents;
+
+		/// <summary>
+		/// Logger for upgrade messages
+		/// </summary>
+		readonly ILogger<LogEventCollection> _logger;
+
+		/// <summary>
+		/// Constructor
+		/// </summary>
+		/// <param name="mongoService">The database service instance</param>
+		/// <param name="logger">The logger instance</param>
+		public LogEventCollection(MongoService mongoService, ILogger<LogEventCollection> logger)
+		{
+			_logger = logger;
+
+			List<MongoIndex<LogEventDocument>> logEventIndexes = new List<MongoIndex<LogEventDocument>>();
+			logEventIndexes.Add(keys => keys.Ascending(x => x.Id.LogId));
+			logEventIndexes.Add(keys => keys.Ascending(x => x.SpanId).Ascending(x => x.Id));
+			_logEvents = mongoService.GetCollection<LogEventDocument>("LogEvents", logEventIndexes);
+
+			_legacyEvents = mongoService.GetCollection<LegacyEventDocument>("Events", keys => keys.Ascending(x => x.LogId));
+		}
+
+		/// <inheritdoc/>
+		public Task AddAsync(NewLogEventData newEvent)
+		{
+			return _logEvents.InsertOneAsync(new LogEventDocument(newEvent));
+		}
+
+		/// <inheritdoc/>
+		public Task AddManyAsync(List<NewLogEventData> newEvents)
+		{
+			return _logEvents.InsertManyAsync(newEvents.ConvertAll(x => new LogEventDocument(x)));
+		}
+
+		/// <inheritdoc/>
+		public async Task<List<ILogEvent>> FindAsync(LogId logId, ObjectId? spanId = null, int? index = null, int? count = null)
+		{
+			_logger.LogInformation("Querying for log events for log {LogId} creation time {CreateTime}", logId, logId.Value.CreationTime);
+
+			FilterDefinitionBuilder<LogEventDocument> builder = Builders<LogEventDocument>.Filter;
+
+			FilterDefinition<LogEventDocument> filter = builder.Eq(x => x.Id.LogId, logId);
+			if(spanId != null)
+			{
+				filter &= builder.Eq(x => x.SpanId, spanId.Value);
+			}
+
+			IFindFluent<LogEventDocument, LogEventDocument> results = _logEvents.Find(filter).SortBy(x => x.Id);
+			if (index != null)
+			{
+				results = results.Skip(index.Value);
+			}
+			if (count != null)
+			{
+				results = results.Limit(count.Value);
+			}
+
+			List<LogEventDocument> eventDocuments = await results.ToListAsync();
+			return eventDocuments.ConvertAll<ILogEvent>(x => x);
+		}
+
+		/// <inheritdoc/>
+		public async Task<List<ILogEvent>> FindEventsForSpansAsync(IEnumerable<ObjectId> spanIds, LogId[]? logIds, int index, int count)
+		{
+			FilterDefinition<LogEventDocument> filter = Builders<LogEventDocument>.Filter.In(x => x.SpanId, spanIds.Select<ObjectId, ObjectId?>(x => x));
+			if (logIds != null && logIds.Length > 0)
+			{
+				filter &= Builders<LogEventDocument>.Filter.In(x => x.Id.LogId, logIds);
+			}
+
+			List<LogEventDocument> eventDocuments = await _logEvents.Find(filter).Skip(index).Limit(count).ToListAsync();
+			return eventDocuments.ConvertAll<ILogEvent>(x => x);
+		}
+
+		/// <inheritdoc/>
+		public async Task DeleteLogAsync(LogId logId)
+		{
+			await _logEvents.DeleteManyAsync(x => x.Id.LogId == logId);
+			await _legacyEvents.DeleteManyAsync(x => x.LogId == logId);
+		}
+
+		/// <inheritdoc/>
+		public async Task AddSpanToEventsAsync(IEnumerable<ILogEvent> events, ObjectId spanId)
+		{
+			FilterDefinition<LogEventDocument> eventFilter = Builders<LogEventDocument>.Filter.In(x => x.Id, events.OfType<LogEventDocument>().Select(x => x.Id));
+			UpdateDefinition<LogEventDocument> eventUpdate = Builders<LogEventDocument>.Update.Set(x => x.SpanId, spanId);
+			await _logEvents.UpdateManyAsync(eventFilter, eventUpdate);
+		}
+	}
+}
diff --git a/Engine/Source/Programs/Horde/Horde.Build/Logs/LogEventMessages.cs b/Engine/Source/Programs/Horde/Horde.Build/Logs/LogEventMessages.cs
new file mode 100644
--- /dev/null
+++ b/Engine/Source/Programs/Horde/Horde.Build/Logs/LogEventMessages.cs
@@ -0,0 +1,62 @@
+// Copyright Epic Games, Inc. All Rights Reserved.
+
+using System.Collections.Generic;
+using System.Text.Json;
+using EpicGames.Core;
+using Horde.Build.Models;
+using HordeCommon;
+
+namespace Horde.Build.Api
+{
+	/// <summary>
+	/// Information about an uploaded event
+	/// </summary>
+	public class GetLogEventResponse
+	{
+		/// <summary>
+		/// Unique id of the log containing this event
+		/// </summary>
+		public string LogId { get; set; }
+
+		/// <summary>
+		/// Severity of this event
+		/// </summary>
+		public EventSeverity Severity { get; set; }
+
+		/// <summary>
+		/// Index of the first line for this event
+		/// </summary>
+		public int LineIndex { get; set; }
+
+		/// <summary>
+		/// Number of lines in the event
+		/// </summary>
+		public int LineCount { get; set; }
+
+		/// <summary>
+		/// The issue id associated with this event
+		/// </summary>
+		public int? IssueId { get; set; }
+
+		/// <summary>
+		/// The structured message data for this event
+		/// </summary>
+		public List<JsonElement> Lines { get; set; }
+
+		/// <summary>
+		/// Constructor
+		/// </summary>
+		/// <param name="logEvent">The event to construct from</param>
+		/// <param name="eventData">The event data</param>
+		/// <param name="issueId">The issue for this event</param>
+		public GetLogEventResponse(ILogEvent logEvent, ILogEventData eventData, int? issueId)
+		{
+			Severity = logEvent.Severity;
+			LogId = logEvent.LogId.ToString();
+			LineIndex = logEvent.LineIndex;
+			LineCount = logEvent.LineCount;
+			IssueId = issueId;
+			Lines = eventData.Lines.ConvertAll(x => x.Data);
+		}
+	}
+}
diff --git a/Engine/Source/Programs/Horde/Horde.Build/Logs/LogFileCollection.cs b/Engine/Source/Programs/Horde/Horde.Build/Logs/LogFileCollection.cs
new file mode 100644
--- /dev/null
+++ b/Engine/Source/Programs/Horde/Horde.Build/Logs/LogFileCollection.cs
@@ -0,0 +1,233 @@
+// Copyright Epic Games, Inc. All Rights Reserved.
+
+using System.Collections.Generic;
+using System.Linq;
+using System.Net;
+using System.Threading.Tasks;
+using Horde.Build.Api;
+using Horde.Build.Models;
+using Horde.Build.Server;
+using Horde.Build.Utilities;
+using MongoDB.Bson.Serialization.Attributes;
+using MongoDB.Driver;
+
+namespace Horde.Build.Collections.Impl
+{
+	using JobId = ObjectId<IJob>;
+	using LogId = ObjectId<ILogFile>;
+	using SessionId = ObjectId<ISession>;
+
+	/// <summary>
+	/// Wrapper around the jobs collection in a mongo DB
+	/// </summary>
+	public class LogFileCollection : ILogFileCollection
+	{
+		class LogChunkDocument : ILogChunk
+		{
+			public long Offset { get; set; }
+			public int Length { get; set; }
+			public int LineIndex { get; set; }
+
+			[BsonIgnoreIfNull]
+			public string? Server { get; set; }
+
+			[BsonConstructor]
+			public LogChunkDocument()
+			{
+			}
+
+			public LogChunkDocument(LogChunkDocument other)
+			{
+				Offset = other.Offset;
+				Length = other.Length;
+				LineIndex = other.LineIndex;
+				Server = other.Server;
+			}
+
+			public LogChunkDocument Clone()
+			{
+				return (LogChunkDocument)MemberwiseClone();
+			}
+		}
+
+		class LogFileDocument : ILogFile
+		{
+			[BsonRequired, BsonId]
+			public LogId Id { get; set; }
+
+			[BsonRequired]
+			public JobId JobId { get; set; }
+
+			public SessionId? SessionId { get; set; }
+			public LogType Type { get; set; }
+
+			[BsonIgnoreIfNull]
+			public int? MaxLineIndex { get; set; }
+
+			[BsonIgnoreIfNull]
+			public long? IndexLength { get; set; }
+
+			public List<LogChunkDocument> Chunks { get; set; } = new List<LogChunkDocument>();
+
+			[BsonRequired]
+			public int UpdateIndex { get; set; }
+
+			IReadOnlyList<ILogChunk> ILogFile.Chunks => Chunks;
+
+			[BsonConstructor]
+			private LogFileDocument()
+			{
+			}
+
+			public LogFileDocument(JobId jobId, SessionId? sessionId, LogType type)
+			{
+				Id = LogId.GenerateNewId();
+				JobId = jobId;
+				SessionId = sessionId;
+				Type = type;
+				MaxLineIndex = 0;
+			}
+
+			public LogFileDocument Clone()
+			{
+				LogFileDocument document = (LogFileDocument)MemberwiseClone();
+				document.Chunks = document.Chunks.ConvertAll(x => x.Clone());
+				return document;
+			}
+		}
+
+		/// <summary>
+		/// The jobs collection
+		/// </summary>
+		readonly IMongoCollection<LogFileDocument> _logFiles;
+
+		/// <summary>
+		/// Hostname for the current server
+		/// </summary>
+		readonly string _hostName;
+
+		/// <summary>
+		/// Constructor
+		/// </summary>
+		/// <param name="mongoService">The database service singleton</param>
+		public LogFileCollection(MongoService mongoService)
+		{
+			_logFiles = mongoService.GetCollection<LogFileDocument>("LogFiles");
+			_hostName = Dns.GetHostName();
+		}
+
+		/// <inheritdoc/>
+		public async Task<ILogFile> CreateLogFileAsync(JobId jobId, SessionId? sessionId, LogType type)
+		{
+			LogFileDocument newLogFile = new LogFileDocument(jobId, sessionId, type);
+			await _logFiles.InsertOneAsync(newLogFile);
+			return newLogFile;
+		}
+
+		/// <inheritdoc/>
+		public async Task<ILogFile?> TryAddChunkAsync(ILogFile logFileInterface, long offset, int lineIndex)
+		{
+			LogFileDocument logFile = ((LogFileDocument)logFileInterface).Clone();
+
+			int chunkIdx = logFile.Chunks.GetChunkForOffset(offset) + 1;
+
+			LogChunkDocument chunk = new LogChunkDocument();
+			chunk.Offset = offset;
+			chunk.LineIndex = lineIndex;
+			chunk.Server = _hostName;
+			logFile.Chunks.Insert(chunkIdx, chunk);
+
+			UpdateDefinition<LogFileDocument> update = Builders<LogFileDocument>.Update.Set(x => x.Chunks, logFile.Chunks);
+			if (chunkIdx == logFile.Chunks.Count - 1)
+			{
+				logFile.MaxLineIndex = null;
+				update = update.Unset(x => x.MaxLineIndex);
+			}
+
+			if (!await TryUpdateLogFileAsync(logFile, update))
+			{
+				return null;
+			}
+
+			return logFile;
+		}
+
+		/// <inheritdoc/>
+		public async Task<ILogFile?> TryCompleteChunksAsync(ILogFile logFileInterface, IEnumerable<CompleteLogChunkUpdate> chunkUpdates)
+		{
+			LogFileDocument logFile = ((LogFileDocument)logFileInterface).Clone();
+
+			// Update the length of any complete chunks
+			UpdateDefinitionBuilder<LogFileDocument> updateBuilder = Builders<LogFileDocument>.Update;
+			List<UpdateDefinition<LogFileDocument>> updates = new List<UpdateDefinition<LogFileDocument>>();
+			foreach (CompleteLogChunkUpdate chunkUpdate in chunkUpdates)
+			{
+				LogChunkDocument chunk = logFile.Chunks[chunkUpdate.Index];
+				chunk.Length = chunkUpdate.Length;
+				updates.Add(updateBuilder.Set(x => x.Chunks[chunkUpdate.Index].Length, chunkUpdate.Length));
+
+				if (chunkUpdate.Index == logFile.Chunks.Count - 1)
+				{
+					logFile.MaxLineIndex = chunk.LineIndex + chunkUpdate.LineCount;
+					updates.Add(updateBuilder.Set(x => x.MaxLineIndex, logFile.MaxLineIndex));
+				}
+			}
+
+			// Try to apply the updates
+			if (updates.Count > 0 && !await TryUpdateLogFileAsync(logFile, updateBuilder.Combine(updates)))
+			{
+				return null;
+			}
+
+			return logFile;
+		}
+
+		/// <inheritdoc/>
+		public async Task<ILogFile?> TryUpdateIndexAsync(ILogFile logFileInterface, long newIndexLength)
+		{
+			LogFileDocument logFile = ((LogFileDocument)logFileInterface).Clone();
+
+			UpdateDefinition<LogFileDocument> update = Builders<LogFileDocument>.Update.Set(x => x.IndexLength, newIndexLength);
+			if (!await TryUpdateLogFileAsync(logFile, update))
+			{
+				return null;
+			}
+
+			logFile.IndexLength = newIndexLength;
+			return logFile;
+		}
+
+		/// <inheritdoc/>
+		private async Task<bool> TryUpdateLogFileAsync(LogFileDocument current, UpdateDefinition<LogFileDocument> update)
+		{
+			int prevUpdateIndex = current.UpdateIndex;
+			current.UpdateIndex++;
+			UpdateResult result = await _logFiles.UpdateOneAsync<LogFileDocument>(x => x.Id == current.Id && x.UpdateIndex == prevUpdateIndex, update.Set(x => x.UpdateIndex, current.UpdateIndex));
+			return result.ModifiedCount == 1;
+		}
+
+		/// <inheritdoc/>
+		public async Task<ILogFile?> GetLogFileAsync(LogId logFileId)
+		{
+			LogFileDocument logFile = await _logFiles.Find<LogFileDocument>(x => x.Id == logFileId).FirstOrDefaultAsync();
+			return logFile;
+		}
+
+		/// <inheritdoc/>
+		public async Task<List<ILogFile>> GetLogFilesAsync(int? index = null, int? count = null)
+		{
+			IFindFluent<LogFileDocument, LogFileDocument> query = _logFiles.Find(FilterDefinition<LogFileDocument>.Empty);
+			if(index != null)
+			{
+				query = query.Skip(index.Value);
+			}
+			if(count != null)
+			{
+				query = query.Limit(count.Value);
+			}
+
+			List<LogFileDocument> results = await query.ToListAsync();
+			return results.ConvertAll<ILogFile>(x => x);
+		}
+	}
+}
diff --git a/Engine/Source/Programs/Horde/Horde.Build/Logs/LogFileMessages.cs b/Engine/Source/Programs/Horde/Horde.Build/Logs/LogFileMessages.cs
new file mode 100644
--- /dev/null
+++ b/Engine/Source/Programs/Horde/Horde.Build/Logs/LogFileMessages.cs
@@ -0,0 +1,87 @@
+// Copyright Epic Games, Inc. All Rights Reserved.
+
+using System.Collections.Generic;
+using Horde.Build.Logs;
+using Horde.Build.Models;
+using Horde.Build.Services;
+
+namespace Horde.Build.Api
+{
+	/// <summary>
+	/// The type of data stored in this log file
+	/// </summary>
+	public enum LogType
+	{
+		/// <summary>
+		/// Plain text data
+		/// </summary>
+		Text,
+
+		/// <summary>
+		/// Structured json objects, output as one object per line (without trailing commas)
+		/// </summary>
+		Json
+	}
+
+	/// <summary>
+	/// Response describing a log file
+	/// </summary>
+	public class GetLogFileResponse
+	{
+        /// <summary>
+		/// Unique id of the log file
+		/// </summary>
+		public string Id { get; set; }
+
+        /// <summary>
+        /// Unique id of the job for this log file
+		/// </summary>
+        public string JobId { get; set; }
+
+		/// <summary>
+		/// Type of events stored in this log
+		/// </summary>
+		public LogType Type { get; set; }
+
+		/// <summary>
+        /// Length of the log, in bytes
+        /// </summary>
+        public long Length { get; set; }
+
+		/// <summary>
+        /// Number of lines in the file
+        /// </summary>
+        public int LineCount { get; set; }
+
+		/// <summary>
+		/// Constructor
+		/// </summary>
+		/// <param name="logFile">The logfile to construct from</param>
+		/// <param name="metadata">Metadata about the log file</param>
+		public GetLogFileResponse(ILogFile logFile, LogMetadata metadata)
+		{
+            Id = logFile.Id.ToString();
+            JobId = logFile.JobId.ToString();
+			Type = logFile.Type;
+            Length = metadata.Length;
+            LineCount = metadata.MaxLineIndex;
+		}
+	}
+
+	/// <summary>
+	/// Response describing a log file
+	/// </summary>
+	public class SearchLogFileResponse
+	{
+		/// <summary>
+		/// List of line numbers containing the search text
+		/// </summary>
+		public List<int> Lines { get; set; } = new List<int>();
+
+		/// <summary>
+		/// Stats for the search
+		/// </summary>
+		public LogSearchStats? Stats { get; set; }
+	}
+}
+
