Horde: Finish re-org of Horde solution started a while ago.

diff --git a/Engine/Source/Programs/Horde/Horde.Build/Logs/LogFileService.cs b/Engine/Source/Programs/Horde/Horde.Build/Logs/LogFileService.cs
new file mode 100644
--- /dev/null
+++ b/Engine/Source/Programs/Horde/Horde.Build/Logs/LogFileService.cs
@@ -0,0 +1,1583 @@
+// Copyright Epic Games, Inc. All Rights Reserved.
+
+using System;
+using System.Collections.Generic;
+using System.Collections.Immutable;
+using System.Diagnostics;
+using System.Globalization;
+using System.IO;
+using System.Linq;
+using System.Security.Claims;
+using System.Text;
+using System.Text.Json;
+using System.Threading;
+using System.Threading.Tasks;
+using EpicGames.Core;
+using Horde.Build.Api;
+using Horde.Build.Collections;
+using Horde.Build.Logs;
+using Horde.Build.Models;
+using Horde.Build.Utilities;
+using HordeCommon;
+using Microsoft.Extensions.Caching.Memory;
+using Microsoft.Extensions.Hosting;
+using Microsoft.Extensions.Logging;
+using MongoDB.Bson;
+using MongoDB.Driver;
+using OpenTracing;
+using OpenTracing.Util;
+using Stream = System.IO.Stream;
+
+namespace Horde.Build.Services
+{
+	using JobId = ObjectId<IJob>;
+	using LogId = ObjectId<ILogFile>;
+	using SessionId = ObjectId<ISession>;
+
+	/// <summary>
+	/// Metadata about a log file
+	/// </summary>
+	public class LogMetadata
+	{
+		/// <summary>
+		/// Length of the log file
+		/// </summary>
+		public long Length { get; set; }
+
+		/// <summary>
+		/// Number of lines in the log file
+		/// </summary>
+		public int MaxLineIndex { get; set; }
+	}
+
+	/// <summary>
+	/// Interface for the log file service
+	/// </summary>
+	public interface ILogFileService
+	{
+		/// <summary>
+		/// Creates a new log
+		/// </summary>
+		/// <param name="jobId">Unique id of the job that owns this log file</param>
+		/// <param name="sessionId">Agent session allowed to update the log</param>
+		/// <param name="type">Type of events to be stored in the log</param>
+		/// <returns>The new log file document</returns>
+		Task<ILogFile> CreateLogFileAsync(JobId jobId, SessionId? sessionId, LogType type);
+
+		/// <summary>
+		/// Gets a logfile by ID
+		/// </summary>
+		/// <param name="logFileId">Unique id of the log file</param>
+		/// <returns>The logfile document</returns>
+		Task<ILogFile?> GetLogFileAsync(LogId logFileId);
+
+		/// <summary>
+		/// Gets a logfile by ID, returning a cached copy if available. This should only be used to retrieve constant properties set at creation, such as the session or job it's associated with.
+		/// </summary>
+		/// <param name="logFileId">Unique id of the log file</param>
+		/// <returns>The logfile document</returns>
+		Task<ILogFile?> GetCachedLogFileAsync(LogId logFileId);
+
+		/// <summary>
+		/// Returns a list of log files
+		/// </summary>
+		/// <param name="index">Index of the first result to return</param>
+		/// <param name="count">Number of results to return</param>
+		/// <returns>List of logfile documents</returns>
+		Task<List<ILogFile>> GetLogFilesAsync(int? index = null, int? count = null);
+
+		/// <summary>
+		/// Writes out chunk data and assigns to a file
+		/// </summary>
+		/// <param name="logFile">The log file</param>
+		/// <param name="offset">Offset within the file of data</param>
+		/// <param name="lineIndex">Current line index of the data (need not be the starting of the line)</param>
+		/// <param name="data">the data to add</param>
+		/// <param name="flush">Whether the current chunk is complete and should be flushed</param>
+		/// <param name="maxChunkLength">The maximum chunk length. Defaults to 128kb.</param>
+		/// <param name="maxSubChunkLineCount">Maximum number of lines in each sub-chunk.</param>
+		/// <returns></returns>
+		Task<ILogFile?> WriteLogDataAsync(ILogFile logFile, long offset, int lineIndex, ReadOnlyMemory<byte> data, bool flush, int maxChunkLength = 256 * 1024, int maxSubChunkLineCount = 128);
+
+		/// <summary>
+		/// Gets metadata about the log file
+		/// </summary>
+		/// <param name="logFile">The log file to query</param>
+		/// <returns>Metadata about the log file</returns>
+		Task<LogMetadata> GetMetadataAsync(ILogFile logFile);
+
+		/// <summary>
+		/// Creates new log events
+		/// </summary>
+		/// <param name="newEvents">List of events</param>
+		/// <returns>Async task</returns>
+		Task CreateEventsAsync(List<NewLogEventData> newEvents);
+
+		/// <summary>
+		/// Find events for a particular log file
+		/// </summary>
+		/// <param name="logFile">The log file instance</param>
+		/// <param name="spanId">Issue span to return events for</param>
+		/// <param name="index">Index of the first event to retrieve</param>
+		/// <param name="count">Number of events to retrieve</param>
+		/// <returns>List of log events</returns>
+		Task<List<ILogEvent>> FindEventsAsync(ILogFile logFile, ObjectId? spanId = null, int? index = null, int? count = null);
+
+		/// <summary>
+		/// Adds events to a log span
+		/// </summary>
+		/// <param name="events">The events to add</param>
+		/// <param name="spanId">The span id</param>
+		/// <returns>Async task</returns>
+		Task AddSpanToEventsAsync(IEnumerable<ILogEvent> events, ObjectId spanId);
+
+		/// <summary>
+		/// Find events for an issue
+		/// </summary>
+		/// <param name="spanIds">The span ids</param>
+		/// <param name="logIds">Log ids to include</param>
+		/// <param name="index">Index within the events for results to return</param>
+		/// <param name="count">Number of results to return</param>
+		/// <returns>Async task</returns>
+		Task<List<ILogEvent>> FindEventsForSpansAsync(IEnumerable<ObjectId> spanIds, LogId[]? logIds, int index, int count);
+
+		/// <summary>
+		/// Gets the data for an event
+		/// </summary>
+		/// <param name="logFile">The log file instance</param>
+		/// <param name="lineIndex">Index of the line in the file</param>
+		/// <param name="lineCount">Number of lines in the event</param>
+		/// <returns>New event data instance</returns>
+		Task<ILogEventData> GetEventDataAsync(ILogFile logFile, int lineIndex, int lineCount);
+
+		/// <summary>
+		/// Gets lines from the given log 
+		/// </summary>
+		/// <param name="logFile">The log file</param>
+		/// <param name="offset">Offset of the data to return</param>
+		/// <param name="length">Length of the data to return</param>
+		/// <returns>Data for the requested range</returns>
+		Task<Stream> OpenRawStreamAsync(ILogFile logFile, long offset, long length);
+
+		/// <summary>
+		/// Parses a stream of json text and outputs plain text
+		/// </summary>
+		/// <param name="logFile">The log file to query</param>
+		/// <param name="offset">Offset within the data to copy from</param>
+		/// <param name="length">Length of the data to copy</param>
+		/// <param name="outputStream">Output stream to receive the text data</param>
+		/// <returns>Async text</returns>
+		Task CopyPlainTextStreamAsync(ILogFile logFile, long offset, long length, Stream outputStream);
+
+		/// <summary>
+		/// Gets the offset of the given line number
+		/// </summary>
+		/// <param name="logFile">The log file to search</param>
+		/// <param name="lineIdx">The line index to retrieve the offset for</param>
+		/// <returns>The actual clamped line number and offset</returns>
+		Task<(int, long)> GetLineOffsetAsync(ILogFile logFile, int lineIdx);
+
+		/// <summary>
+		/// Search for the specified text in a log file
+		/// </summary>
+		/// <param name="logFile">The log file to search</param>
+		/// <param name="text">Text to search for</param>
+		/// <param name="firstLine">Line to start search from</param>
+		/// <param name="count">Number of results to return</param>
+		/// <param name="stats">Receives stats for the search</param>
+		/// <returns>List of line numbers containing the given term</returns>
+		Task<List<int>> SearchLogDataAsync(ILogFile logFile, string text, int firstLine, int count, LogSearchStats stats);
+	}
+
+	/// <summary>
+	/// Extension methods for dealing with log files
+	/// </summary>
+	public static class LogFileServiceExtensions
+	{
+		/// <summary>
+		/// Parses a stream of json text and outputs plain text
+		/// </summary>
+		/// <param name="logFileService">The log file service</param>
+		/// <param name="logFile">The log file to query</param>
+		/// <param name="offset">Offset within the log file to copy</param>
+		/// <param name="length">Length of the data to copy</param>
+		/// <param name="outputStream">Output stream to receive the text data</param>
+		/// <returns>Async text</returns>
+		public static async Task CopyRawStreamAsync(this ILogFileService logFileService, ILogFile logFile, long offset, long length, Stream outputStream)
+		{
+			using (Stream stream = await logFileService.OpenRawStreamAsync(logFile, offset, length))
+			{
+				await stream.CopyToAsync(outputStream);
+			}
+		}
+	}
+	
+	/// <summary>
+	/// Wraps functionality for manipulating logs
+	/// </summary>
+	public sealed class LogFileService : IHostedService, ILogFileService, IDisposable
+	{
+		private readonly ILogger<LogFileService> _logger;
+		private readonly ILogFileCollection _logFiles;
+		private readonly ILogEventCollection _logEvents;
+		private readonly ILogStorage _storage;
+		private readonly ILogBuilder _builder;
+
+		// Lock object for the <see cref="_writeTasks"/> and <see cref="_writeChunks"/> members
+		private readonly object _writeLock = new object();
+		private readonly List<Task> _writeTasks = new List<Task>();
+		private readonly HashSet<(LogId, long)> _writeChunks = new HashSet<(LogId, long)>();
+		private readonly IMemoryCache _logFileCache;
+
+		/// <summary>
+		/// Streams log data to a caller
+		/// </summary>
+		class ResponseStream : Stream
+		{
+			/// <summary>
+			/// The log file service that created this stream
+			/// </summary>
+			readonly LogFileService _logFileService;
+
+			/// <summary>
+			/// The log file being read
+			/// </summary>
+			readonly ILogFile _logFile;
+
+			/// <summary>
+			/// Starting offset within the file of the data to return 
+			/// </summary>
+			readonly long _responseOffset;
+
+			/// <summary>
+			/// Length of data to return
+			/// </summary>
+			readonly long _responseLength;
+
+			/// <summary>
+			/// Current offset within the stream
+			/// </summary>
+			long _currentOffset;
+
+			/// <summary>
+			/// The current chunk index
+			/// </summary>
+			int _chunkIdx;
+
+			/// <summary>
+			/// Buffer containing a message for missing data
+			/// </summary>
+			ReadOnlyMemory<byte> _sourceBuffer;
+
+			/// <summary>
+			/// Offset within the source buffer
+			/// </summary>
+			int _sourcePos;
+
+			/// <summary>
+			/// Length of the source buffer being copied from
+			/// </summary>
+			int _sourceEnd;
+
+			/// <summary>
+			/// Constructor
+			/// </summary>
+			/// <param name="logFileService">The log file service, for q</param>
+			/// <param name="logFile"></param>
+			/// <param name="offset"></param>
+			/// <param name="length"></param>
+			public ResponseStream(LogFileService logFileService, ILogFile logFile, long offset, long length)
+			{
+				_logFileService = logFileService;
+				_logFile = logFile;
+
+				_responseOffset = offset;
+				_responseLength = length;
+
+				_currentOffset = offset;
+
+				_chunkIdx = logFile.Chunks.GetChunkForOffset(offset);
+				_sourceBuffer = null!;
+			}
+
+			/// <inheritdoc/>
+			public override bool CanRead => true;
+
+			/// <inheritdoc/>
+			public override bool CanSeek => false;
+
+			/// <inheritdoc/>
+			public override bool CanWrite => false;
+
+			/// <inheritdoc/>
+			public override long Length => _responseLength;
+
+			/// <inheritdoc/>
+			public override long Position
+			{
+				get => _currentOffset - _responseOffset;
+				set => throw new NotImplementedException();
+			}
+
+			/// <inheritdoc/>
+			public override void Flush()
+			{
+			}
+
+			/// <inheritdoc/>
+			public override int Read(byte[] buffer, int offset, int count)
+			{
+				return ReadAsync(buffer, offset, count, CancellationToken.None).Result;
+			}
+
+			/// <inheritdoc/>
+			public override async Task<int> ReadAsync(byte[] buffer, int offset, int length, CancellationToken cancellationToken)
+			{
+				return await ReadAsync(buffer.AsMemory(offset, length), cancellationToken);
+			}
+
+			/// <inheritdoc/>
+			public override async ValueTask<int> ReadAsync(Memory<byte> buffer, CancellationToken cancellationToken)
+			{
+				int readBytes = 0;
+				while (readBytes < buffer.Length)
+				{
+					if (_sourcePos < _sourceEnd)
+					{
+						// Try to copy from the current buffer
+						int blockSize = Math.Min(_sourceEnd - _sourcePos, buffer.Length - readBytes);
+						_sourceBuffer.Slice(_sourcePos, blockSize).Span.CopyTo(buffer.Slice(readBytes).Span);
+						_currentOffset += blockSize;
+						readBytes += blockSize;
+						_sourcePos += blockSize;
+					}
+					else if (_currentOffset < _responseOffset + _responseLength)
+					{
+						// Move to the right chunk
+						while (_chunkIdx + 1 < _logFile.Chunks.Count && _currentOffset >= _logFile.Chunks[_chunkIdx + 1].Offset)
+						{
+							_chunkIdx++;
+						}
+
+						// Get the chunk data
+						ILogChunk chunk = _logFile.Chunks[_chunkIdx];
+						LogChunkData chunkData = await _logFileService.ReadChunkAsync(_logFile, _chunkIdx);
+
+						// Figure out which sub-chunk to use
+						int subChunkIdx = chunkData.GetSubChunkForOffsetWithinChunk((int)(_currentOffset - chunk.Offset));
+						LogSubChunkData subChunkData = chunkData.SubChunks[subChunkIdx];
+
+						// Get the source data
+						long subChunkOffset = chunk.Offset + chunkData.SubChunkOffset[subChunkIdx];
+						_sourceBuffer = subChunkData.InflateText().Data;
+						_sourcePos = (int)(_currentOffset - subChunkOffset);
+						_sourceEnd = (int)Math.Min(_sourceBuffer.Length, (_responseOffset + _responseLength) - subChunkOffset);
+					}
+					else
+					{
+						// End of the log
+						break;
+					}
+				}
+				return readBytes;
+			}
+
+			/// <inheritdoc/>
+			public override long Seek(long offset, SeekOrigin origin) => throw new NotImplementedException();
+
+			/// <inheritdoc/>
+			public override void SetLength(long value) => throw new NotImplementedException();
+
+			/// <inheritdoc/>
+			public override void Write(byte[] buffer, int offset, int count) => throw new NotImplementedException();
+		}
+
+		readonly ITicker _ticker;
+
+		/// <summary>
+		/// Constructor
+		/// </summary>
+		public LogFileService(ILogFileCollection logFiles, ILogEventCollection logEvents, ILogBuilder builder, ILogStorage storage, IClock clock, ILogger<LogFileService> logger)
+		{
+			_logFiles = logFiles;
+			_logEvents = logEvents;
+			_logFileCache = new MemoryCache(new MemoryCacheOptions());
+			_builder = builder;
+			_storage = storage;
+			_ticker = clock.AddTicker<LogFileService>(TimeSpan.FromSeconds(30.0), TickAsync, logger);
+			_logger = logger;
+		}
+
+		/// <inheritdoc/>
+		public Task StartAsync(CancellationToken cancellationToken) => _ticker.StartAsync();
+
+		/// <inheritdoc/>
+		public async Task StopAsync(CancellationToken cancellationToken)
+		{
+			_logger.LogInformation("Stopping log file service");
+			if (_builder.FlushOnShutdown)
+			{
+				await FlushAsync();
+			}
+			await _ticker.StopAsync();
+			_logger.LogInformation("Log service stopped");
+		}
+
+		/// <inheritdoc/>
+		public void Dispose()
+		{
+			_logFileCache.Dispose();
+			_storage.Dispose();
+			_ticker.Dispose();
+		}
+
+		/// <inheritdoc/>
+		public Task<ILogFile> CreateLogFileAsync(JobId jobId, SessionId? sessionId, LogType type)
+		{
+			return _logFiles.CreateLogFileAsync(jobId, sessionId, type);
+		}
+
+		/// <inheritdoc/>
+		public async Task<ILogFile?> GetLogFileAsync(LogId logFileId)
+		{
+			ILogFile? logFile = await _logFiles.GetLogFileAsync(logFileId);
+			if(logFile != null)
+			{
+				AddCachedLogFile(logFile);
+			}
+			return logFile;
+		}
+
+		/// <summary>
+		/// Adds a log file to the cache
+		/// </summary>
+		/// <param name="logFile">The log file to cache</param>
+		void AddCachedLogFile(ILogFile logFile)
+		{
+			MemoryCacheEntryOptions options = new MemoryCacheEntryOptions().SetSlidingExpiration(TimeSpan.FromSeconds(30));
+			_logFileCache.Set(logFile.Id, logFile, options);
+		}
+
+		/// <summary>
+		/// Gets a cached log file by id
+		/// </summary>
+		/// <param name="logFileId">The log file id</param>
+		/// <returns>New log file, or null if not found</returns>
+		public async Task<ILogFile?> GetCachedLogFileAsync(LogId logFileId)
+		{
+			object? logFile;
+			if (!_logFileCache.TryGetValue(logFileId, out logFile))
+			{
+				logFile = await GetLogFileAsync(logFileId);
+			}
+			return (ILogFile?)logFile;
+		}
+
+		/// <inheritdoc/>
+		public Task<List<ILogFile>> GetLogFilesAsync(int? index = null, int? count = null)
+		{
+			return _logFiles.GetLogFilesAsync(index, count);
+		}
+
+		class WriteState
+		{
+			public long _offset;
+			public int _lineIndex;
+			public ReadOnlyMemory<byte> _memory;
+
+			public WriteState(long offset, int lineIndex, ReadOnlyMemory<byte> memory)
+			{
+				_offset = offset;
+				_lineIndex = lineIndex;
+				_memory = memory;
+			}
+		}
+
+		/// <inheritdoc/>
+		public async Task<ILogFile?> WriteLogDataAsync(ILogFile logFile, long offset, int lineIndex, ReadOnlyMemory<byte> data, bool flush, int maxChunkLength, int maxSubChunkLineCount)
+		{
+			using IScope scope = GlobalTracer.Instance.BuildSpan("WriteLogDataAsync").StartActive();
+			scope.Span.SetTag("LogId", logFile.Id.ToString());
+			scope.Span.SetTag("Offset", offset.ToString(CultureInfo.InvariantCulture));
+			scope.Span.SetTag("Length", data.Length.ToString(CultureInfo.InvariantCulture));
+			scope.Span.SetTag("LineIndex", lineIndex.ToString(CultureInfo.InvariantCulture));
+
+			// Make sure the data ends in a newline
+			if (data.Length > 0 && data.Span[data.Length - 1] != '\n')
+			{
+				throw new ArgumentException("Log data must consist of a whole number of lines", nameof(data));
+			}
+
+			// Make sure the line count is a power of two
+			if ((maxSubChunkLineCount & (maxSubChunkLineCount - 1)) != 0)
+			{
+				throw new ArgumentException("Maximum line count per sub-chunk must be a power of two", nameof(maxSubChunkLineCount));
+			}
+
+			// List of the flushed chunks
+			List<long> completeOffsets = new List<long>();
+
+			// Add the data to new chunks
+			WriteState state = new WriteState(offset, lineIndex, data);
+			while (state._memory.Length > 0)
+			{
+				// Find an existing chunk to append to
+				int chunkIdx = logFile.Chunks.GetChunkForOffset(state._offset);
+				if (chunkIdx >= 0)
+				{
+					ILogChunk chunk = logFile.Chunks[chunkIdx];
+					if (await WriteLogChunkDataAsync(logFile, chunk, state, completeOffsets, maxChunkLength, maxSubChunkLineCount))
+					{
+						continue;
+					}
+				}
+
+				// Create a new chunk. Ensure that there's a chunk at the start of the file, even if the current write is beyond it.
+				ILogFile? newLogFile;
+				if (logFile.Chunks.Count == 0)
+				{
+					newLogFile = await _logFiles.TryAddChunkAsync(logFile, 0, 0);
+				}
+				else
+				{
+					newLogFile = await _logFiles.TryAddChunkAsync(logFile, state._offset, state._lineIndex);
+				}
+
+				// Try to add a new chunk at the new location
+				if (newLogFile == null)
+				{
+					newLogFile = await _logFiles.GetLogFileAsync(logFile.Id);
+					if (newLogFile == null)
+					{
+						_logger.LogError("Unable to update log file {LogId}", logFile.Id);
+						return null;
+					}
+					logFile = newLogFile;
+				}
+				else
+				{
+					// Logger.LogDebug("Added new chunk at offset {Offset} to log {LogId}", State.Offset, LogFile.Id);
+					logFile = newLogFile;
+				}
+			}
+
+			// Flush any pending chunks on this log file
+			if (flush)
+			{
+				foreach(ILogChunk chunk in logFile.Chunks)
+				{
+					if (chunk.Length == 0 && !completeOffsets.Contains(chunk.Offset))
+					{
+						await _builder.CompleteChunkAsync(logFile.Id, chunk.Offset);
+						completeOffsets.Add(chunk.Offset);
+					}
+				}
+			}
+
+			// Write all the chunks
+			if (completeOffsets.Count > 0 || flush)
+			{
+				ILogFile? newLogFile = await WriteCompleteChunksForLogAsync(logFile, completeOffsets, flush);
+				if (newLogFile == null)
+				{
+					return null;
+				}
+				logFile = newLogFile;
+			}
+			return logFile;
+		}
+
+		/// <summary>
+		/// Append data to an existing chunk.
+		/// </summary>
+		/// <param name="logFile">The log file to append to</param>
+		/// <param name="chunk">Chunk within the log file to update</param>
+		/// <param name="state">Data remaining to be written</param>
+		/// <param name="completeOffsets">List of complete chunks</param>
+		/// <param name="maxChunkLength">Maximum length of each chunk</param>
+		/// <param name="maxSubChunkLineCount">Maximum number of lines in each subchunk</param>
+		/// <returns>True if data was appended to </returns>
+		private async Task<bool> WriteLogChunkDataAsync(ILogFile logFile, ILogChunk chunk, WriteState state, List<long> completeOffsets, int maxChunkLength, int maxSubChunkLineCount)
+		{
+			// Don't allow data to be appended if the chunk is complete
+			if(chunk.Length > 0)
+			{
+				return false;
+			}
+
+			// Otherwise keep appending subchunks
+			bool bResult = false;
+			for (; ; )
+			{
+				// Flush the current sub-chunk if we're on a boundary
+				if (state._lineIndex > 0 && (state._lineIndex & (maxSubChunkLineCount - 1)) == 0)
+				{
+					_logger.LogDebug("Completing log {LogId} chunk offset {Offset} sub-chunk at line {LineIndex}", logFile.Id, chunk.Offset, state._lineIndex);
+					await _builder.CompleteSubChunkAsync(logFile.Id, chunk.Offset);
+				}
+
+				// Figure out the max length to write to the current chunk
+				int maxLength = Math.Min((int)((chunk.Offset + maxChunkLength) - state._offset), state._memory.Length);
+
+				// Figure out the maximum line index for the current sub chunk
+				int minLineIndex = state._lineIndex;
+				int maxLineIndex = (minLineIndex & ~(maxSubChunkLineCount - 1)) + maxSubChunkLineCount;
+
+				// Append this data
+				(int length, int lineCount) = GetWriteLength(state._memory.Span, maxLength, maxLineIndex - minLineIndex, state._offset == chunk.Offset);
+				if (length > 0)
+				{
+					// Append this data
+					ReadOnlyMemory<byte> appendData = state._memory.Slice(0, length);
+					if (!await _builder.AppendAsync(logFile.Id, chunk.Offset, state._offset, state._lineIndex, lineCount, appendData, logFile.Type))
+					{
+						break;
+					}
+
+					// Update the state
+					//Logger.LogDebug("Append to log {LogId} chunk offset {Offset} (LineIndex={LineIndex}, LineCount={LineCount}, Offset={WriteOffset}, Length={WriteLength})", LogFile.Id, Chunk.Offset, State.LineIndex, LineCount, State.Offset, Length);
+					state._offset += length;
+					state._lineIndex += lineCount;
+					state._memory = state._memory.Slice(length);
+					bResult = true;
+
+					// If this is the end of the data, bail out
+					if(state._memory.Length == 0)
+					{
+						break;
+					}
+				}
+
+				// Flush the sub-chunk if it's full
+				if (state._lineIndex < maxLineIndex)
+				{
+					_logger.LogDebug("Completing chunk for log {LogId} at offset {Offset}", logFile.Id, chunk.Offset);
+					await _builder.CompleteChunkAsync(logFile.Id, chunk.Offset);
+					completeOffsets.Add(chunk.Offset);
+					break;
+				}
+			}
+			return bResult;
+		}
+
+		/// <summary>
+		/// Get the amount of data to write from the given span
+		/// </summary>
+		/// <param name="span">Data to write</param>
+		/// <param name="maxLength">Maximum length of the data to write</param>
+		/// <param name="maxLineCount">Maximum number of lines to write</param>
+		/// <param name="bIsEmptyChunk">Whether the current chunk is empty</param>
+		/// <returns>A tuple consisting of the amount of data to write and number of lines in it</returns>
+		private static (int, int) GetWriteLength(ReadOnlySpan<byte> span, int maxLength, int maxLineCount, bool bIsEmptyChunk)
+		{
+			int length = 0;
+			int lineCount = 0;
+			for (int idx = 0; idx < maxLength || bIsEmptyChunk; idx++)
+			{
+				if (span[idx] == '\n')
+				{
+					length = idx + 1;
+					lineCount++;
+					bIsEmptyChunk = false;
+
+					if (lineCount >= maxLineCount)
+					{
+						break;
+					}
+				}
+			}
+			return (length, lineCount);
+		}
+
+		/// <inheritdoc/>
+		public async Task<LogMetadata> GetMetadataAsync(ILogFile logFile)
+		{
+			LogMetadata metadata = new LogMetadata();
+			if (logFile.Chunks.Count > 0)
+			{
+				ILogChunk chunk = logFile.Chunks[logFile.Chunks.Count - 1];
+				if (logFile.MaxLineIndex == null || chunk.Length == 0)
+				{
+					LogChunkData chunkData = await ReadChunkAsync(logFile, logFile.Chunks.Count - 1);
+					metadata.Length = chunk.Offset + chunkData.Length;
+					metadata.MaxLineIndex = chunk.LineIndex + chunkData.LineCount;
+				}
+				else
+				{
+					metadata.Length = chunk.Offset + chunk.Length;
+					metadata.MaxLineIndex = logFile.MaxLineIndex.Value;
+				}
+			}
+			return metadata;
+		}
+
+		/// <inheritdoc/>
+		public Task CreateEventsAsync(List<NewLogEventData> newEvents)
+		{
+			return _logEvents.AddManyAsync(newEvents);
+		}
+
+		/// <inheritdoc/>
+		public Task<List<ILogEvent>> FindEventsAsync(ILogFile logFile, ObjectId? spanId = null, int? index = null, int? count = null)
+		{
+			return _logEvents.FindAsync(logFile.Id, spanId, index, count);
+		}
+
+		class LogEventLine : ILogEventLine
+		{
+			readonly LogLevel _level;
+			public EventId? EventId { get; }
+			public string Message { get; }
+			public JsonElement Data { get; }
+
+			LogLevel ILogEventLine.Level => _level;
+
+			public LogEventLine(ReadOnlySpan<byte> data)
+				: this(JsonSerializer.Deserialize<JsonElement>(data))
+			{
+			}
+
+			public LogEventLine(JsonElement data)
+			{
+				Data = data;
+
+				JsonElement levelElement;
+				if (!data.TryGetProperty("level", out levelElement) || !Enum.TryParse(levelElement.GetString(), out _level))
+				{
+					_level = LogLevel.Information;
+				}
+
+				JsonElement idElement;
+				if (data.TryGetProperty("id", out idElement))
+				{
+					int idValue;
+					if (idElement.TryGetInt32(out idValue))
+					{
+						EventId = idValue;
+					}
+				}
+
+				JsonElement messageElement;
+				if (data.TryGetProperty("renderedMessage", out messageElement) || data.TryGetProperty("message", out messageElement))
+				{
+					Message = messageElement.GetString() ?? "(Invalid)";
+				}
+				else
+				{
+					Message = "(Missing message or renderedMessage field)";
+				}
+			}
+		}
+
+		class LogEventData : ILogEventData
+		{
+			public IReadOnlyList<ILogEventLine> Lines { get; }
+
+			EventId? ILogEventData.EventId => (Lines.Count > 0) ? Lines[0].EventId : null;
+			EventSeverity ILogEventData.Severity => (Lines.Count == 0) ? EventSeverity.Information : (Lines[0].Level == LogLevel.Warning) ? EventSeverity.Warning : EventSeverity.Error;
+			string ILogEventData.Message => String.Join("\n", Lines.Select(x => x.Message));
+
+			public LogEventData(IReadOnlyList<ILogEventLine> lines)
+			{
+				Lines = lines;
+			}
+		}
+
+		/// <inheritdoc/>
+		public Task AddSpanToEventsAsync(IEnumerable<ILogEvent> events, ObjectId spanId)
+		{
+			return _logEvents.AddSpanToEventsAsync(events, spanId);
+		}
+
+		/// <inheritdoc/>
+		public Task<List<ILogEvent>> FindEventsForSpansAsync(IEnumerable<ObjectId> spanIds, LogId[]? logIds, int index, int count)
+		{
+			return _logEvents.FindEventsForSpansAsync(spanIds, logIds, index, count);
+		}
+
+		/// <inheritdoc/>
+		public async Task<ILogEventData> GetEventDataAsync(ILogFile logFile, int lineIndex, int lineCount)
+		{
+			using IScope scope = GlobalTracer.Instance.BuildSpan("GetEventDataAsync").StartActive();
+			scope.Span.SetTag("LogId", logFile.Id.ToString());
+			scope.Span.SetTag("LineIndex", lineIndex.ToString(CultureInfo.InvariantCulture));
+			scope.Span.SetTag("LineCount", lineCount.ToString(CultureInfo.InvariantCulture));
+
+			(_, long minOffset) = await GetLineOffsetAsync(logFile, lineIndex);
+			(_, long maxOffset) = await GetLineOffsetAsync(logFile, lineIndex + lineCount);
+
+			byte[] data = new byte[maxOffset - minOffset];
+			using (Stream stream = await OpenRawStreamAsync(logFile, minOffset, maxOffset - minOffset))
+			{
+				int length = await stream.ReadAsync(data.AsMemory());
+				if(length != data.Length)
+				{
+					_logger.LogWarning("Read less than expected from log stream (Expected {Expected}, Got {Got})", data.Length, length);
+				}
+				return ParseEventData(data.AsSpan(0, length));
+			}
+		}
+
+		/// <summary>
+		/// Parses event data from a buffer
+		/// </summary>
+		/// <param name="data">Data to parse</param>
+		/// <returns>Parsed event data</returns>
+		private ILogEventData ParseEventData(ReadOnlySpan<byte> data)
+		{
+			List<LogEventLine> lines = new List<LogEventLine>();
+
+			ReadOnlySpan<byte> remainingData = data;
+			while(remainingData.Length > 0)
+			{
+				int endOfLine = remainingData.IndexOf((byte)'\n');
+				if(endOfLine == -1)
+				{
+					break;
+				}
+
+				ReadOnlySpan<byte> lineData = remainingData.Slice(0, endOfLine);
+				try
+				{
+					lines.Add(new LogEventLine(lineData));
+				}
+				catch(JsonException ex)
+				{
+					_logger.LogWarning(ex, "Unable to parse line from log file: {Line}", Encoding.UTF8.GetString(lineData));
+				}
+				remainingData = remainingData.Slice(endOfLine + 1);
+			}
+
+			return new LogEventData(lines);
+		}
+
+		/// <inheritdoc/>
+		public async Task<Stream> OpenRawStreamAsync(ILogFile logFile, long offset, long length)
+		{
+			if (logFile.Chunks.Count == 0)
+			{
+				return new MemoryStream(Array.Empty<byte>(), false);
+			}
+			else
+			{
+				int lastChunkIdx = logFile.Chunks.Count - 1;
+
+				// Clamp the length of the request
+				ILogChunk lastChunk = logFile.Chunks[lastChunkIdx];
+				if (length > lastChunk.Offset)
+				{
+					long lastChunkLength = lastChunk.Length;
+					if (lastChunkLength <= 0)
+					{
+						LogChunkData lastChunkData = await ReadChunkAsync(logFile, lastChunkIdx);
+						lastChunkLength = lastChunkData.Length;
+					}
+					length = Math.Min(length, (lastChunk.Offset + lastChunkLength) - offset);
+				}
+
+				// Create the new stream
+				return new ResponseStream(this, logFile, offset, length);
+			}
+		}
+
+		/// <summary>
+		/// Helper method for catching exceptions in <see cref="LogText.ConvertToPlainText(ReadOnlySpan{Byte}, Byte[], Int32)"/>
+		/// </summary>
+		public static int GuardedConvertToPlainText(ReadOnlySpan<byte> input, byte[] output, int outputOffset, ILogger Logger)
+		{
+			try
+			{
+				return LogText.ConvertToPlainText(input, output, outputOffset);
+			}
+			catch (Exception Ex)
+			{
+				Logger.LogWarning(Ex, "Unable to convert log line to plain text: {Line}", Encoding.UTF8.GetString(input));
+				output[outputOffset] = (byte)'\n';
+				return outputOffset + 1;
+			}
+		}
+
+		/// <inheritdoc/>
+		public async Task CopyPlainTextStreamAsync(ILogFile logFile, long offset, long length, Stream outputStream)
+		{
+			using (Stream stream = await OpenRawStreamAsync(logFile, 0, Int64.MaxValue))
+			{
+				byte[] readBuffer = new byte[4096];
+				int readBufferLength = 0;
+
+				byte[] writeBuffer = new byte[4096];
+				int writeBufferLength = 0;
+
+				while (length > 0)
+				{
+					// Add more data to the buffer
+					int readBytes = await stream.ReadAsync(readBuffer.AsMemory(readBufferLength, readBuffer.Length - readBufferLength));
+					readBufferLength += readBytes;
+
+					// Copy as many lines as possible to the output
+					int convertedBytes = 0;
+					for (int endIdx = 1; endIdx < readBufferLength; endIdx++)
+					{
+						if (readBuffer[endIdx] == '\n')
+						{
+							writeBufferLength = GuardedConvertToPlainText(readBuffer.AsSpan(convertedBytes, endIdx - convertedBytes), writeBuffer, writeBufferLength, _logger);
+							convertedBytes = endIdx + 1;
+						}
+					}
+
+					// If there's anything in the write buffer, write it out
+					if (writeBufferLength > 0)
+					{
+						if (offset < writeBufferLength)
+						{
+							int writeLength = (int)Math.Min((long)writeBufferLength - offset, length);
+							await outputStream.WriteAsync(writeBuffer.AsMemory((int)offset, writeLength));
+							length -= writeLength;
+						}
+						offset = Math.Max(offset - writeBufferLength, 0);
+						writeBufferLength = 0;
+					}
+
+					// If we were able to read something, shuffle down the rest of the buffer. Otherwise expand the read buffer.
+					if (convertedBytes > 0)
+					{
+						Buffer.BlockCopy(readBuffer, convertedBytes, readBuffer, 0, readBufferLength - convertedBytes);
+						readBufferLength -= convertedBytes;
+					}
+					else if (readBufferLength > 0)
+					{
+						Array.Resize(ref readBuffer, readBuffer.Length + 128);
+						writeBuffer = new byte[readBuffer.Length];
+					}
+
+					// Exit if we didn't read anything in this iteration
+					if (readBytes == 0)
+					{
+						break;
+					}
+				}
+			}
+		}
+
+		/// <inheritdoc/>
+		public async Task<(int, long)> GetLineOffsetAsync(ILogFile logFile, int lineIdx)
+		{
+			int chunkIdx = logFile.Chunks.GetChunkForLine(lineIdx);
+
+			ILogChunk chunk = logFile.Chunks[chunkIdx];
+			LogChunkData chunkData = await ReadChunkAsync(logFile, chunkIdx);
+
+			if (lineIdx < chunk.LineIndex)
+			{
+				lineIdx = chunk.LineIndex;
+			}
+
+			int maxLineIndex = chunk.LineIndex + chunkData.LineCount;
+			if (lineIdx >= maxLineIndex)
+			{
+				lineIdx = maxLineIndex;
+			}
+
+			long offset = chunk.Offset + chunkData.GetLineOffsetWithinChunk(lineIdx - chunk.LineIndex);
+			return (lineIdx, offset);
+		}
+
+		/// <summary>
+		/// Executes a background task
+		/// </summary>
+		/// <param name="stoppingToken">Cancellation token</param>
+		async ValueTask TickAsync(CancellationToken stoppingToken)
+		{
+			lock (_writeLock)
+			{
+				try
+				{
+					_writeTasks.RemoveCompleteTasks();
+				}
+				catch (Exception ex)
+				{
+					_logger.LogError(ex, "Exception while waiting for write tasks to complete");
+				}
+			}
+			await IncrementalFlush();
+		}
+				
+		/// <summary>
+		/// Flushes complete chunks to the storage provider
+		/// </summary>
+		/// <returns>Async task</returns>
+		private async Task IncrementalFlush()
+		{
+			// Get all the chunks older than 20 minutes
+			List<(LogId, long)> flushChunks = await _builder.TouchChunksAsync(TimeSpan.FromMinutes(10.0));
+			_logger.LogDebug("Performing incremental flush of log builder ({NumChunks} chunks)", flushChunks.Count);
+
+			// Mark them all as complete
+			foreach ((LogId logId, long offset) in flushChunks)
+			{
+				await _builder.CompleteChunkAsync(logId, offset);
+			}
+
+			// Add tasks for flushing all the chunks
+			WriteCompleteChunks(flushChunks, true);
+		}
+
+		/// <summary>
+		/// Flushes the write cache
+		/// </summary>
+		/// <returns>Async task</returns>
+		public async Task FlushAsync()
+		{
+			_logger.LogInformation("Forcing flush of pending log chunks...");
+
+			// Mark everything in the cache as complete
+			List<(LogId, long)> writeChunks = await _builder.TouchChunksAsync(TimeSpan.Zero);
+			WriteCompleteChunks(writeChunks, true);
+
+			// Wait for everything to flush
+			await FlushPendingWritesAsync();
+		}
+
+		/// <summary>
+		/// Flush any writes in progress
+		/// </summary>
+		/// <returns>Async task</returns>
+		public async Task FlushPendingWritesAsync()
+		{
+			for(; ;)
+			{
+				// Capture the current contents of the WriteTasks list
+				List<Task> tasks;
+				lock (_writeLock)
+				{
+					_writeTasks.RemoveCompleteTasks();
+					tasks = new List<Task>(_writeTasks);
+				}
+				if (tasks.Count == 0)
+				{
+					break;
+				}
+
+				// Also add a delay so we'll periodically refresh the list
+				tasks.Add(Task.Delay(TimeSpan.FromSeconds(5.0)));
+				await Task.WhenAny(tasks);
+			}
+		}
+
+		/// <summary>
+		/// Adds tasks for writing a list of complete chunks
+		/// </summary>
+		/// <param name="chunksToWrite">List of chunks to write</param>
+		/// <param name="bCreateIndex">Create an index for the log</param>
+		private void WriteCompleteChunks(List<(LogId, long)> chunksToWrite, bool bCreateIndex)
+		{
+			foreach (IGrouping<LogId, long> group in chunksToWrite.GroupBy(x => x.Item1, x => x.Item2))
+			{
+				LogId logId = group.Key;
+
+				// Find offsets of new chunks to write
+				List<long> offsets = new List<long>();
+				lock (_writeLock)
+				{
+					foreach (long offset in group.OrderBy(x => x))
+					{
+						if (_writeChunks.Add((logId, offset)))
+						{
+							offsets.Add(offset);
+						}
+					}
+				}
+
+				// Create the write task
+				if (offsets.Count > 0)
+				{
+					Task task = Task.Run(() => WriteCompleteChunksForLogAsync(logId, offsets, bCreateIndex));
+					lock (_writeLock)
+					{
+						_writeTasks.Add(task);
+					}
+				}
+			}
+		}
+
+		/// <summary>
+		/// Writes a set of chunks to the database
+		/// </summary>
+		/// <param name="logId">Log file to update</param>
+		/// <param name="offsets">Chunks to write</param>
+		/// <param name="bCreateIndex">Whether to create the index for this log</param>
+		/// <returns>Async task</returns>
+		private async Task<ILogFile?> WriteCompleteChunksForLogAsync(LogId logId, List<long> offsets, bool bCreateIndex)
+		{
+			ILogFile? logFile = await _logFiles.GetLogFileAsync(logId);
+			if(logFile != null)
+			{
+				logFile = await WriteCompleteChunksForLogAsync(logFile, offsets, bCreateIndex);
+			}
+			return logFile;
+		}
+
+		/// <summary>
+		/// Writes a set of chunks to the database
+		/// </summary>
+		/// <param name="logFileInterface">Log file to update</param>
+		/// <param name="offsets">Chunks to write</param>
+		/// <param name="bCreateIndex">Whether to create the index for this log</param>
+		/// <returns>Async task</returns>
+		private async Task<ILogFile?> WriteCompleteChunksForLogAsync(ILogFile logFileInterface, List<long> offsets, bool bCreateIndex)
+		{
+			// Write the data to the storage provider
+			List<Task<LogChunkData?>> chunkWriteTasks = new List<Task<LogChunkData?>>();
+			foreach (long offset in offsets)
+			{
+				int chunkIdx = logFileInterface.Chunks.BinarySearch(x => x.Offset, offset);
+				if (chunkIdx >= 0)
+				{
+					_logger.LogDebug("Queuing write of log {LogId} chunk {ChunkIdx} offset {Offset}", logFileInterface.Id, chunkIdx, offset);
+					int lineIndex = logFileInterface.Chunks[chunkIdx].LineIndex;
+					chunkWriteTasks.Add(Task.Run(() => WriteChunkAsync(logFileInterface.Id, offset, lineIndex)));
+				}
+			}
+
+			// Wait for the tasks to complete, periodically updating the log file object
+			ILogFile? logFile = logFileInterface;
+			while (chunkWriteTasks.Count > 0)
+			{
+				// Wait for all tasks to be complete OR (any task has completed AND 30 seconds has elapsed)
+				Task allCompleteTask = Task.WhenAll(chunkWriteTasks);
+				Task anyCompleteTask = Task.WhenAny(chunkWriteTasks);
+				await Task.WhenAny(allCompleteTask, Task.WhenAll(anyCompleteTask, Task.Delay(TimeSpan.FromSeconds(30.0))));
+
+				// Update the log file with the written chunks
+				List<LogChunkData?> writtenChunks = chunkWriteTasks.RemoveCompleteTasks();
+				while (logFile != null)
+				{
+					// Update the length of any complete chunks
+					List<CompleteLogChunkUpdate> updates = new List<CompleteLogChunkUpdate>();
+					foreach (LogChunkData? chunkData in writtenChunks)
+					{
+						if (chunkData != null)
+						{
+							int chunkIdx = logFile.Chunks.GetChunkForOffset(chunkData.Offset);
+							if (chunkIdx >= 0)
+							{
+								ILogChunk chunk = logFile.Chunks[chunkIdx];
+								if (chunk.Offset == chunkData.Offset)
+								{
+									CompleteLogChunkUpdate update = new CompleteLogChunkUpdate(chunkIdx, chunkData.Length, chunkData.LineCount);
+									updates.Add(update);
+								}
+							}
+						}
+					}
+
+					// Try to apply the updates
+					ILogFile? newLogFile = await _logFiles.TryCompleteChunksAsync(logFile, updates);
+					if (newLogFile != null)
+					{
+						logFile = newLogFile;
+						break;
+					}
+
+					// Update the log file
+					logFile = await GetLogFileAsync(logFile.Id);
+				}
+			}
+
+			// Create the index if necessary
+			if (bCreateIndex && logFile != null)
+			{
+				try
+				{
+					logFile = await CreateIndexAsync(logFile);
+				}
+				catch(Exception ex)
+				{
+					_logger.LogError(ex, "Failed to create index for log {LogId}", logFileInterface.Id);
+				}
+			}
+
+			return logFile;
+		}
+
+		/// <summary>
+		/// Creates an index for the given log file
+		/// </summary>
+		/// <param name="logFile">The log file object</param>
+		/// <returns>Updated log file</returns>
+		private async Task<ILogFile?> CreateIndexAsync(ILogFile logFile)
+		{
+			if(logFile.Chunks.Count == 0)
+			{
+				return logFile;
+			}
+
+			// Get the new length of the log, and early out if it won't be any longer
+			ILogChunk lastChunk = logFile.Chunks[logFile.Chunks.Count - 1];
+			if(lastChunk.Offset + lastChunk.Length <= (logFile.IndexLength ?? 0))
+			{
+				return logFile;
+			}
+
+			// Save stats for the index creation
+			using IScope scope = GlobalTracer.Instance.BuildSpan("CreateIndexAsync").StartActive();
+			scope.Span.SetTag("LogId", logFile.Id.ToString());
+			scope.Span.SetTag("Length", (lastChunk.Offset + lastChunk.Length).ToString(CultureInfo.InvariantCulture));
+
+			long newLength = 0;
+			int newLineCount = 0;
+
+			// Read the existing index if there is one
+			List<LogIndexData> indexes = new List<LogIndexData>();
+			if (logFile.IndexLength != null)
+			{
+				LogIndexData? existingIndex = await ReadIndexAsync(logFile, logFile.IndexLength.Value);
+				if(existingIndex != null)
+				{
+					indexes.Add(existingIndex);
+					newLineCount = existingIndex.LineCount;
+				}
+			}
+
+			// Add all the new chunks
+			int chunkIdx = logFile.Chunks.GetChunkForLine(newLineCount);
+			if (chunkIdx < 0)
+			{
+				int firstLine = (logFile.Chunks.Count > 0) ? logFile.Chunks[0].LineIndex : -1;
+				throw new Exception($"Invalid chunk index {chunkIdx}. Index.LineCount={newLineCount}, Chunks={logFile.Chunks.Count}, First line={firstLine}");
+			}
+
+			for (; chunkIdx < logFile.Chunks.Count; chunkIdx++)
+			{
+				ILogChunk chunk = logFile.Chunks[chunkIdx];
+				LogChunkData chunkData = await ReadChunkAsync(logFile, chunkIdx);
+
+				int subChunkIdx = chunkData.GetSubChunkForLine(Math.Max(newLineCount - chunk.LineIndex, 0));
+				if(subChunkIdx < 0)
+				{
+					throw new Exception($"Invalid subchunk index {subChunkIdx}. Chunk {chunkIdx}/{logFile.Chunks.Count}. Index.LineCount={newLineCount}, Chunk.LineIndex={chunk.LineIndex}, First subchunk {chunkData.SubChunkLineIndex[0]}");
+				}
+
+				for (; subChunkIdx < chunkData.SubChunks.Count; subChunkIdx++)
+				{
+					LogSubChunkData subChunkData = chunkData.SubChunks[subChunkIdx];
+					if (subChunkData.LineIndex >= newLineCount)
+					{
+						try
+						{
+							indexes.Add(subChunkData.BuildIndex(_logger));
+						}
+						catch (Exception ex)
+						{
+							throw new Exception($"Failed to create index block - log {logFile.Id}, chunk {chunkIdx} ({logFile.Chunks.Count}), subchunk {subChunkIdx} ({chunkData.SubChunks.Count}), index lines: {newLineCount}, chunk index: {chunk.LineIndex}, subchunk index: {chunk.LineIndex + chunkData.SubChunkLineIndex[subChunkIdx]}, subchunk count: {subChunkData.LineCount}", ex);
+						}
+
+						newLength = subChunkData.Offset + subChunkData.Length;
+						newLineCount = subChunkData.LineIndex + subChunkData.LineCount;
+					}
+				}
+			}
+
+			// Try to update the log file
+			ILogFile? newLogFile = logFile;
+			if (newLength > (logFile.IndexLength ?? 0))
+			{
+				LogIndexData index = LogIndexData.Merge(indexes);
+				_logger.LogDebug("Writing index for log {LogId} covering {Length} (index length {IndexLength})", logFile.Id, newLength, index.GetSerializedSize());
+
+				await WriteIndexAsync(logFile.Id, newLength, index);
+
+				while(newLogFile != null && newLength > (newLogFile.IndexLength ?? 0))
+				{
+					newLogFile = await _logFiles.TryUpdateIndexAsync(newLogFile, newLength);
+					if(newLogFile != null)
+					{
+						break;
+					}
+					newLogFile = await _logFiles.GetLogFileAsync(logFile.Id);
+				}
+			}
+			return newLogFile;
+		}
+
+		/// <summary>
+		/// Reads a chunk from storage
+		/// </summary>
+		/// <param name="logFile">Log file to read from</param>
+		/// <param name="chunkIdx">The chunk to read</param>
+		/// <returns>Chunk data</returns>
+		private async Task<LogChunkData> ReadChunkAsync(ILogFile logFile, int chunkIdx)
+		{
+			ILogChunk chunk = logFile.Chunks[chunkIdx];
+
+			// Try to read the chunk data from storage
+			LogChunkData? chunkData = null;
+			try
+			{
+				// If the chunk is not yet complete, query the log builder
+				if (chunk.Length == 0)
+				{
+					chunkData = await _builder.GetChunkAsync(logFile.Id, chunk.Offset, chunk.LineIndex);
+				}
+
+				// Otherwise go directly to the log storage
+				if (chunkData == null)
+				{
+					chunkData = await _storage.ReadChunkAsync(logFile.Id, chunk.Offset, chunk.LineIndex);
+				}
+			}
+			catch (Exception ex)
+			{
+				_logger.LogError(ex, "Unable to read log {LogId} at offset {Offset}", logFile.Id, chunk.Offset);
+			}
+
+			// Get the minimum length and line count for the chunk
+			if (chunkIdx + 1 < logFile.Chunks.Count)
+			{
+				ILogChunk nextChunk = logFile.Chunks[chunkIdx + 1];
+				chunkData = await RepairChunkDataAsync(logFile, chunkIdx, chunkData, (int)(nextChunk.Offset - chunk.Offset), nextChunk.LineIndex - chunk.LineIndex);
+			}
+			else
+			{
+				if (logFile.MaxLineIndex != null && chunk.Length != 0)
+				{
+					chunkData = await RepairChunkDataAsync(logFile, chunkIdx, chunkData, chunk.Length, logFile.MaxLineIndex.Value - chunk.LineIndex);
+				}
+				else if(chunkData == null)
+				{
+					chunkData = await RepairChunkDataAsync(logFile, chunkIdx, chunkData, 1024, 1);
+				}
+			}
+
+			return chunkData;
+		}
+
+		/// <summary>
+		/// Validates the given chunk data, and fix it up if necessary
+		/// </summary>
+		/// <param name="logFile">The log file instance</param>
+		/// <param name="chunkIdx">Index of the chunk within the logfile</param>
+		/// <param name="chunkData">The chunk data that was read</param>
+		/// <param name="length">Expected length of the data</param>
+		/// <param name="lineCount">Expected number of lines in the data</param>
+		/// <returns>Repaired chunk data</returns>
+		async Task<LogChunkData> RepairChunkDataAsync(ILogFile logFile, int chunkIdx, LogChunkData? chunkData, int length, int lineCount)
+		{
+			int currentLength = 0;
+			int currentLineCount = 0;
+			if(chunkData != null)
+			{
+				currentLength = chunkData.Length;
+				currentLineCount = chunkData.LineCount;
+			}
+
+			if (chunkData == null || currentLength < length || currentLineCount < lineCount)
+			{
+				_logger.LogWarning("Creating placeholder subchunk for log {LogId} chunk {ChunkIdx} (length {Length} vs expected {ExpLength}, lines {LineCount} vs expected {ExpLineCount})", logFile.Id, chunkIdx, currentLength, length, currentLineCount, lineCount);
+
+				List<LogSubChunkData> subChunks = new List<LogSubChunkData>();
+				if (chunkData != null && chunkData.Length < length && chunkData.LineCount < lineCount)
+				{
+					subChunks.AddRange(chunkData.SubChunks);
+				}
+
+				LogText text = new LogText();
+				text.AppendMissingDataInfo(chunkIdx, logFile.Chunks[chunkIdx].Server, length - currentLength, lineCount - currentLineCount);
+				subChunks.Add(new LogSubChunkData(logFile.Type, currentLength, currentLineCount, text));
+
+				ILogChunk chunk = logFile.Chunks[chunkIdx];
+				chunkData = new LogChunkData(chunk.Offset, chunk.LineIndex, subChunks);
+
+				try
+				{
+					await _storage.WriteChunkAsync(logFile.Id, chunk.Offset, chunkData);
+				}
+				catch (Exception ex)
+				{
+					_logger.LogWarning(ex, "Unable to put repaired log data for log {LogId} chunk {ChunkIdx}", logFile.Id, chunkIdx);
+				}
+			}
+			return chunkData;
+		}
+
+		/// <summary>
+		/// Writes a set of chunks to the database
+		/// </summary>
+		/// <param name="logFileId">Unique id of the log file</param>
+		/// <param name="offset">Offset of the chunk to write</param>
+		/// <param name="lineIndex">First line index of the chunk</param>
+		/// <returns>Chunk daata</returns>
+		private async Task<LogChunkData?> WriteChunkAsync(LogId logFileId, long offset, int lineIndex)
+		{
+			// Write the chunk to storage
+			LogChunkData? chunkData = await _builder.GetChunkAsync(logFileId, offset, lineIndex);
+			if (chunkData == null)
+			{
+				_logger.LogDebug("Log {LogId} offset {Offset} not found in log builder", logFileId, offset);
+			}
+			else
+			{
+				try
+				{
+					await _storage.WriteChunkAsync(logFileId, chunkData.Offset, chunkData);
+				}
+				catch (Exception ex)
+				{
+					_logger.LogError(ex, "Unable to write log {LogId} at offset {Offset}", logFileId, chunkData.Offset);
+				}
+			}
+
+			// Remove it from the log builder
+			try
+			{
+				await _builder.RemoveChunkAsync(logFileId, offset);
+			}
+			catch (Exception ex)
+			{
+				_logger.LogError(ex, "Unable to remove log {LogId} at offset {Offset} from log builder", logFileId, offset);
+			}
+
+			return chunkData;
+		}
+
+		/// <summary>
+		/// Reads a chunk from storage
+		/// </summary>
+		/// <param name="logFile">Log file to read from</param>
+		/// <param name="length">Length of the log covered by the index</param>
+		/// <returns>Chunk data</returns>
+		private async Task<LogIndexData?> ReadIndexAsync(ILogFile logFile, long length)
+		{
+			try
+			{
+				LogIndexData? index = await _storage.ReadIndexAsync(logFile.Id, length);
+				return index;
+			}
+			catch (Exception ex)
+			{
+				_logger.LogError(ex, "Unable to read log {LogId} index at length {Length}", logFile.Id, length);
+				return null;
+			}
+		}
+
+		/// <summary>
+		/// Writes an index to the database
+		/// </summary>
+		/// <param name="logFileId">Unique id of the log file</param>
+		/// <param name="length">Length of the data covered by the index</param>
+		/// <param name="index">Index to write</param>
+		/// <returns>Async task</returns>
+		private async Task WriteIndexAsync(LogId logFileId, long length, LogIndexData index)
+		{
+			try
+			{
+				await _storage.WriteIndexAsync(logFileId, length, index);
+			}
+			catch(Exception ex)
+			{
+				_logger.LogError(ex, "Unable to write index for log {LogId}", logFileId);
+			}
+		}
+
+		/// <summary>
+		/// Determines if the user is authorized to perform an action on a particular template
+		/// </summary>
+		/// <param name="logFile">The template to check</param>
+		/// <param name="user">The principal to authorize</param>
+		/// <returns>True if the action is authorized</returns>
+		public static bool AuthorizeForSession(ILogFile logFile, ClaimsPrincipal user)
+		{
+			if(logFile.SessionId != null)
+			{
+				return user.HasClaim(HordeClaimTypes.AgentSessionId, logFile.SessionId.Value.ToString());
+			}
+			else
+			{
+				return false;
+			}
+		}
+
+		/// <inheritdoc/>
+		public async Task<List<int>> SearchLogDataAsync(ILogFile logFile, string text, int firstLine, int count, LogSearchStats searchStats)
+		{
+			Stopwatch timer = Stopwatch.StartNew();
+
+			using IScope scope = GlobalTracer.Instance.BuildSpan("SearchLogDataAsync").StartActive();
+			scope.Span.SetTag("LogId", logFile.Id.ToString());
+			scope.Span.SetTag("Text", text);
+			scope.Span.SetTag("Count", count.ToString(CultureInfo.InvariantCulture));
+
+			List<int> results = new List<int>();
+			if (count > 0)
+			{
+				IAsyncEnumerator<int> enumerator = SearchLogDataInternalAsync(logFile, text, firstLine, searchStats).GetAsyncEnumerator();
+				while (await enumerator.MoveNextAsync() && results.Count < count)
+				{
+					results.Add(enumerator.Current);
+				}
+			}
+
+			_logger.LogDebug("Search for \"{SearchText}\" in log {LogId} found {NumResults}/{MaxResults} results, took {Time}ms ({@Stats})", text, logFile.Id, results.Count, count, timer.ElapsedMilliseconds, searchStats);
+			return results;
+		}
+
+		/// <inheritdoc/>
+		public async IAsyncEnumerable<int> SearchLogDataInternalAsync(ILogFile logFile, string text, int firstLine, LogSearchStats searchStats)
+		{
+			SearchText searchText = new SearchText(text);
+
+			// Read the index for this log file
+			if (logFile.IndexLength != null)
+			{
+				LogIndexData? indexData = await ReadIndexAsync(logFile, logFile.IndexLength.Value);
+				if(indexData != null && firstLine < indexData.LineCount)
+				{
+					using IScope indexScope = GlobalTracer.Instance.BuildSpan("Indexed").StartActive();
+					indexScope.Span.SetTag("LineCount", indexData.LineCount.ToString(CultureInfo.InvariantCulture));
+
+					foreach(int lineIndex in indexData.Search(firstLine, searchText, searchStats))
+					{
+						yield return lineIndex;
+					}
+
+					firstLine = indexData.LineCount;
+				}
+			}
+
+			// Manually search through the rest of the log
+			int chunkIdx = logFile.Chunks.GetChunkForLine(firstLine);
+			for (; chunkIdx < logFile.Chunks.Count; chunkIdx++)
+			{
+				ILogChunk chunk = logFile.Chunks[chunkIdx];
+
+				// Read the chunk data
+				LogChunkData chunkData = await ReadChunkAsync(logFile, chunkIdx);
+				if (firstLine < chunkData.LineIndex + chunkData.LineCount)
+				{
+					// Find the first sub-chunk we're looking for
+					int subChunkIdx = 0;
+					if (firstLine > chunk.LineIndex)
+					{
+						subChunkIdx = chunkData.GetSubChunkForLine(firstLine - chunk.LineIndex);
+					}
+
+					// Search through the sub-chunks
+					for (; subChunkIdx < chunkData.SubChunks.Count; subChunkIdx++)
+					{
+						LogSubChunkData subChunkData = chunkData.SubChunks[subChunkIdx];
+						if (firstLine < subChunkData.LineIndex + subChunkData.LineCount)
+						{
+							// Create an index containing just this sub-chunk
+							LogIndexData index = subChunkData.BuildIndex(_logger);
+							foreach (int lineIndex in index.Search(firstLine, searchText, searchStats))
+							{
+								yield return lineIndex;
+							}
+						}
+					}
+				}
+			}
+		}
+	}
+}
+
diff --git a/Engine/Source/Programs/Horde/Horde.Build/Logs/LogsController.cs b/Engine/Source/Programs/Horde/Horde.Build/Logs/LogsController.cs
new file mode 100644
--- /dev/null
+++ b/Engine/Source/Programs/Horde/Horde.Build/Logs/LogsController.cs
@@ -0,0 +1,378 @@
+// Copyright Epic Games, Inc. All Rights Reserved.
+
+using System;
+using System.Collections.Generic;
+using System.IO;
+using System.Security.Claims;
+using System.Text;
+using System.Threading.Tasks;
+using Horde.Build.Acls;
+using Horde.Build.Api;
+using Horde.Build.Collections;
+using Horde.Build.Logs;
+using Horde.Build.Models;
+using Horde.Build.Services;
+using Horde.Build.Utilities;
+using Microsoft.AspNetCore.Authorization;
+using Microsoft.AspNetCore.Mvc;
+using MongoDB.Bson;
+
+namespace Horde.Build.Controllers
+{
+	using JobId = ObjectId<IJob>;
+	using LogId = ObjectId<ILogFile>;
+
+	/// <summary>
+	/// Format for the returned data
+	/// </summary>
+	public enum LogOutputFormat
+	{
+		/// <summary>
+		/// Plain text
+		/// </summary>
+		Text,
+
+		/// <summary>
+		/// Raw output (text/json)
+		/// </summary>
+		Raw,
+	}
+
+	/// <summary>
+	/// Controller for the /api/logs endpoint
+	/// </summary>
+	[ApiController]
+	[Authorize]
+	[Route("[controller]")]
+	public class LogsController : ControllerBase
+	{
+		private readonly ILogFileService _logFileService;
+		private readonly IIssueCollection _issueCollection;
+		private readonly AclService _aclService;
+		private readonly JobService _jobService;
+
+		/// <summary>
+		/// Constructor
+		/// </summary>
+		public LogsController(ILogFileService logFileService, IIssueCollection issueCollection, AclService aclService, JobService jobService)
+		{
+			_logFileService = logFileService;
+			_issueCollection = issueCollection;
+			_aclService = aclService;
+			_jobService = jobService;
+		}
+
+		/// <summary>
+		/// Retrieve metadata about a specific log file
+		/// </summary>
+		/// <param name="logFileId">Id of the log file to get information about</param>
+		/// <param name="filter">Filter for the properties to return</param>
+		/// <returns>Information about the requested project</returns>
+		[HttpGet]
+		[Route("/api/v1/logs/{logFileId}")]
+		[ProducesResponseType(typeof(GetLogFileResponse), 200)]
+		public async Task<ActionResult<object>> GetLog(LogId logFileId, [FromQuery] PropertyFilter? filter = null)
+		{
+			ILogFile? logFile = await _logFileService.GetLogFileAsync(logFileId);
+			if (logFile == null)
+			{
+				return NotFound();
+			}
+			if (!await AuthorizeAsync(logFile, AclAction.ViewLog, User, null))
+			{
+				return Forbid();
+			}
+
+			LogMetadata metadata = await _logFileService.GetMetadataAsync(logFile);
+			return new GetLogFileResponse(logFile, metadata).ApplyFilter(filter);       
+		}
+
+		/// <summary>
+		/// Retrieve raw data for a log file
+		/// </summary>
+		/// <param name="logFileId">Id of the log file to get information about</param>
+		/// <param name="format">Format for the returned data</param>
+		/// <param name="offset">The log offset in bytes</param>
+		/// <param name="length">Number of bytes to return</param>
+		/// <param name="fileName">Name of the default filename to download</param>
+		/// <param name="download">Whether to download the file rather than display in the browser</param>
+		/// <returns>Raw log data for the requested range</returns>
+		[HttpGet]
+		[Route("/api/v1/logs/{logFileId}/data")]
+		public async Task<ActionResult> GetLogData(LogId logFileId, [FromQuery] LogOutputFormat format = LogOutputFormat.Raw, [FromQuery] long offset = 0, [FromQuery] long length = Int64.MaxValue, [FromQuery] string? fileName = null, [FromQuery] bool download = false)
+		{
+			ILogFile? logFile = await _logFileService.GetLogFileAsync(logFileId);
+			if (logFile == null)
+			{
+				return NotFound();
+			}
+			if (!await AuthorizeAsync(logFile, AclAction.ViewLog, User, null))
+			{
+				return Forbid();
+			}
+
+			Func<Stream, ActionContext, Task> copyTask;
+			if (format == LogOutputFormat.Text && logFile.Type == LogType.Json)
+			{
+				copyTask = (outputStream, context) => _logFileService.CopyPlainTextStreamAsync(logFile, offset, length, outputStream);
+			}
+			else
+			{
+				copyTask = (outputStream, context) => _logFileService.CopyRawStreamAsync(logFile, offset, length, outputStream);
+			}
+
+			return new CustomFileCallbackResult(fileName ?? $"log-{logFileId}.txt", "text/plain", !download, copyTask);
+		}
+
+		/// <summary>
+		/// Retrieve line data for a logfile
+		/// </summary>
+		/// <param name="logFileId">Id of the log file to get information about</param>
+		/// <param name="index">Index of the first line to retrieve</param>
+		/// <param name="count">Number of lines to retrieve</param>
+		/// <returns>Information about the requested project</returns>
+		[HttpGet]
+		[Route("/api/v1/logs/{logFileId}/lines")]
+		public async Task<ActionResult> GetLogLines(LogId logFileId, [FromQuery] int index = 0, [FromQuery] int count = Int32.MaxValue)
+		{
+			ILogFile? logFile = await _logFileService.GetLogFileAsync(logFileId);
+			if (logFile == null)
+			{
+				return NotFound();
+			}
+			if (!await AuthorizeAsync(logFile, AclAction.ViewLog, User, null))
+			{
+				return Forbid();
+			}
+
+			LogMetadata metadata = await _logFileService.GetMetadataAsync(logFile);
+
+			(int minIndex, long minOffset) = await _logFileService.GetLineOffsetAsync(logFile, index);
+			(int maxIndex, long maxOffset) = await _logFileService.GetLineOffsetAsync(logFile, index + Math.Min(count, Int32.MaxValue - index));
+			index = minIndex;
+			count = maxIndex - minIndex;
+
+			byte[] result;
+			using (System.IO.Stream stream = await _logFileService.OpenRawStreamAsync(logFile, minOffset, maxOffset - minOffset))
+			{
+				result = new byte[stream.Length];
+				await stream.ReadFixedSizeDataAsync(result, 0, result.Length);
+			}
+
+			using (MemoryStream stream = new MemoryStream(result.Length + (count * 20)))
+			{
+				stream.WriteByte((byte)'{');
+
+				stream.Write(Encoding.UTF8.GetBytes($"\"index\":{index},"));
+				stream.Write(Encoding.UTF8.GetBytes($"\"count\":{count},"));
+				stream.Write(Encoding.UTF8.GetBytes($"\"maxLineIndex\":{metadata.MaxLineIndex},"));
+				stream.Write(Encoding.UTF8.GetBytes($"\"format\":{ (logFile.Type == LogType.Json ? "\"JSON\"" : "\"TEXT\"")},"));
+
+//				Stream.Write(Encoding.UTF8.GetBytes($"\"minIndex\":{MinIndex},"));
+//				Stream.Write(Encoding.UTF8.GetBytes($"\"minOffset\":{MinOffset},"));
+//				Stream.Write(Encoding.UTF8.GetBytes($"\"maxIndex\":{MaxIndex},"));
+//				Stream.Write(Encoding.UTF8.GetBytes($"\"maxOffset\":{MaxOffset},"));
+//				Stream.Write(Encoding.UTF8.GetBytes($"\"length\":{Result.Length},"));
+
+				stream.Write(Encoding.UTF8.GetBytes($"\"lines\":["));
+				stream.WriteByte((byte)'\n');
+
+				int offset = 0;
+				for (int line = index; line < index + count; line++)
+				{
+					stream.WriteByte((byte)' ');
+					stream.WriteByte((byte)' ');
+
+					if (logFile.Type == LogType.Json)
+					{
+						// Find the end of the line and output it as an opaque blob
+						int startOffset = offset;
+						for (; ; offset++)
+						{
+							if (offset == result.Length)
+							{
+								stream.WriteByte((byte)'{');
+								stream.WriteByte((byte)'}');
+								break;
+							}
+							else if (result[offset] == (byte)'\n')
+							{
+								await stream.WriteAsync(result.AsMemory(startOffset, offset - startOffset));
+								offset++;
+								break;
+							}
+						}
+					}
+					else
+					{
+						stream.WriteByte((byte)'\"');
+						for (; offset < result.Length; offset++)
+						{
+							if (result[offset] == '\\' || result[offset] == '\"')
+							{
+								stream.WriteByte((byte)'\\');
+								stream.WriteByte(result[offset]);
+							}
+							else if (result[offset] == (byte)'\n')
+							{
+								offset++;
+								break;
+							}
+							else if (result[offset] >= 32 && result[offset] <= 126)
+							{
+								stream.WriteByte(result[offset]);
+							}
+							else
+							{
+								stream.Write(Encoding.UTF8.GetBytes($"\\x{result[offset]:x2}"));
+							}
+						}
+						stream.WriteByte((byte)'\"');
+					}
+
+					if (line + 1 < index + count)
+					{
+						stream.WriteByte((byte)',');
+					}
+
+					stream.WriteByte((byte)'\n');
+				}
+
+				if (logFile.Type == LogType.Json)
+				{
+					stream.Write(Encoding.UTF8.GetBytes($"]"));
+				}
+
+				stream.WriteByte((byte)'}');
+
+				Response.ContentType = "application/json";
+				Response.Headers.ContentLength = stream.Length;
+				stream.Position = 0;
+				await stream.CopyToAsync(Response.Body);
+			}
+			return new EmptyResult();
+		}
+
+		/// <summary>
+		/// Search log data
+		/// </summary>
+		/// <param name="logFileId">Id of the log file to get information about</param>
+		/// <param name="text">Text to search for</param>
+		/// <param name="firstLine">First line to search from</param>
+		/// <param name="count">Number of results to return</param>
+		/// <returns>Raw log data for the requested range</returns>
+		[HttpGet]
+		[Route("/api/v1/logs/{logFileId}/search")]
+		public async Task<ActionResult<SearchLogFileResponse>> SearchLogFileAsync(LogId logFileId, [FromQuery] string text, [FromQuery] int firstLine = 0, [FromQuery] int count = 5)
+		{
+			ILogFile? logFile = await _logFileService.GetLogFileAsync(logFileId);
+			if (logFile == null)
+			{
+				return NotFound();
+			}
+			if (!await AuthorizeAsync(logFile, AclAction.ViewLog, User, null))
+			{
+				return Forbid();
+			}
+
+			SearchLogFileResponse response = new SearchLogFileResponse();
+			response.Stats = new LogSearchStats();
+			response.Lines = await _logFileService.SearchLogDataAsync(logFile, text, firstLine, count, response.Stats);
+			return response;
+		}
+
+		/// <summary>
+		/// Retrieve events for a logfile
+		/// </summary>
+		/// <param name="logFileId">Id of the log file to get information about</param>
+		/// <param name="index">Index of the first line to retrieve</param>
+		/// <param name="count">Number of lines to retrieve</param>
+		/// <returns>Information about the requested project</returns>
+		[HttpGet]
+		[Route("/api/v1/logs/{logFileId}/events")]
+		[ProducesResponseType(typeof(List<GetLogEventResponse>), 200)]
+		public async Task<ActionResult<List<GetLogEventResponse>>> GetEventsAsync(LogId logFileId, [FromQuery] int? index = null, [FromQuery] int? count = null)
+		{
+			ILogFile? logFile = await _logFileService.GetLogFileAsync(logFileId);
+			if (logFile == null)
+			{
+				return NotFound();
+			}
+			if (!await AuthorizeAsync(logFile, AclAction.ViewLog, User, null))
+			{
+				return Forbid();
+			}
+
+			List<ILogEvent> logEvents = await _logFileService.FindEventsAsync(logFile, null, index, count);
+
+			Dictionary<ObjectId, int?> spanIdToIssueId = new Dictionary<ObjectId, int?>();
+
+			List<GetLogEventResponse> responses = new List<GetLogEventResponse>();
+			foreach (ILogEvent logEvent in logEvents)
+			{
+				ILogEventData logEventData = await _logFileService.GetEventDataAsync(logFile, logEvent.LineIndex, logEvent.LineCount);
+
+				int? issueId = null;
+				if (logEvent.SpanId != null && !spanIdToIssueId.TryGetValue(logEvent.SpanId.Value, out issueId))
+				{
+					IIssueSpan? span = await _issueCollection.GetSpanAsync(logEvent.SpanId.Value);
+					issueId = span?.IssueId;
+					spanIdToIssueId[logEvent.SpanId.Value] = issueId;
+				}
+
+				responses.Add(new GetLogEventResponse(logEvent, logEventData, issueId));
+			}
+			return responses;
+		}
+
+		/// <summary>
+		/// Appends data to a log file
+		/// </summary>
+		/// <param name="logFileId">The logfile id</param>
+		/// <param name="offset">Offset within the log file</param>
+		/// <param name="lineIndex">The line index</param>
+		/// <returns>Http result code</returns>
+		[HttpPost]
+		[Route("/api/v1/logs/{logFileId}")]
+		public async Task<ActionResult> WriteData(LogId logFileId, [FromQuery] long offset, [FromQuery] int lineIndex)
+		{
+			ILogFile? logFile = await _logFileService.GetLogFileAsync(logFileId);
+			if (logFile == null)
+			{
+				return NotFound();
+			}
+			if (!await AuthorizeAsync(logFile, AclAction.WriteLogData, User, null))
+			{
+				return Forbid();
+			}
+
+			using (MemoryStream bodyStream = new MemoryStream())
+			{
+				await Request.Body.CopyToAsync(bodyStream);
+				await _logFileService.WriteLogDataAsync(logFile, offset, lineIndex, bodyStream.ToArray(), false);
+			}
+			return Ok();
+		}
+
+		/// <summary>
+		/// Determines if the user is authorized to perform an action on a particular template
+		/// </summary>
+		/// <param name="logFile">The template to check</param>
+		/// <param name="action">The action being performed</param>
+		/// <param name="user">The principal to authorize</param>
+		/// <param name="permissionsCache">Permissions cache</param>
+		/// <returns>True if the action is authorized</returns>
+		async Task<bool> AuthorizeAsync(ILogFile logFile, AclAction action, ClaimsPrincipal user, JobPermissionsCache? permissionsCache)
+		{
+			if (logFile.JobId != JobId.Empty && await _jobService.AuthorizeAsync(logFile.JobId, action, user, permissionsCache))
+			{
+				return true;
+			}
+			if (logFile.SessionId != null && await _aclService.AuthorizeAsync(AclAction.ViewSession, user, permissionsCache))
+			{
+				return true;
+			}
+			return false;
+		}
+	}
+}
diff --git a/Engine/Source/Programs/Horde/Horde.Build/Logs/Storage/LocalLogStorage.cs b/Engine/Source/Programs/Horde/Horde.Build/Logs/Storage/LocalLogStorage.cs
new file mode 100644
--- /dev/null
+++ b/Engine/Source/Programs/Horde/Horde.Build/Logs/Storage/LocalLogStorage.cs
@@ -0,0 +1,144 @@
+// Copyright Epic Games, Inc. All Rights Reserved.
+
+using System;
+using System.Threading.Tasks;
+using Horde.Build.Models;
+using Horde.Build.Utilities;
+using Microsoft.Extensions.Caching.Memory;
+
+namespace Horde.Build.Logs.Storage
+{
+	using LogId = ObjectId<ILogFile>;
+
+	/// <summary>
+	/// In-memory cache for chunk and index data
+	/// </summary>
+	class LocalLogStorage : ILogStorage
+	{
+		/// <summary>
+		/// The memory cache instance
+		/// </summary>
+		readonly MemoryCache _cache;
+
+		/// <summary>
+		/// The inner storage
+		/// </summary>
+		readonly ILogStorage _inner;
+
+		/// <summary>
+		/// Constructor
+		/// </summary>
+		/// <param name="numItems">Maximum size for the memory cache</param>
+		/// <param name="inner">The inner storage provider</param>
+		public LocalLogStorage(int numItems, ILogStorage inner)
+		{
+			MemoryCacheOptions options = new MemoryCacheOptions();
+			options.SizeLimit = numItems;
+
+			_cache = new MemoryCache(options);
+			_inner = inner;
+		}
+
+		/// <inheritdoc/>
+		public void Dispose()
+		{
+			_cache.Dispose();
+			_inner.Dispose();
+		}
+
+		/// <summary>
+		/// Gets the cache key for a particular index
+		/// </summary>
+		/// <param name="logId">The log file to retrieve an index for</param>
+		/// <param name="length">Length of the file covered by the index</param>
+		/// <returns>Cache key for the index</returns>
+		static string IndexKey(LogId logId, long length)
+		{
+			return $"{logId}/index-{length}";
+		}
+
+		/// <summary>
+		/// Gets the cache key for a particular chunk
+		/// </summary>
+		/// <param name="logId">The log file to retrieve an index for</param>
+		/// <param name="offset">The chunk offset</param>
+		/// <returns>Cache key for the chunk</returns>
+		static string ChunkKey(LogId logId, long offset)
+		{
+			return $"{logId}/chunk-{offset}";
+		}
+
+		/// <summary>
+		/// Adds an entry to the cache
+		/// </summary>
+		/// <param name="key">The cache key</param>
+		/// <param name="value">Value to store</param>
+		void AddEntry(string key, object? value)
+		{
+			using (ICacheEntry entry = _cache.CreateEntry(key))
+			{
+				if (value == null)
+				{
+					entry.SetAbsoluteExpiration(TimeSpan.FromSeconds(30.0));
+					entry.SetSize(0);
+				}
+				else
+				{
+					entry.SetSlidingExpiration(TimeSpan.FromHours(4));
+					entry.SetSize(1);
+				}
+				entry.SetValue(value);
+			}
+		}
+
+		/// <summary>
+		/// Reads a value from the cache, or from the inner storage provider, adding the result to the cache
+		/// </summary>
+		/// <typeparam name="T">Type of object to read</typeparam>
+		/// <param name="key">The cache key</param>
+		/// <param name="readInner">Delegate to read from the inner storage provider</param>
+		/// <returns>The retrieved value</returns>
+		async Task<T?> ReadValueAsync<T>(string key, Func<Task<T?>> readInner) where T : class
+		{
+			T? value;
+			if (!_cache.TryGetValue(key, out value))
+			{
+				try
+				{
+					value = await readInner();
+				}
+				finally
+				{
+					AddEntry(key, value);
+				}
+			}
+			return value;
+		}
+
+		/// <inheritdoc/>
+		public Task<LogIndexData?> ReadIndexAsync(LogId logId, long length)
+		{
+			return ReadValueAsync(IndexKey(logId, length), () => _inner.ReadIndexAsync(logId, length));
+		}
+
+		/// <inheritdoc/>
+		public Task WriteIndexAsync(LogId logId, long length, LogIndexData indexData)
+		{
+			AddEntry(IndexKey(logId, length), indexData);
+			return _inner.WriteIndexAsync(logId, length, indexData);
+		}
+
+		/// <inheritdoc/>
+		public Task<LogChunkData?> ReadChunkAsync(LogId logId, long offset, int lineIndex)
+		{
+			return ReadValueAsync(ChunkKey(logId, offset), () => _inner.ReadChunkAsync(logId, offset, lineIndex));
+		}
+
+		/// <inheritdoc/>
+		public Task WriteChunkAsync(LogId logId, long offset, LogChunkData chunkData)
+		{
+			AddEntry(ChunkKey(logId, offset), chunkData);
+			return _inner.WriteChunkAsync(logId, offset, chunkData);
+		}
+	}
+}
diff --git a/Engine/Source/Programs/Horde/Horde.Build/Logs/Storage/NullLogStorage.cs b/Engine/Source/Programs/Horde/Horde.Build/Logs/Storage/NullLogStorage.cs
new file mode 100644
--- /dev/null
+++ b/Engine/Source/Programs/Horde/Horde.Build/Logs/Storage/NullLogStorage.cs
@@ -0,0 +1,45 @@
+// Copyright Epic Games, Inc. All Rights Reserved.
+
+using System.Threading.Tasks;
+using Horde.Build.Models;
+using Horde.Build.Utilities;
+
+namespace Horde.Build.Logs.Storage.Impl
+{
+	using LogId = ObjectId<ILogFile>;
+
+	/// <summary>
+	/// Empty implementation of log storage
+	/// </summary>
+	public sealed class NullLogStorage : ILogStorage
+	{
+		/// <inheritdoc/>
+		public void Dispose()
+		{
+		}
+
+		/// <inheritdoc/>
+		public Task<LogIndexData?> ReadIndexAsync(LogId logId, long length)
+		{
+			return Task.FromResult<LogIndexData?>(null);
+		}
+
+		/// <inheritdoc/>
+		public Task WriteIndexAsync(LogId logId, long length, LogIndexData index)
+		{
+			return Task.CompletedTask;
+		}
+
+		/// <inheritdoc/>
+		public Task<LogChunkData?> ReadChunkAsync(LogId logId, long offset, int lineIndex)
+		{
+			return Task.FromResult<LogChunkData?>(null);
+		}
+
+		/// <inheritdoc/>
+		public Task WriteChunkAsync(LogId logId, long offset, LogChunkData chunkData)
+		{
+			return Task.CompletedTask;
+		}
+	}
+}
diff --git a/Engine/Source/Programs/Horde/Horde.Build/Logs/Storage/PersistentLogStorage.cs b/Engine/Source/Programs/Horde/Horde.Build/Logs/Storage/PersistentLogStorage.cs
new file mode 100644
--- /dev/null
+++ b/Engine/Source/Programs/Horde/Horde.Build/Logs/Storage/PersistentLogStorage.cs
@@ -0,0 +1,107 @@
+// Copyright Epic Games, Inc. All Rights Reserved.
+
+using System;
+using System.Threading.Tasks;
+using EpicGames.Core;
+using Horde.Build.Models;
+using Horde.Build.Storage;
+using Horde.Build.Utilities;
+using Microsoft.Extensions.Logging;
+
+namespace Horde.Build.Logs.Readers
+{
+	using LogId = ObjectId<ILogFile>;
+
+	/// <summary>
+	/// Bulk storage for log file data
+	/// </summary>
+	class PersistentLogStorage : ILogStorage
+	{
+		/// <summary>
+		/// The bulk storage provider to use
+		/// </summary>
+		readonly IStorageBackend _storageProvider;
+
+		/// <summary>
+		/// Log provider
+		/// </summary>
+		readonly ILogger _logger;
+
+		/// <summary>
+		/// Constructor
+		/// </summary>
+		/// <param name="storageProvider">The storage provider</param>
+		/// <param name="logger">Logging provider</param>
+		public PersistentLogStorage(IStorageBackend<PersistentLogStorage> storageProvider, ILogger<PersistentLogStorage> logger)
+		{
+			_storageProvider = storageProvider;
+			_logger = logger;
+		}
+
+		/// <inheritdoc/>
+		public void Dispose()
+		{
+		}
+
+		/// <inheritdoc/>
+		public async Task<LogIndexData?> ReadIndexAsync(LogId logId, long length)
+		{
+			_logger.LogDebug("Reading log {LogId} index length {Length} from persistent storage", logId, length);
+
+			string path = $"{logId}/index_{length}";
+			ReadOnlyMemory<byte>? data = await _storageProvider.ReadBytesAsync(path);
+			if (data == null)
+			{
+				return null;
+			}
+			return LogIndexData.FromMemory(data.Value);
+		}
+
+		/// <inheritdoc/>
+		public Task WriteIndexAsync(LogId logId, long length, LogIndexData indexData)
+		{
+			_logger.LogDebug("Writing log {LogId} index length {Length} to persistent storage", logId, length);
+
+			string path = $"{logId}/index_{length}";
+			ReadOnlyMemory<byte> data = indexData.ToByteArray();
+			return _storageProvider.WriteBytesAsync(path, data);
+		}
+
+		/// <inheritdoc/>
+		public async Task<LogChunkData?> ReadChunkAsync(LogId logId, long offset, int lineIndex)
+		{
+			_logger.LogDebug("Reading log {LogId} chunk offset {Offset} from persistent storage", logId, offset);
+
+			string path = $"{logId}/offset_{offset}";
+			ReadOnlyMemory<byte>? data = await _storageProvider.ReadBytesAsync(path);
+			if(data == null)
+			{
+				return null;
+			}
+
+			MemoryReader reader = new MemoryReader(data.Value);
+			LogChunkData chunkData = reader.ReadLogChunkData(offset, lineIndex);
+
+			if (reader.Offset != data.Value.Length)
+			{
+				throw new Exception($"Serialization of persistent chunk {path} is not at expected offset (expected {data.Value.Length}, actual {reader.Offset})");
+			}
+
+			return chunkData;
+		}
+
+		/// <inheritdoc/>
+		public Task WriteChunkAsync(LogId logId, long offset, LogChunkData chunkData)
+		{
+			_logger.LogDebug("Writing log {LogId} chunk offset {Offset} to persistent storage", logId, offset);
+
+			string path = $"{logId}/offset_{offset}";
+			byte[] data = new byte[chunkData.GetSerializedSize(_logger)];
+			MemoryWriter writer = new MemoryWriter(data);
+			writer.WriteLogChunkData(chunkData, _logger);
+			writer.CheckOffset(data.Length);
+
+			return _storageProvider.WriteBytesAsync(path, data);
+		}
+	}
+}
diff --git a/Engine/Source/Programs/Horde/Horde.Build/Logs/Storage/RedisLogStorage.cs b/Engine/Source/Programs/Horde/Horde.Build/Logs/Storage/RedisLogStorage.cs
new file mode 100644
--- /dev/null
+++ b/Engine/Source/Programs/Horde/Horde.Build/Logs/Storage/RedisLogStorage.cs
@@ -0,0 +1,169 @@
+// Copyright Epic Games, Inc. All Rights Reserved.
+
+using System;
+using System.Threading.Tasks;
+using Horde.Build.Models;
+using Horde.Build.Utilities;
+using Microsoft.Extensions.Logging;
+using StackExchange.Redis;
+
+namespace Horde.Build.Logs.Storage.Impl
+{
+	using LogId = ObjectId<ILogFile>;
+
+	/// <summary>
+	/// Redis log file storage
+	/// </summary>
+	public sealed class RedisLogStorage : ILogStorage
+	{
+		/// <summary>
+		/// The Redis database
+		/// </summary>
+		readonly IDatabase _redisDb;
+
+		/// <summary>
+		/// Logging device
+		/// </summary>
+		readonly ILogger _logger;
+
+		/// <summary>
+		/// Inner storage layer
+		/// </summary>
+		readonly ILogStorage _inner;
+
+		/// <summary>
+		/// Constructor
+		/// </summary>
+		/// <param name="redisDb">Redis database</param>
+		/// <param name="logger">Logger instance</param>
+		/// <param name="inner">Inner storage layer</param>
+		public RedisLogStorage(IDatabase redisDb, ILogger logger, ILogStorage inner)
+		{
+			_redisDb = redisDb;
+			_logger = logger;
+			_inner = inner;
+		}
+
+		/// <inheritdoc/>
+		public void Dispose()
+		{
+			_inner.Dispose();
+		}
+
+		/// <summary>
+		/// Gets the key for a log file's index
+		/// </summary>
+		/// <param name="logId">The log file id</param>
+		/// <param name="length">Length of the file covered by the index</param>
+		/// <returns>The index key</returns>
+		static string IndexKey(LogId logId, long length)
+		{
+			return $"log-{logId}-index-{length}";
+		}
+
+		/// <summary>
+		/// Gets the key for a chunk's data
+		/// </summary>
+		/// <param name="logId">The log file id</param>
+		/// <param name="offset">Offset of the chunk within the log file</param>
+		/// <returns>The chunk key</returns>
+		static string ChunkKey(LogId logId, long offset)
+		{
+			return $"log-{logId}-chunk-{offset}";
+		}
+
+		/// <summary>
+		/// Adds an index to the cache
+		/// </summary>
+		/// <param name="key">Key for the item to store</param>
+		/// <param name="value">Value to be stored</param>
+		/// <returns>Async task</returns>
+		async Task AddAsync(string key, ReadOnlyMemory<byte> value)
+		{
+			try
+			{
+				await _redisDb.StringSetAsync(key, value, expiry: TimeSpan.FromHours(1.0));
+				_logger.LogDebug("Added key {Key} to Redis cache ({Size} bytes)", key, value.Length);
+			}
+			catch (Exception ex)
+			{
+				_logger.LogError(ex, "Error writing Redis key {Key}", key);
+			}
+		}
+
+		/// <summary>
+		/// Adds an index to the cache
+		/// </summary>
+		/// <param name="key">Key for the item to store</param>
+		/// <param name="deserialize">Delegate to deserialize the item</param>
+		/// <returns>Async task</returns>
+		async Task<T?> GetAsync<T>(string key, Func<ReadOnlyMemory<byte>, T> deserialize) where T : class
+		{
+			try
+			{
+				RedisValue value = await _redisDb.StringGetAsync(key);
+				if (value.IsNullOrEmpty)
+				{
+					_logger.LogDebug("Redis cache miss for {Key}", key);
+					return null;
+				}
+				else
+				{
+					_logger.LogDebug("Redis cache hit for {Key}", key);
+					return deserialize(value);
+				}
+			}
+			catch (Exception ex)
+			{
+				_logger.LogError(ex, "Error reading Redis key {Key}", key);
+				return null;
+			}
+		}
+
+		/// <inheritdoc/>
+		public async Task<LogIndexData?> ReadIndexAsync(LogId logId, long length)
+		{
+			string key = IndexKey(logId, length);
+
+			LogIndexData? indexData = await GetAsync(key, memory => LogIndexData.FromMemory(memory));
+			if(indexData == null)
+			{
+				indexData = await _inner.ReadIndexAsync(logId, length);
+				if(indexData != null)
+				{
+					await AddAsync(key, indexData.ToByteArray());
+				}
+			}
+			return indexData;
+		}
+
+		/// <inheritdoc/>
+		public async Task WriteIndexAsync(LogId logId, long length, LogIndexData index)
+		{
+			await _inner.WriteIndexAsync(logId, length, index);
+		}
+
+		/// <inheritdoc/>
+		public async Task<LogChunkData?> ReadChunkAsync(LogId logId, long offset, int lineIndex)
+		{
+			string key = ChunkKey(logId, offset);
+
+			LogChunkData? chunkData = await GetAsync(key, memory => LogChunkData.FromMemory(memory, offset, lineIndex));
+			if (chunkData == null)
+			{
+				chunkData = await _inner.ReadChunkAsync(logId, offset, lineIndex);
+				if (chunkData != null)
+				{
+					await AddAsync(key, chunkData.ToByteArray(_logger));
+				}
+			}
+			return chunkData;
+		}
+
+		/// <inheritdoc/>
+		public async Task WriteChunkAsync(LogId logId, long offset, LogChunkData chunkData)
+		{
+			await _inner.WriteChunkAsync(logId, offset, chunkData);
+		}
+	}
+}
diff --git a/Engine/Source/Programs/Horde/Horde.Build/Logs/Storage/SequencedLogStorage.cs b/Engine/Source/Programs/Horde/Horde.Build/Logs/Storage/SequencedLogStorage.cs
new file mode 100644
--- /dev/null
+++ b/Engine/Source/Programs/Horde/Horde.Build/Logs/Storage/SequencedLogStorage.cs
@@ -0,0 +1,193 @@
+// Copyright Epic Games, Inc. All Rights Reserved.
+
+using System.Collections.Generic;
+using System.Threading.Tasks;
+using Horde.Build.Models;
+using Horde.Build.Utilities;
+
+namespace Horde.Build.Logs.Storage.Impl
+{
+	using LogId = ObjectId<ILogFile>;
+
+	/// <summary>
+	/// Storage layer which caches pending read tasks, to avoid fetching the same item more than once
+	/// </summary>
+	class SequencedLogStorage : ILogStorage
+	{
+		/// <summary>
+		/// Inner storage implementation
+		/// </summary>
+		readonly ILogStorage _inner;
+
+		/// <summary>
+		/// Pending index reads
+		/// </summary>
+		readonly Dictionary<(LogId, long), Task<LogIndexData?>> _indexReadTasks = new Dictionary<(LogId, long), Task<LogIndexData?>>();
+
+		/// <summary>
+		/// Pending index reads
+		/// </summary>
+		readonly Dictionary<(LogId, long), Task> _indexWriteTasks = new Dictionary<(LogId, long), Task>();
+
+		/// <summary>
+		/// Pending chunk reads
+		/// </summary>
+		readonly Dictionary<(LogId, long), Task<LogChunkData?>> _chunkReadTasks = new Dictionary<(LogId, long), Task<LogChunkData?>>();
+
+		/// <summary>
+		/// Pending chunk reads
+		/// </summary>
+		readonly Dictionary<(LogId, long), Task> _chunkWriteTasks = new Dictionary<(LogId, long), Task>();
+
+		/// <summary>
+		/// Constructor
+		/// </summary>
+		/// <param name="inner">The inner storage provider</param>
+		public SequencedLogStorage(ILogStorage inner)
+		{
+			_inner = inner;
+		}
+
+		/// <inheritdoc/>
+		public void Dispose()
+		{
+			_inner.Dispose();
+		}
+
+		/// <inheritdoc/>
+		public Task<LogIndexData?> ReadIndexAsync(LogId logId, long length)
+		{
+			Task<LogIndexData?>? task;
+			lock (_indexReadTasks)
+			{
+				if (!_indexReadTasks.TryGetValue((logId, length), out task))
+				{
+					task = InnerReadIndexAsync(logId, length);
+					_indexReadTasks.Add((logId, length), task);
+				}
+			}
+			return task;
+		}
+
+		/// <inheritdoc/>
+		public Task WriteIndexAsync(LogId logId, long length, LogIndexData indexData)
+		{
+			Task? task;
+			lock (_indexWriteTasks)
+			{
+				if (!_indexWriteTasks.TryGetValue((logId, length), out task))
+				{
+					task = InnerWriteIndexAsync(logId, length, indexData);
+					_indexWriteTasks.Add((logId, length), task);
+				}
+			}
+			return task;
+		}
+
+		/// <inheritdoc/>
+		public Task<LogChunkData?> ReadChunkAsync(LogId logId, long offset, int lineIndex)
+		{
+			Task<LogChunkData?>? task;
+			lock (_chunkReadTasks)
+			{
+				if (!_chunkReadTasks.TryGetValue((logId, offset), out task))
+				{
+					task = InnerReadChunkAsync(logId, offset, lineIndex);
+					_chunkReadTasks[(logId, offset)] = task;
+				}
+			}
+			return task;
+		}
+
+		/// <inheritdoc/>
+		public Task WriteChunkAsync(LogId logId, long offset, LogChunkData chunkData)
+		{
+			Task? task;
+			lock (_chunkWriteTasks)
+			{
+				if (!_chunkWriteTasks.TryGetValue((logId, offset), out task))
+				{
+					task = InnerWriteChunkAsync(logId, offset, chunkData);
+					_chunkWriteTasks[(logId, offset)] = task;
+				}
+			}
+			return task;
+		}
+
+		/// <summary>
+		/// Wrapper for reading an index from the inner storage provider
+		/// </summary>
+		/// <param name="logId">The log file to read the index for</param>
+		/// <param name="length">Length of the file that's indexed</param>
+		/// <returns>The index data</returns>
+		async Task<LogIndexData?> InnerReadIndexAsync(LogId logId, long length)
+		{
+			await Task.Yield();
+
+			LogIndexData? indexData = await _inner.ReadIndexAsync(logId, length);
+			lock (_indexReadTasks)
+			{
+				_indexReadTasks.Remove((logId, length));
+			}
+			return indexData;
+		}
+
+		/// <summary>
+		/// Wrapper for reading an index from the inner storage provider
+		/// </summary>
+		/// <param name="logId">The log file to read the index for</param>
+		/// <param name="length">Length of the indexed data</param>
+		/// <param name="indexData">The index data to write</param>
+		/// <returns>The index data</returns>
+		async Task InnerWriteIndexAsync(LogId logId, long length, LogIndexData indexData)
+		{
+			await Task.Yield();
+
+			await _inner.WriteIndexAsync(logId, length, indexData);
+
+			lock (_indexWriteTasks)
+			{
+				_indexWriteTasks.Remove((logId, length));
+			}
+		}
+
+		/// <summary>
+		/// Wrapper for reading a chunk from the inner storage provider
+		/// </summary>
+		/// <param name="logId">The log file to read the index for</param>
+		/// <param name="offset">Offset of the chunk to read</param>
+		/// <param name="lineIndex">Index of the first line in this chunk</param>
+		/// <returns>The index data</returns>
+		async Task<LogChunkData?> InnerReadChunkAsync(LogId logId, long offset, int lineIndex)
+		{
+			await Task.Yield();
+
+			LogChunkData? chunkData = await _inner.ReadChunkAsync(logId, offset, lineIndex);
+			lock (_chunkReadTasks)
+			{
+				_chunkReadTasks.Remove((logId, offset));
+			}
+			return chunkData;
+		}
+
+		/// <summary>
+		/// Wrapper for reading a chunk from the inner storage provider
+		/// </summary>
+		/// <param name="logId">The log file to write the chunk for</param>
+		/// <param name="offset">Offset of the chunk within the log</param>
+		/// <param name="chunkData">The chunk data</param>
+		/// <returns>The index data</returns>
+		async Task<LogChunkData> InnerWriteChunkAsync(LogId logId, long offset, LogChunkData chunkData)
+		{
+			await Task.Yield();
+
+			await _inner.WriteChunkAsync(logId, offset, chunkData);
+
+			lock (_chunkWriteTasks)
+			{
+				_chunkWriteTasks.Remove((logId, offset));
+			}
+			return chunkData;
+		}
+	}
+}
