Horde: Finish re-org of Horde solution started a while ago.

diff --git a/Engine/Source/Programs/Horde/Horde.Build/Notifications/NotificationService.cs b/Engine/Source/Programs/Horde/Horde.Build/Notifications/NotificationService.cs
new file mode 100644
--- /dev/null
+++ b/Engine/Source/Programs/Horde/Horde.Build/Notifications/NotificationService.cs
@@ -0,0 +1,738 @@
+// Copyright Epic Games, Inc. All Rights Reserved.
+
+using System;
+using System.Collections.Concurrent;
+using System.Collections.Generic;
+using System.Linq;
+using System.Security.Claims;
+using System.Threading;
+using System.Threading.Tasks;
+using EpicGames.Core;
+using Horde.Build.Api;
+using Horde.Build.Collections;
+using Horde.Build.Issues.Impl;
+using Horde.Build.Models;
+using Horde.Build.Server;
+using Horde.Build.Services;
+using Horde.Build.Utilities;
+using HordeCommon;
+using Microsoft.Extensions.Hosting;
+using Microsoft.Extensions.Logging;
+using Microsoft.Extensions.Options;
+using MongoDB.Bson;
+using MongoDB.Driver;
+using StackExchange.Redis;
+using StatsdClient;
+using JsonSerializer = System.Text.Json.JsonSerializer;
+
+namespace Horde.Build.Notifications.Impl
+{
+	using UserId = ObjectId<IUser>;
+
+	/// <summary>
+	/// Wraps functionality for delivering notifications.
+	/// </summary>
+	public class NotificationService : BackgroundService, INotificationService
+	{
+		/// <summary>
+		/// The available notification sinks
+		/// </summary>
+		private readonly List<INotificationSink> _sinks;
+
+		/// <summary>
+		/// Collection of subscriptions
+		/// </summary>
+		private readonly ISubscriptionCollection _subscriptionCollection;
+
+		/// <summary>
+		/// Collection of notification request documents.
+		/// </summary>
+		private readonly INotificationTriggerCollection _triggerCollection;
+
+		/// <summary>
+		/// Instance of the <see cref="_graphCollection"/>.
+		/// </summary>
+		private readonly IGraphCollection _graphCollection;
+
+		/// <summary>
+		/// 
+		/// </summary>
+		private readonly IUserCollection _userCollection;
+
+		/// <summary>
+		/// Job service instance
+		/// </summary>
+		private readonly JobService _jobService;
+
+		/// <summary>
+		/// Instance of the <see cref="_streamService"/>.
+		/// </summary>
+		private readonly StreamService _streamService;
+
+		/// <summary>
+		/// 
+		/// </summary>
+		private readonly IssueService _issueService;
+
+		/// <summary>
+		/// Instance of the <see cref="_logFileService"/>.
+		/// </summary>
+		private readonly ILogFileService _logFileService;
+		
+		/// <summary>
+		/// Instance of the <see cref="IDogStatsd"/>.
+		/// </summary>
+		private readonly IDogStatsd _dogStatsd;
+
+		/// <summary>
+		/// Redis database
+		/// </summary>
+		private readonly IDatabase _redisDatabase;
+		
+		/// <summary>
+		/// Settings for the application.
+		/// </summary>
+		private readonly IOptionsMonitor<ServerSettings> _settings;
+
+		/// <summary>
+		/// List of asychronous tasks currently executing
+		/// </summary>
+		private readonly ConcurrentQueue<Task> _newTasks = new ConcurrentQueue<Task>();
+
+		/// <summary>
+		/// Set when there are new tasks to wait for
+		/// </summary>
+		private AsyncEvent _newTaskEvent = new AsyncEvent();
+
+		/// <summary>
+		/// Settings for the application.
+		/// </summary>
+		private readonly ILogger<NotificationService> _logger;
+		
+		/// <summary>
+		/// Ticker for running batch sender method
+		/// </summary>
+		internal ITicker _ticker;
+
+		/// <summary>
+		/// Interval at which queued notifications should be sent as a batch 
+		/// </summary>
+		internal TimeSpan _notificationBatchInterval = TimeSpan.FromHours(12);
+		
+		static string RedisQueueListKey(string notificationType) => "NotificationService.queued." + notificationType;
+
+		/// <summary>
+		/// Constructor
+		/// </summary>
+		public NotificationService(IEnumerable<INotificationSink> sinks, IOptionsMonitor<ServerSettings> settings, ILogger<NotificationService> logger, IGraphCollection graphCollection, ISubscriptionCollection subscriptionCollection, INotificationTriggerCollection triggerCollection, IUserCollection userCollection, JobService jobService, StreamService streamService, IssueService issueService, ILogFileService logFileService, IDogStatsd dogStatsd, RedisService redisService, IClock clock)
+		{
+			_sinks = sinks.ToList();
+			_settings = settings;
+			_logger = logger;
+			_graphCollection = graphCollection;
+			_subscriptionCollection = subscriptionCollection;
+			_triggerCollection = triggerCollection;
+			_userCollection = userCollection;
+			_jobService = jobService;
+			_streamService = streamService;
+			_issueService = issueService;
+			_logFileService = logFileService;
+			_dogStatsd = dogStatsd;
+			_redisDatabase = redisService.Database;
+
+			issueService.OnIssueUpdated += NotifyIssueUpdated;
+			jobService.OnJobStepComplete += NotifyJobStepComplete;
+			jobService.OnJobScheduled += NotifyJobScheduled;
+			jobService.OnLabelUpdate += NotifyLabelUpdate;
+			
+			_ticker = clock.AddSharedTicker<NotificationService>(_notificationBatchInterval, TickEveryTwelveHoursAsync, logger);
+		}
+
+		/// <inheritdoc/>
+		public override void Dispose()
+		{
+			base.Dispose();
+
+			_issueService.OnIssueUpdated -= NotifyIssueUpdated;
+			_jobService.OnJobStepComplete -= NotifyJobStepComplete;
+			_jobService.OnJobScheduled += NotifyJobScheduled;
+			_jobService.OnLabelUpdate -= NotifyLabelUpdate;
+
+			GC.SuppressFinalize(this);
+			_ticker.Dispose();
+		}
+
+		/// <inheritdoc/>
+		public async Task<bool> UpdateSubscriptionsAsync(ObjectId triggerId, ClaimsPrincipal user, bool? email, bool? slack)
+		{
+			UserId? userId = user.GetUserId();
+			if (userId == null)
+			{
+				_logger.LogWarning("Unable to find username for principal {User}", user.Identity?.Name);
+				return false;
+			}
+
+			INotificationTrigger trigger = await _triggerCollection.FindOrAddAsync(triggerId);
+			await _triggerCollection.UpdateSubscriptionsAsync(trigger, userId.Value, email, slack);
+			return true;
+		}
+
+		/// <inheritdoc/>
+		public async Task<INotificationSubscription?> GetSubscriptionsAsync(ObjectId triggerId, ClaimsPrincipal user)
+		{
+			UserId? userId = user.GetUserId();
+			if (userId == null)
+			{
+				return null;
+			}
+
+			INotificationTrigger? trigger = await _triggerCollection.GetAsync(triggerId);
+			if(trigger == null)
+			{
+				return null;
+			}
+
+			return trigger.Subscriptions.FirstOrDefault(x => x.UserId == userId.Value);
+		}
+
+		/// <inheritdoc/>
+		public void NotifyJobStepComplete(IJob job, IGraph graph, SubResourceId batchId, SubResourceId stepId)
+		{
+			// Enqueue job step complete notifications if needed
+			if (job.TryGetStep(batchId, stepId, out IJobStep? step))
+			{
+				_logger.LogInformation("Queuing step notifications for {JobId}:{BatchId}:{StepId}", job.Id, batchId, stepId);
+				EnqueueTask(() => SendJobStepNotificationsAsync(job, batchId, stepId));
+			}
+
+			// Enqueue job complete notifications if needed
+			if (job.GetState() == JobState.Complete)
+			{
+				_logger.LogInformation("Queuing job notifications for {JobId}:{BatchId}:{StepId}", job.Id, batchId, stepId);
+				EnqueueTask(() => SendJobNotificationsAsync(job, graph));
+				EnqueueTask(() => RecordJobCompleteMetrics(job));
+			}
+		}
+		
+		/// <inheritdoc/>
+		public void NotifyJobScheduled(IPool pool, bool poolHasAgentsOnline, IJob job, IGraph graph, SubResourceId batchId)
+		{
+			if (pool.EnableAutoscaling && !poolHasAgentsOnline)
+			{
+				EnqueueTasks(sink => EnqueueNotificationForBatchSending(new JobScheduledNotification(job.Id.ToString(), job.Name, pool.Name)));
+			}
+		}
+
+		/// <inheritdoc/>
+		public void NotifyLabelUpdate(IJob job, IReadOnlyList<(LabelState, LabelOutcome)> oldLabelStates, IReadOnlyList<(LabelState, LabelOutcome)> newLabelStates)
+		{
+			// If job has any label trigger IDs, send label complete notifications if needed
+			for (int idx = 0; idx < oldLabelStates.Count && idx < newLabelStates.Count; idx++)
+			{
+				if (oldLabelStates[idx] != newLabelStates[idx])
+				{
+					EnqueueTask(() => SendAllLabelNotificationsAsync(job, oldLabelStates, newLabelStates));
+					break;
+				}
+			}
+		}
+
+		/// <inheritdoc/>
+		public void NotifyIssueUpdated(IIssue issue)
+		{
+			_logger.LogInformation("Issue {IssueId} updated", issue.Id);
+			EnqueueTasks(sink => sink.NotifyIssueUpdatedAsync(issue));
+		}
+
+		/// <inheritdoc/>
+		public void NotifyConfigUpdateFailure(string errorMessage, string fileName, int? change = null, IUser? author = null, string? description = null)
+		{
+			EnqueueTasks(sink => sink.NotifyConfigUpdateFailureAsync(errorMessage, fileName, change, author, description));
+		}
+
+		/// <inheritdoc/>
+		public void NotifyDeviceService(string message, IDevice? device = null, IDevicePool? pool = null, IStream? stream = null, IJob? job = null, IJobStep? step = null, INode? node = null, IUser? user = null)
+		{
+			EnqueueTasks(sink => sink.NotifyDeviceServiceAsync(message, device, pool, stream, job, step, node, user));
+		}
+
+		/// <summary>
+		/// Enqueues an async task
+		/// </summary>
+		/// <param name="taskFunc">Function to generate an async task</param>
+		void EnqueueTask(Func<Task> taskFunc)
+		{
+			_newTasks.Enqueue(Task.Run(taskFunc));
+		}
+
+		/// <summary>
+		/// Enqueues an async task
+		/// </summary>
+		/// <param name="taskFunc">Function to generate an async task</param>
+		void EnqueueTasks(Func<INotificationSink, Task> taskFunc)
+		{
+			foreach (INotificationSink sink in _sinks)
+			{
+				EnqueueTask(() => taskFunc(sink));
+			}
+		}
+
+		/// <summary>
+		/// Enqueue a notification in Redis for batch sending later on 
+		/// </summary>
+		/// <param name="notification"></param>
+		/// <typeparam name="T">Any INotification type</typeparam>
+		private async Task EnqueueNotificationForBatchSending<T>(T notification) where T : INotification
+		{
+			try
+			{
+				byte[] data = JsonSerializer.SerializeToUtf8Bytes(notification);
+				await _redisDatabase.ListRightPushAsync(RedisQueueListKey(typeof(T).ToString()), data);
+			}
+			catch (Exception e)
+			{
+				_logger.LogError(e, "Unable to serialize and queue notification {Type} for batch sending in Redis", notification.GetType());
+			}
+		}
+
+		private async ValueTask TickEveryTwelveHoursAsync(CancellationToken stoppingToken)
+		{
+			List<JobScheduledNotification> jobScheduledNotifications = await GetAllQueuedNotificationsAsync<JobScheduledNotification>();
+
+			foreach (INotificationSink sink in _sinks)
+			{
+				await sink.NotifyJobScheduledAsync(jobScheduledNotifications);
+			}
+		}
+		
+		/// <summary>
+		/// Get and clear all queued notifications of type T from Redis
+		/// </summary>
+		/// <typeparam name="T"></typeparam>
+		/// <returns>Deserialized notifications</returns>
+		/// <exception cref="Exception"></exception>
+		private async Task<List<T>> GetAllQueuedNotificationsAsync<T>() where T : INotification
+		{
+			// Reading and deleting the entire list from Redis in this way is not thread-safe.
+			// But likely not a big deal considering the alternative of distributed locks.
+			string redisKey = RedisQueueListKey(typeof(T).ToString());
+			RedisValue[] rawNotifications = await _redisDatabase.ListRangeAsync(redisKey);
+			await _redisDatabase.KeyDeleteAsync(redisKey);
+
+			List<T> notifications = new List<T>();
+			foreach (byte[] data in rawNotifications)
+			{
+				try
+				{
+					T? notification = JsonSerializer.Deserialize<T>(data);
+					if (notification == null)
+					{
+						throw new Exception("Unable to deserialize");
+					}
+					notifications.Add(notification);
+				}
+				catch (Exception e)
+				{
+					_logger.LogError(e, "Unable to deserialize notification data of {Type}", typeof(T));
+				}
+			}
+
+			return notifications;
+		}
+
+		/// <inheritdoc/>
+		protected override async Task ExecuteAsync(CancellationToken stoppingToken)
+		{
+			await _ticker.StartAsync();
+			
+			// This background service just waits for tasks to finish and prints any exception info. The only reason to do this is to
+			// ensure we finish processing everything before shutdown.
+			using (CancellationTask stoppingTask = new CancellationTask(stoppingToken))
+			{
+				List<Task> tasks = new List<Task>();
+				tasks.Add(_newTaskEvent.Task);
+
+				for (; ; )
+				{
+					// Add any new tasks to be monitored
+					Task? newTask;
+					while (_newTasks.TryDequeue(out newTask))
+					{
+						tasks.Add(newTask);
+					}
+
+					// If we don't have any
+					if (tasks.Count == 1)
+					{
+						await Task.WhenAny(_newTaskEvent.Task, stoppingTask.Task);
+						if (stoppingToken.IsCancellationRequested)
+						{
+							break;
+						}
+					}
+					else
+					{
+						// Wait for something to finish
+						Task task = await Task.WhenAny(tasks);
+						if (task == _newTaskEvent.Task)
+						{
+							_newTaskEvent = new AsyncEvent();
+							tasks[0] = _newTaskEvent.Task;
+						}
+						else
+						{
+							try
+							{
+								await task;
+							}
+							catch (Exception ex)
+							{
+								_logger.LogError(ex, "Exception while executing notification");
+							}
+							tasks.Remove(task);
+						}
+					}
+				}
+			}
+			
+			await _ticker.StopAsync();
+		}
+
+		internal Task ExecuteBackgroundForTest(CancellationToken stoppingToken)
+		{
+			return ExecuteAsync(stoppingToken);
+		}
+
+		/// <summary>
+		/// Gets the <see cref="INotificationTrigger"/> for a given trigger ID, if any.
+		/// </summary>
+		/// <param name="triggerId"></param>
+		/// <param name="bFireTrigger">If true, the trigger is fired and cannot be reused</param>
+		/// <returns></returns>
+		private async Task<INotificationTrigger?> GetNotificationTrigger(ObjectId? triggerId, bool bFireTrigger)
+		{
+			if (triggerId == null)
+			{
+				return null;
+			}
+
+			INotificationTrigger? trigger = await _triggerCollection.GetAsync(triggerId.Value);
+			if (trigger == null)
+			{
+				return null;
+			}
+
+			return bFireTrigger ? await _triggerCollection.FireAsync(trigger) : trigger;
+		}
+	
+		private async Task SendJobNotificationsAsync(IJob job, IGraph graph)
+		{
+			using IDisposable scope = _logger.BeginScope("Sending notifications for job {JobId}", job.Id);
+
+			job.GetJobState(job.GetStepForNodeMap(), out _, out LabelOutcome outcome);
+			JobCompleteEventRecord jobCompleteEvent = new JobCompleteEventRecord(job.StreamId, job.TemplateId, outcome);
+
+			IStream? jobStream = await _streamService.GetStreamAsync(job.StreamId);
+			if (jobStream == null)
+			{
+				_logger.LogError("Unable to get stream {StreamId}", job.StreamId);
+				return;
+			}
+
+			List<IUser> usersToNotify = await GetUsersToNotify(jobCompleteEvent, job.NotificationTriggerId, true);
+			foreach (IUser userToNotify in usersToNotify)
+			{
+				if(job.PreflightChange != 0)
+				{
+					if(userToNotify.Id != job.StartedByUserId)
+					{
+						continue;
+					}
+				}
+				EnqueueTasks(sink => sink.NotifyJobCompleteAsync(userToNotify, jobStream, job, graph, outcome));
+			}
+
+			if (job.PreflightChange == 0)
+			{
+				EnqueueTasks(sink => sink.NotifyJobCompleteAsync(jobStream, job, graph, outcome));
+			}
+
+			_logger.LogDebug("Finished sending notifications for job {JobId}", job.Id);
+		}
+
+		private Task RecordJobCompleteMetrics(IJob job)
+		{
+			void RecordMetric(string type, JobStepOutcome outcome, DateTimeOffset? startTime, DateTimeOffset? finishTime)
+			{
+				string outcomeStr = outcome switch
+				{
+					JobStepOutcome.Unspecified => "unspecified",
+					JobStepOutcome.Failure => "failure",
+					JobStepOutcome.Warnings => "warnings",
+					JobStepOutcome.Success => "success",
+					_ => "unspecified"
+				};
+				
+				string[] tags = {"stream:" + job.StreamId, "template:" + job.TemplateId};
+				_dogStatsd.Increment($"horde.{type}.{outcomeStr}.count", 1, tags: tags);
+
+				if (startTime == null || finishTime == null)
+				{
+					_logger.LogDebug("Completed job or step is missing start or finish time, cannot record duration metric. Job ID={JobId}", job.Id);
+					return;
+				}
+
+				TimeSpan duration = finishTime.Value - startTime.Value;
+				_dogStatsd.Timer($"horde.{type}.{outcomeStr}.duration", duration.TotalSeconds, tags: tags);
+			}
+
+			JobStepOutcome jobOutcome = job.Batches.SelectMany(x => x.Steps).Min(x => x.Outcome);
+			DateTime? startTime = job.Batches.Select(x => x.StartTimeUtc).Min();
+			DateTime? finishTime = job.Batches.Select(x => x.FinishTimeUtc).Max();
+			RecordMetric("job", jobOutcome, startTime, finishTime);
+
+			// TODO: record metrics for individual steps
+			// foreach (IJobStepBatch Batch in Job.Batches)
+			// {
+			// 	foreach (IJobStep Step in Batch.Steps)
+			// 	{
+			// 	}
+			// }
+
+			return Task.CompletedTask;
+		}
+
+		private async Task<List<IUser>> GetUsersToNotify(EventRecord? eventRecord, ObjectId? notificationTriggerId, bool bFireTrigger)
+		{
+			List<UserId> userIds = new List<UserId>();
+
+			// Find the notifications for all steps of this type
+			if (eventRecord != null)
+			{
+				List<ISubscription> subscriptions = await _subscriptionCollection.FindSubscribersAsync(eventRecord);
+				foreach (ISubscription subscription in subscriptions)
+				{
+					if (subscription.NotificationType == NotificationType.Slack)
+					{
+						userIds.Add(subscription.UserId);
+					}
+				}
+			}
+
+			// Find the notifications for this particular step
+			if (notificationTriggerId != null)
+			{
+				INotificationTrigger? trigger = await GetNotificationTrigger(notificationTriggerId, bFireTrigger);
+				if (trigger != null)
+				{
+					foreach (INotificationSubscription subscription in trigger.Subscriptions)
+					{
+						if (subscription.Email)
+						{
+							// TODO?
+						}
+						if (subscription.Slack)
+						{
+							userIds.Add(subscription.UserId);
+						}
+					}
+				}
+			}
+			return await _userCollection.FindUsersAsync(userIds);
+		}
+
+		private async Task SendJobStepNotificationsAsync(IJob job, SubResourceId batchId, SubResourceId stepId)
+		{
+			using IDisposable scope = _logger.BeginScope("Sending notifications for step {JobId}:{BatchId}:{StepId}", job.Id, batchId, stepId);
+
+			IJobStepBatch? batch;
+			if(!job.TryGetBatch(batchId, out batch))
+			{
+				_logger.LogError("Unable to find batch {BatchId} in job {JobId}", batchId, job.Id);
+				return;
+			}
+
+			IJobStep? step;
+			if (!batch.TryGetStep(stepId, out step))
+			{
+				_logger.LogError("Unable to find step {StepId} in batch {JobId}:{BatchId}", stepId, job.Id, batchId);
+				return;
+			}
+
+			IGraph jobGraph = await _graphCollection.GetAsync(job.GraphHash);
+			INode node = jobGraph.GetNode(new NodeRef(batch.GroupIdx, step.NodeIdx));
+
+			// Find the notifications for this particular step
+			EventRecord eventRecord = new StepCompleteEventRecord(job.StreamId, job.TemplateId, node.Name, step.Outcome);
+
+			List<IUser> usersToNotify = await GetUsersToNotify(eventRecord, step.NotificationTriggerId, true);
+
+			// If this is not a success notification and the author isn't in the list to notify, add them manually if this is the outcome has gotten worse.
+			int failures = job.Batches.Sum(x => x.Steps.Count(y => y.Outcome == JobStepOutcome.Failure));
+			bool firstFailure = step.Outcome == JobStepOutcome.Failure && failures == 1;
+			bool firstWarning = step.Outcome == JobStepOutcome.Warnings && failures == 0 && job.Batches.Sum(x => x.Steps.Count(y => y.Outcome == JobStepOutcome.Warnings)) == 1;
+			if (job.StartedByUserId.HasValue && !usersToNotify.Any(x => x.Id == job.StartedByUserId) && (firstFailure || firstWarning))
+			{
+				_logger.LogInformation("Author {AuthorUserId} is not in notify list but step outcome is {JobStepOutcome}, adding them to the list...", job.StartedByUserId, step.Outcome);
+				IUser? authorUser = await _userCollection.GetUserAsync(job.StartedByUserId.Value);
+				if (authorUser != null)
+				{
+					usersToNotify.Add(authorUser);
+				}
+			}
+
+			if (usersToNotify.Count == 0)
+			{
+				_logger.LogInformation("No users to notify for step {JobId}:{BatchId}:{StepId}", job.Id, batchId, stepId);
+				return;
+			}
+
+			IStream? jobStream = await _streamService.GetStreamAsync(job.StreamId);
+			if (jobStream == null)
+			{
+				_logger.LogError("Unable to find stream {StreamId}", job.StreamId);
+				return;
+			}
+
+			if(step.LogId == null)
+			{
+				_logger.LogError("Step does not have a log file");
+				return;
+			}
+
+			ILogFile? logFile = await _logFileService.GetLogFileAsync(step.LogId.Value);
+			if(logFile == null)
+			{
+				_logger.LogError("Step does not have a log file");
+				return;
+			}
+
+			List<ILogEvent> jobStepEvents = await _logFileService.FindEventsAsync(logFile);
+			List<ILogEventData> jobStepEventData = new List<ILogEventData>();
+			foreach (ILogEvent logEvent in jobStepEvents)
+			{
+				ILogEventData eventData = await _logFileService.GetEventDataAsync(logFile, logEvent.LineIndex, logEvent.LineCount);
+				jobStepEventData.Add(eventData);
+			}
+
+			foreach (IUser slackUser in usersToNotify)
+			{
+				if(job.PreflightChange != 0)
+				{
+					if(slackUser.Id != job.StartedByUserId)
+					{
+						continue;
+					}
+				}
+				EnqueueTasks(sink => sink.NotifyJobStepCompleteAsync(slackUser, jobStream, job, batch, step, node, jobStepEventData));
+			}
+			_logger.LogDebug("Finished sending notifications for step {JobId}:{BatchId}:{StepId}", job.Id, batchId, stepId);
+		}
+
+		private async Task SendAllLabelNotificationsAsync(IJob job, IReadOnlyList<(LabelState State, LabelOutcome Outcome)> oldLabelStates, IReadOnlyList<(LabelState, LabelOutcome)> newLabelStates)
+		{
+			IStream? stream = await _streamService.GetStreamAsync(job.StreamId);
+			if (stream == null)
+			{
+				_logger.LogError("Unable to find stream {StreamId} for job {JobId}", job.StreamId, job.Id);
+				return;
+			}
+
+			IGraph? graph = await _graphCollection.GetAsync(job.GraphHash);
+			if (graph == null)
+			{
+				_logger.LogError("Unable to find graph {GraphHash} for job {JobId}", job.GraphHash, job.Id);
+				return;
+			}
+
+			IReadOnlyDictionary<NodeRef, IJobStep> stepForNode = job.GetStepForNodeMap();
+			for (int labelIdx = 0; labelIdx < graph.Labels.Count; ++labelIdx)
+			{
+				(LabelState State, LabelOutcome Outcome) oldLabel = oldLabelStates[labelIdx];
+				(LabelState State, LabelOutcome Outcome) newLabel = newLabelStates[labelIdx];
+				if (oldLabel != newLabel)
+				{
+					// If the state transitioned from Unspecified to Running, don't update unless the outcome also changed.
+					if (oldLabel.State == LabelState.Unspecified && newLabel.State == LabelState.Running && oldLabel.Outcome == newLabel.Outcome)
+					{
+						continue;
+					}
+
+					// If the label isn't complete, don't report on outcome changing to success, this will be reported when the label state becomes complete.
+					if (newLabel.State != LabelState.Complete && newLabel.Outcome == LabelOutcome.Success)
+					{
+						return;
+					}
+
+					ILabel label = graph.Labels[labelIdx];
+
+					EventRecord? eventId;
+					if (String.IsNullOrEmpty(label.DashboardName))
+					{
+						eventId = null;
+					}
+					else
+					{
+						eventId = new LabelCompleteEventRecord(job.StreamId, job.TemplateId, label.DashboardCategory, label.DashboardName, newLabel.Outcome);
+					}
+
+					ObjectId? triggerId;
+					if (job.LabelIdxToTriggerId.TryGetValue(labelIdx, out ObjectId newTriggerId))
+					{
+						triggerId = newTriggerId;
+					}
+					else
+					{
+						triggerId = null;
+					}
+
+					bool bFireTrigger = newLabel.State == LabelState.Complete;
+
+					List<IUser> usersToNotify = await GetUsersToNotify(eventId, triggerId, bFireTrigger);
+
+					// filter preflight label notifications to only include initiator
+					if (usersToNotify.Count > 0 && job.PreflightChange != 0 && job.StartedByUserId != null)
+					{
+						usersToNotify = usersToNotify.Where(x => x.Id == job.StartedByUserId).ToList();
+					}
+
+					if (usersToNotify.Count > 0)
+					{
+						SendLabelUpdateNotifications(job, stream, graph, stepForNode, graph.Labels[labelIdx], labelIdx, newLabel.Outcome, usersToNotify);
+					}
+					else
+					{
+						_logger.LogDebug("No users to notify for label {DashboardName}/{UgsName} in job {JobId}", graph.Labels[labelIdx].DashboardName, graph.Labels[labelIdx].UgsName, job.Id);
+					}
+				}
+			}
+		}
+
+		private void SendLabelUpdateNotifications(IJob job, IStream stream, IGraph graph, IReadOnlyDictionary<NodeRef, IJobStep> stepForNode, ILabel label, int labelIdx, LabelOutcome outcome, List<IUser> slackUsers)
+		{
+			List<(string, JobStepOutcome, Uri)> stepData = new List<(string, JobStepOutcome, Uri)>();
+			if (outcome != LabelOutcome.Success)
+			{
+				foreach (NodeRef includedNodeRef in label.IncludedNodes)
+				{
+					INode includedNode = graph.GetNode(includedNodeRef);
+					IJobStep includedStep = stepForNode[includedNodeRef];
+					stepData.Add((includedNode.Name, includedStep.Outcome, new Uri($"{_settings.CurrentValue.DashboardUrl}job/{job.Id}?step={includedStep.Id}")));
+				}
+			}
+
+			foreach (IUser slackUser in slackUsers)
+			{
+				EnqueueTasks(sink => sink.NotifyLabelCompleteAsync(slackUser, job, stream, label, labelIdx, outcome, stepData));
+			}
+
+			_logger.LogDebug("Finished sending label notifications for label {DashboardName}/{UgsName} in job {JobId}", label.DashboardName, label.UgsName, job.Id);
+		}
+
+		/// <inheritdoc/>
+		public void SendIssueReport(IssueReport report)
+		{
+			EnqueueTasks(sink => sink.SendIssueReportAsync(report));
+		}
+	}
+}
diff --git a/Engine/Source/Programs/Horde/Horde.Build/Notifications/NotificationTriggerCollection.cs b/Engine/Source/Programs/Horde/Horde.Build/Notifications/NotificationTriggerCollection.cs
new file mode 100644
--- /dev/null
+++ b/Engine/Source/Programs/Horde/Horde.Build/Notifications/NotificationTriggerCollection.cs
@@ -0,0 +1,243 @@
+// Copyright Epic Games, Inc. All Rights Reserved.
+
+using System.Collections.Generic;
+using System.Linq;
+using System.Threading.Tasks;
+using Horde.Build.Models;
+using Horde.Build.Server;
+using Horde.Build.Utilities;
+using MongoDB.Bson;
+using MongoDB.Bson.Serialization.Attributes;
+using MongoDB.Driver;
+
+namespace Horde.Build.Collections.Impl
+{
+	using UserId = ObjectId<IUser>;
+
+	/// <summary>
+	/// Collection of notification triggers
+	/// </summary>
+	public class NotificationTriggerCollection : INotificationTriggerCollection
+	{
+		class SubscriptionDocument : INotificationSubscription
+		{
+			[BsonIgnoreIfNull]
+			public string? User { get; set; }
+
+			public UserId UserId { get; set; }
+			public bool Email { get; set; }
+			public bool Slack { get; set; }
+
+			[BsonConstructor]
+			private SubscriptionDocument()
+			{
+				User = null!;
+			}
+
+			public SubscriptionDocument(UserId userId, bool email, bool slack)
+			{
+				UserId = userId;
+				Email = email;
+				Slack = slack;
+			}
+
+			public SubscriptionDocument(INotificationSubscription subscription)
+			{
+				UserId = subscription.UserId;
+				Email = subscription.Email;
+				Slack = subscription.Slack;
+			}
+		}
+
+		class TriggerDocument : INotificationTrigger
+		{
+			public ObjectId Id { get; set; }
+			public bool Fired { get; set; }
+			public List<SubscriptionDocument> Subscriptions { get; set; } = new List<SubscriptionDocument>();
+			public int UpdateIndex { get; set; }
+
+			IReadOnlyList<INotificationSubscription> INotificationTrigger.Subscriptions => Subscriptions;
+		}
+
+		readonly IMongoCollection<TriggerDocument> _triggers;
+		readonly IUserCollection _userCollection;
+
+		/// <summary>
+		/// Constructor
+		/// </summary>
+		/// <param name="mongoService">The database singleton</param>
+		/// <param name="userCollection"></param>
+		public NotificationTriggerCollection(MongoService mongoService, IUserCollection userCollection)
+		{
+			_triggers = mongoService.GetCollection<TriggerDocument>("NotificationTriggers");
+			_userCollection = userCollection;
+		}
+
+		/// <inheritdoc/>
+		public async Task<INotificationTrigger?> GetAsync(ObjectId triggerId)
+		{
+			TriggerDocument? trigger = await _triggers.Find(x => x.Id == triggerId).FirstOrDefaultAsync();
+			if (trigger != null)
+			{
+				for (int idx = 0; idx < trigger.Subscriptions.Count; idx++)
+				{
+					SubscriptionDocument subscription = trigger.Subscriptions[idx];
+					if (subscription.User != null)
+					{
+						IUser? user = await _userCollection.FindUserByLoginAsync(subscription.User);
+						if (user == null)
+						{
+							trigger.Subscriptions.RemoveAt(idx);
+							idx--;
+						}
+						else
+						{
+							subscription.UserId = user.Id;
+							subscription.User = null;
+						}
+					}
+				}
+			}
+			return trigger;
+		}
+
+		/// <inheritdoc/>
+		public async Task<INotificationTrigger> FindOrAddAsync(ObjectId triggerId)
+		{
+			for (; ; )
+			{
+				// Find an existing trigger
+				INotificationTrigger? existing = await GetAsync(triggerId);
+				if (existing != null)
+				{
+					return existing;
+				}
+
+				// Try to insert a new document
+				try
+				{
+					TriggerDocument newDocument = new TriggerDocument();
+					newDocument.Id = triggerId;
+					await _triggers.InsertOneAsync(newDocument);
+					return newDocument;
+				}
+				catch (MongoWriteException ex)
+				{
+					if (ex.WriteError.Category != ServerErrorCategory.DuplicateKey)
+					{
+						throw;
+					}
+				}
+			}
+		}
+
+		/// <summary>
+		/// Updates an existing document
+		/// </summary>
+		/// <param name="trigger">The trigger to update</param>
+		/// <param name="transaction">The update definition</param>
+		/// <returns>The updated document</returns>
+		async Task<INotificationTrigger?> TryUpdateAsync(INotificationTrigger trigger, TransactionBuilder<TriggerDocument> transaction)
+		{
+			TriggerDocument document = (TriggerDocument)trigger;
+			int nextUpdateIndex = document.UpdateIndex + 1;
+
+			FilterDefinition<TriggerDocument> filter = Builders<TriggerDocument>.Filter.Expr(x => x.Id == trigger.Id && x.UpdateIndex == document.UpdateIndex);
+			UpdateDefinition<TriggerDocument> update = transaction.ToUpdateDefinition().Set(x => x.UpdateIndex, nextUpdateIndex);
+
+			UpdateResult result = await _triggers.UpdateOneAsync(filter, update);
+			if (result.ModifiedCount > 0)
+			{
+				transaction.ApplyTo(document);
+				document.UpdateIndex = nextUpdateIndex;
+				return document;
+			}
+
+			return null;
+		}
+
+		/// <inheritdoc/>
+		public async Task DeleteAsync(ObjectId triggerId)
+		{
+			await _triggers.DeleteOneAsync(x => x.Id == triggerId);
+		}
+
+		/// <inheritdoc/>
+		public async Task DeleteAsync(List<ObjectId> triggerIds)
+		{
+			FilterDefinition<TriggerDocument> filter = Builders<TriggerDocument>.Filter.In(x => x.Id, triggerIds);
+			await _triggers.DeleteManyAsync(filter);
+		}
+
+		/// <inheritdoc/>
+		public async Task<INotificationTrigger?> FireAsync(INotificationTrigger trigger)
+		{
+			if (trigger.Fired)
+			{
+				return null;
+			}
+
+			for (; ; )
+			{
+				TransactionBuilder<TriggerDocument> transaction = new TransactionBuilder<TriggerDocument>();
+				transaction.Set(x => x.Fired, true);
+
+				INotificationTrigger? newTrigger = await TryUpdateAsync(trigger, transaction);
+				if (newTrigger != null)
+				{
+					return newTrigger;
+				}
+
+				newTrigger = await FindOrAddAsync(trigger.Id); // Need to add to prevent race condition on triggering vs adding
+				if (newTrigger == null || newTrigger.Fired)
+				{
+					return null;
+				}
+			}
+		}
+
+		/// <inheritdoc/>
+		public async Task<INotificationTrigger?> UpdateSubscriptionsAsync(INotificationTrigger trigger, UserId userId, bool? email, bool? slack)
+		{
+			for (; ; )
+			{
+				// If the trigger has already fired, don't add a new subscription to it
+				if(trigger.Fired)
+				{
+					return trigger;
+				}
+
+				// Try to update the trigger
+				List<SubscriptionDocument> newSubscriptions = new List<SubscriptionDocument>();
+				newSubscriptions.AddRange(trigger.Subscriptions.Select(x => new SubscriptionDocument(x)));
+
+				SubscriptionDocument? newSubscription = newSubscriptions.FirstOrDefault(x => x.UserId == userId);
+				if (newSubscription == null)
+				{
+					newSubscription = new SubscriptionDocument(userId, email ?? false, slack ?? false);
+					newSubscriptions.Add(newSubscription);
+				}
+				else
+				{
+					newSubscription.Email = email ?? newSubscription.Email;
+					newSubscription.Slack = slack ?? newSubscription.Slack;
+				}
+
+				TransactionBuilder<TriggerDocument> transaction = new TransactionBuilder<TriggerDocument>();
+				transaction.Set(x => x.Subscriptions, newSubscriptions);
+
+				INotificationTrigger? newTrigger = await TryUpdateAsync(trigger, transaction);
+				if (newTrigger != null)
+				{
+					return newTrigger;
+				}
+
+				newTrigger = await GetAsync(trigger.Id);
+				if (newTrigger == null)
+				{
+					return null;
+				}
+			}
+		}
+	}
+}
