Horde: Finish re-org of Horde solution started a while ago.

diff --git a/Engine/Source/Programs/Horde/Horde.Build/Notifications/Sinks/SlackNotificationSink.cs b/Engine/Source/Programs/Horde/Horde.Build/Notifications/Sinks/SlackNotificationSink.cs
new file mode 100644
--- /dev/null
+++ b/Engine/Source/Programs/Horde/Horde.Build/Notifications/Sinks/SlackNotificationSink.cs
@@ -0,0 +1,1962 @@
+// Copyright Epic Games, Inc. All Rights Reserved.
+
+using System;
+using System.Collections.Generic;
+using System.Diagnostics.CodeAnalysis;
+using System.Drawing;
+using System.Globalization;
+using System.Linq;
+using System.Net.Http;
+using System.Net.WebSockets;
+using System.Text;
+using System.Text.Json;
+using System.Text.Json.Serialization;
+using System.Text.RegularExpressions;
+using System.Threading;
+using System.Threading.Tasks;
+using EpicGames.Core;
+using Horde.Build.Api;
+using Horde.Build.Collections;
+using Horde.Build.Config;
+using Horde.Build.Issues.Impl;
+using Horde.Build.Models;
+using Horde.Build.Server;
+using Horde.Build.Services;
+using Horde.Build.Utilities;
+using Horde.Build.Utilities.Slack.BlockKit;
+using HordeCommon;
+using Microsoft.AspNetCore.Hosting;
+using Microsoft.Extensions.Caching.Memory;
+using Microsoft.Extensions.Hosting;
+using Microsoft.Extensions.Logging;
+using Microsoft.Extensions.Options;
+using MongoDB.Bson;
+using MongoDB.Bson.Serialization.Attributes;
+using MongoDB.Driver;
+
+namespace Horde.Build.Notifications.Impl
+{
+	using JobId = ObjectId<IJob>;
+	using LogId = ObjectId<ILogFile>;
+	using StreamId = StringId<IStream>;
+	using UserId = ObjectId<IUser>;
+	using WorkflowId = StringId<WorkflowConfig>;
+
+	/// <summary>
+	/// Maintains a connection to Slack, in order to receive socket-mode notifications of user interactions
+	/// </summary>
+	public class SlackNotificationSink : BackgroundService, INotificationSink, IAvatarService
+	{
+		const string AddReactionUrl = "https://slack.com/api/reactions.add";
+		const string RemoveReactionUrl = "https://slack.com/api/reactions.remove";
+		const string PostMessageUrl = "https://slack.com/api/chat.postMessage";
+		const string UpdateMessageUrl = "https://slack.com/api/chat.update";
+		const string GetPermalinkUrl = "https://slack.com/api/chat.getPermalink";
+
+		const bool defaultAllowMentions = true;
+
+		class SocketResponse
+		{
+			[JsonPropertyName("ok")]
+			public bool Ok { get; set; }
+
+			[JsonPropertyName("url")]
+			public Uri? Url { get; set; }
+		}
+
+		class EventMessage
+		{
+			[JsonPropertyName("envelope_id")]
+			public string? EnvelopeId { get; set; }
+
+			[JsonPropertyName("type")]
+			public string? Type { get; set; }
+
+			[JsonPropertyName("payload")]
+			public EventPayload? Payload { get; set; }
+		}
+
+		class ReactionMessage
+		{
+			[JsonPropertyName("channel")]
+			public string? Channel { get; set; }
+
+			[JsonPropertyName("timestamp")]
+			public string? Ts { get; set; }
+
+			[JsonPropertyName("name")]
+			public string? Name { get; set; }
+		}
+
+		class EventPayload
+		{
+			[JsonPropertyName("type")]
+			public string? Type { get; set; }
+
+			[JsonPropertyName("user")]
+			public UserInfo? User { get; set; }
+
+			[JsonPropertyName("response_url")]
+			public string? ResponseUrl { get; set; }
+
+			[JsonPropertyName("actions")]
+			public List<ActionInfo> Actions { get; set; } = new List<ActionInfo>();
+		}
+
+		class UserResponse
+		{
+			[JsonPropertyName("ok")]
+			public bool Ok { get; set; }
+
+			[JsonPropertyName("user")]
+			public UserInfo? User { get; set; }
+		}
+
+		class UserInfo
+		{
+			[JsonPropertyName("id")]
+			public string? UserName { get; set; }
+
+			[JsonPropertyName("profile")]
+			public UserProfile? Profile { get; set; }
+		}
+
+		class UserProfile
+		{
+			[JsonPropertyName("is_custom_image")]
+			public bool IsCustomImage { get; set; }
+
+			[JsonPropertyName("image_24")]
+			public string? Image24 { get; set; }
+
+			[JsonPropertyName("image_32")]
+			public string? Image32 { get; set; }
+
+			[JsonPropertyName("image_48")]
+			public string? Image48 { get; set; }
+
+			[JsonPropertyName("image_72")]
+			public string? Image72 { get; set; }
+		}
+
+		class ActionInfo
+		{
+			[JsonPropertyName("value")]
+			public string? Value { get; set; }
+		}
+
+		class MessageStateDocument
+		{
+			[BsonId]
+			public ObjectId Id { get; set; }
+
+			[BsonElement("uid")]
+			public string Recipient { get; set; } = String.Empty;
+
+			[BsonElement("usr")]
+			public UserId? UserId { get; set; }
+
+			[BsonElement("eid")]
+			public string EventId { get; set; } = String.Empty;
+
+			[BsonElement("ch")]
+			public string Channel { get; set; } = String.Empty;
+
+			[BsonElement("ts")]
+			public string Ts { get; set; } = String.Empty;
+
+			[BsonElement("dig")]
+			public string Digest { get; set; } = String.Empty;
+
+			[BsonElement("lnk"), BsonIgnoreIfNull]
+			public string? Permalink { get; set; }
+		}
+
+		class SlackUser : IAvatar
+		{
+			public const int CurrentVersion = 2;
+			
+			public UserId Id { get; set; }
+
+			[BsonElement("u")]
+			public string? SlackUserId { get; set; }
+
+			[BsonElement("i24")]
+			public string? Image24 { get; set; }
+
+			[BsonElement("i32")]
+			public string? Image32 { get; set; }
+			
+			[BsonElement("i48")]
+			public string? Image48 { get; set; }
+
+			[BsonElement("i72")]
+			public string? Image72 { get; set; }
+
+			[BsonElement("t")]
+			public DateTime Time { get; set; }
+
+			[BsonElement("v")]
+			public int Version { get; set; }
+
+			private SlackUser()
+			{
+			}
+
+			public SlackUser(UserId id, UserInfo? info)
+			{
+				Id = id;
+				SlackUserId = info?.UserName;
+				if (info != null && info.Profile != null && info.Profile.IsCustomImage)
+				{
+					Image24 = info.Profile.Image24;
+					Image32 = info.Profile.Image32;
+					Image48 = info.Profile.Image48;
+					Image72 = info.Profile.Image72;
+				}
+				Time = DateTime.UtcNow;
+				Version = CurrentVersion;
+			}
+		}
+
+		readonly IssueService _issueService;
+		readonly IUserCollection _userCollection;
+		readonly ILogFileService _logFileService;
+		readonly StreamService _streamService;
+		readonly ConfigCollection _configCollection;
+		readonly IWebHostEnvironment _environment;
+		readonly ServerSettings _settings;
+		readonly IMongoCollection<MessageStateDocument> _messageStates;
+		readonly IMongoCollection<SlackUser> _slackUsers;
+		readonly HashSet<string>? _allowUsers;
+		readonly IExternalIssueService _externalIssueService;
+		readonly ILogger _logger;
+
+		/// <summary>
+		/// Map of email address to Slack user ID.
+		/// </summary>
+		private readonly MemoryCache _userCache = new MemoryCache(new MemoryCacheOptions());
+
+		/// <summary>
+		/// Constructor
+		/// </summary>
+		public SlackNotificationSink(MongoService mongoService, IssueService issueService, IUserCollection userCollection, ILogFileService logFileService, StreamService streamService, ConfigCollection configCollection, IExternalIssueService externalIssueService, IWebHostEnvironment environment, IOptions<ServerSettings> settings, ILogger<SlackNotificationSink> logger)
+		{
+			_issueService = issueService;
+			_userCollection = userCollection;
+			_logFileService = logFileService;
+			_streamService = streamService;
+			_configCollection = configCollection;
+			_externalIssueService = externalIssueService;
+			_environment = environment;
+			_settings = settings.Value;
+			_messageStates = mongoService.Database.GetCollection<MessageStateDocument>("Slack");
+			_slackUsers = mongoService.Database.GetCollection<SlackUser>("Slack.UsersV2");
+			_logger = logger;
+
+			if (!String.IsNullOrEmpty(settings.Value.SlackUsers))
+			{
+				_allowUsers = new HashSet<string>(settings.Value.SlackUsers.Split(','), StringComparer.OrdinalIgnoreCase);
+			}
+		}
+
+		/// <inheritdoc/>
+		public override void Dispose()
+		{
+			base.Dispose();
+
+			_userCache.Dispose();
+
+			GC.SuppressFinalize(this);
+		}
+
+		#region Avatars
+
+		/// <inheritdoc/>
+		public async Task<IAvatar?> GetAvatarAsync(IUser user)
+		{
+			return await GetSlackUser(user);
+		}
+
+		#endregion
+
+		#region Message state 
+
+		async Task<(MessageStateDocument, bool)> AddOrUpdateMessageStateAsync(string recipient, string eventId, UserId? userId, string digest)
+		{
+			ObjectId newId = ObjectId.GenerateNewId();
+
+			FilterDefinition<MessageStateDocument> filter = Builders<MessageStateDocument>.Filter.Eq(x => x.Recipient, recipient) & Builders<MessageStateDocument>.Filter.Eq(x => x.EventId, eventId);
+			UpdateDefinition<MessageStateDocument> update = Builders<MessageStateDocument>.Update.SetOnInsert(x => x.Id, newId).Set(x => x.UserId, userId).Set(x => x.Digest, digest);
+
+			MessageStateDocument state = await _messageStates.FindOneAndUpdateAsync(filter, update, new FindOneAndUpdateOptions<MessageStateDocument> { IsUpsert = true, ReturnDocument = ReturnDocument.After });
+			return (state, state.Id == newId);
+		}
+
+		async Task<MessageStateDocument?> GetMessageStateAsync(string recipient, string eventId)
+		{
+			FilterDefinition<MessageStateDocument> filter = Builders<MessageStateDocument>.Filter.Eq(x => x.Recipient, recipient) & Builders<MessageStateDocument>.Filter.Eq(x => x.EventId, eventId);
+			return await _messageStates.Find(filter).FirstOrDefaultAsync();
+		}
+
+		async Task SetMessageTimestampAsync(ObjectId messageId, string channel, string ts, string? permalink = null)
+		{
+			FilterDefinition<MessageStateDocument> filter = Builders<MessageStateDocument>.Filter.Eq(x => x.Id, messageId);
+			UpdateDefinition<MessageStateDocument> update = Builders<MessageStateDocument>.Update.Set(x => x.Channel, channel).Set(x => x.Ts, ts);
+			if (permalink == null)
+			{
+				update = update.Unset(x => x.Permalink);
+			}
+			else
+			{
+				update = update.Set(x => x.Permalink, permalink);
+			}
+			await _messageStates.FindOneAndUpdateAsync(filter, update);
+		}
+
+		#endregion
+
+		/// <inheritdoc/>
+		public async Task NotifyJobScheduledAsync(List<JobScheduledNotification> notifications)
+		{
+			if (_settings.JobNotificationChannel != null)
+			{
+				string jobIds = String.Join(", ", notifications.Select(x => x.JobId));
+				_logger.LogInformation("Sending Slack notification for scheduled job IDs {JobIds} to channel {SlackChannel}", jobIds, _settings.JobNotificationChannel);
+				await SendJobScheduledOnEmptyAutoScaledPoolMessageAsync($"#{_settings.JobNotificationChannel}", notifications);
+			}
+		}
+		
+		private async Task SendJobScheduledOnEmptyAutoScaledPoolMessageAsync(string recipient, List<JobScheduledNotification> notifications)
+		{
+			string jobIds = String.Join(", ", notifications.Select(x => x.JobId));
+				
+			StringBuilder sb = new();
+			foreach (JobScheduledNotification notification in notifications)
+			{
+				string jobUrl = _settings.DashboardUrl + "/job/" + notification.JobId;
+				sb.AppendLine($"Job `{notification.JobName}` with ID <{jobUrl}|{notification.JobId}> in pool `{notification.PoolName}`");
+			}
+			
+			Color outcomeColor = BlockKitAttachmentColors.Warning;
+			BlockKitAttachment attachment = new BlockKitAttachment();
+			attachment.Color = outcomeColor;
+			attachment.FallbackText = $"Job(s) scheduled in an auto-scaled pool with no agents online. Job IDs {jobIds}";
+
+			attachment.Blocks.Add(new HeaderBlock($"Jobs scheduled in empty pool", true));
+			attachment.Blocks.Add(new SectionBlock($"One or more jobs were scheduled in an auto-scaled pool but with no current agents online."));
+			attachment.Blocks.Add(new SectionBlock(sb.ToString()));
+			
+			await SendMessageAsync(recipient, attachments: new[] { attachment });
+		}
+
+		#region Job Complete
+
+		/// <inheritdoc/>
+		public async Task NotifyJobCompleteAsync(IStream jobStream, IJob job, IGraph graph, LabelOutcome outcome)
+		{
+			if (job.NotificationChannel != null)
+			{
+				await SendJobCompleteNotificationToChannelAsync(job.NotificationChannel, job.NotificationChannelFilter, jobStream, job, graph, outcome);
+			}
+			if (jobStream.NotificationChannel != null)
+			{
+				await SendJobCompleteNotificationToChannelAsync(jobStream.NotificationChannel, jobStream.NotificationChannelFilter, jobStream, job, graph, outcome);
+			}
+		}
+
+		Task SendJobCompleteNotificationToChannelAsync(string notificationChannel, string? notificationFilter, IStream jobStream, IJob job, IGraph graph, LabelOutcome outcome)
+		{
+			if (notificationFilter != null)
+			{
+				List<LabelOutcome> outcomes = new List<LabelOutcome>();
+				foreach (string filterOption in notificationFilter.Split('|'))
+				{
+					LabelOutcome result;
+					if (Enum.TryParse(filterOption, out result))
+					{
+						outcomes.Add(result);
+					}
+					else
+					{
+						_logger.LogWarning("Invalid filter option {Option} specified in job filter {NotificationChannelFilter} in job {JobId} or stream {StreamId}", filterOption, notificationFilter, job.Id, job.StreamId);
+					}
+				}
+				if (!outcomes.Contains(outcome))
+				{
+					return Task.CompletedTask;
+				}
+			}
+			return SendJobCompleteMessageAsync(notificationChannel, jobStream, job, graph);
+		}
+
+		/// <inheritdoc/>
+		public async Task NotifyJobCompleteAsync(IUser slackUser, IStream jobStream, IJob job, IGraph graph, LabelOutcome outcome)
+		{
+			string? slackUserId = await GetSlackUserId(slackUser);
+			if (slackUserId != null)
+			{
+				await SendJobCompleteMessageAsync(slackUserId, jobStream, job, graph);
+			}
+		}
+
+		private Task SendJobCompleteMessageAsync(string recipient, IStream stream, IJob job, IGraph graph)
+		{
+			JobStepOutcome jobOutcome = job.Batches.SelectMany(x => x.Steps).Min(x => x.Outcome);
+			_logger.LogInformation("Sending Slack notification for job {JobId} outcome {Outcome} to {SlackUser}", job.Id, jobOutcome, recipient);
+
+			Uri jobLink = new Uri($"{_settings.DashboardUrl}job/{job.Id}");
+
+			Color outcomeColor = jobOutcome == JobStepOutcome.Failure ? BlockKitAttachmentColors.Error : jobOutcome == JobStepOutcome.Warnings ? BlockKitAttachmentColors.Warning : BlockKitAttachmentColors.Success;
+
+			BlockKitAttachment attachment = new BlockKitAttachment();
+			attachment.FallbackText = $"{stream.Name} - {GetJobChangeText(job)} - {job.Name} - {jobOutcome}";
+			attachment.Color = outcomeColor;
+			attachment.Blocks.Add(new SectionBlock($"*<{jobLink}|{stream.Name} - {GetJobChangeText(job)} - {job.Name}>*"));
+			if (jobOutcome == JobStepOutcome.Success)
+			{
+				attachment.Blocks.Add(new SectionBlock($"*Job Succeeded*"));
+			}
+			else
+			{
+				List<string> failedStepStrings = new List<string>();
+				List<string> warningStepStrings = new List<string>();
+
+				IReadOnlyDictionary<NodeRef, IJobStep> nodeToStep = job.GetStepForNodeMap();
+				foreach ((NodeRef nodeRef, IJobStep step) in nodeToStep)
+				{
+					if (step.State == JobStepState.Completed)
+					{
+						INode stepNode = graph.GetNode(nodeRef);
+						string stepName = $"<{jobLink}?step={step.Id}|{stepNode.Name}>";
+						if (step.Outcome == JobStepOutcome.Failure)
+						{
+							failedStepStrings.Add(stepName);
+						}
+						else if (step.Outcome == JobStepOutcome.Warnings)
+						{
+							warningStepStrings.Add(stepName);
+						}
+					}
+				}
+
+				if (failedStepStrings.Any())
+				{
+					string msg = $"*Errors*\n{String.Join(", ", failedStepStrings)}";
+					attachment.Blocks.Add(new SectionBlock(msg.Substring(0, Math.Min(msg.Length, 3000))));
+				}
+				else if (warningStepStrings.Any())
+				{
+					string msg = $"*Warnings*\n{String.Join(", ", warningStepStrings)}";
+					attachment.Blocks.Add(new SectionBlock(msg.Substring(0, Math.Min(msg.Length, 3000))));
+				}
+			}
+
+			if (job.AutoSubmit)
+			{
+				attachment.Blocks.Add(new DividerBlock());
+				if (job.AutoSubmitChange != null)
+				{
+					attachment.Blocks.Add(new SectionBlock($"Shelved files were submitted in CL {job.AutoSubmitChange}."));
+				}
+				else
+				{
+					attachment.Color = BlockKitAttachmentColors.Warning;
+
+					string autoSubmitMessage = String.Empty;
+					if (!String.IsNullOrEmpty(job.AutoSubmitMessage))
+					{
+						autoSubmitMessage = $"\n\n```{job.AutoSubmitMessage}```";
+					}
+
+					attachment.Blocks.Add(new SectionBlock($"Files in CL *{job.PreflightChange}* were *not submitted*. Please resolve the following issues and submit manually.{autoSubmitMessage}"));
+				}
+			}
+
+			return SendMessageAsync(recipient, attachments: new[] { attachment });
+		}
+
+		#endregion
+
+		#region Job step complete
+
+		/// <inheritdoc/>
+		public async Task NotifyJobStepCompleteAsync(IUser slackUser, IStream jobStream, IJob job, IJobStepBatch batch, IJobStep step, INode node, List<ILogEventData> jobStepEventData)
+		{
+			_logger.LogInformation("Sending Slack notification for job {JobId}, batch {BatchId}, step {StepId}, outcome {Outcome} to {SlackUser} ({UserId})", job.Id, batch.Id, step.Id, step.Outcome, slackUser.Name, slackUser.Id);
+
+			string? slackUserId = await GetSlackUserId(slackUser);
+			if (slackUserId != null)
+			{
+				await SendJobStepCompleteMessageAsync(slackUserId, jobStream, job, step, node, jobStepEventData);
+			}
+		}
+
+		/// <summary>
+		/// Creates a Slack message about a completed step job.
+		/// </summary>
+		/// <param name="recipient"></param>
+		/// <param name="stream"></param>
+		/// <param name="job">The job that contains the step that completed.</param>
+		/// <param name="step">The job step that completed.</param>
+		/// <param name="node">The node for the job step.</param>
+		/// <param name="events">Any events that occurred during the job step.</param>
+		private Task SendJobStepCompleteMessageAsync(string recipient, IStream stream, IJob job, IJobStep step, INode node, List<ILogEventData> events)
+		{
+			Uri jobStepLink = new Uri($"{_settings.DashboardUrl}job/{job.Id}?step={step.Id}");
+			Uri jobStepLogLink = new Uri($"{_settings.DashboardUrl}log/{step.LogId}");
+
+			Color outcomeColor = step.Outcome == JobStepOutcome.Failure ? BlockKitAttachmentColors.Error : step.Outcome == JobStepOutcome.Warnings ? BlockKitAttachmentColors.Warning : BlockKitAttachmentColors.Success;
+
+			BlockKitAttachment attachment = new BlockKitAttachment();
+			attachment.FallbackText = $"{stream.Name} - {GetJobChangeText(job)} - {job.Name} - {node.Name} - {step.Outcome}";
+			attachment.Color = outcomeColor;
+			attachment.Blocks.Add(new SectionBlock($"*<{jobStepLink}|{stream.Name} - {GetJobChangeText(job)} - {job.Name} - {node.Name}>*"));
+			if (step.Outcome == JobStepOutcome.Success)
+			{
+				attachment.Blocks.Add(new SectionBlock($"*Job Step Succeeded*"));
+			}
+			else
+			{
+				List<ILogEventData> errors = events.Where(x => x.Severity == EventSeverity.Error).ToList();
+				List<ILogEventData> warnings = events.Where(x => x.Severity == EventSeverity.Warning).ToList();
+				List<string> eventStrings = new List<string>();
+				if (errors.Any())
+				{
+					string errorSummary = errors.Count > MaxJobStepEvents ? $"*Errors (First {MaxJobStepEvents} shown)*" : $"*Errors*";
+					attachment.Blocks.Add(new SectionBlock(errorSummary));
+					foreach (ILogEventData error in errors.Take(MaxJobStepEvents))
+					{
+						attachment.Blocks.Add(QuoteBlock(error.Message));
+					}
+				}
+				else if (warnings.Any())
+				{
+					string warningSummary = warnings.Count > MaxJobStepEvents ? $"*Warnings (First {MaxJobStepEvents} shown)*" : $"*Warnings*";
+					eventStrings.Add(warningSummary);
+					foreach (ILogEventData warning in warnings.Take(MaxJobStepEvents))
+					{
+						attachment.Blocks.Add(QuoteBlock(warning.Message));
+					}
+				}
+
+				attachment.Blocks.Add(new SectionBlock($"<{jobStepLogLink}|View Job Step Log>"));
+			}
+
+			return SendMessageAsync(recipient, attachments: new[] { attachment });
+		}
+
+		#endregion
+
+		#region Label complete
+
+		/// <inheritdoc/>
+		public async Task NotifyLabelCompleteAsync(IUser user, IJob job, IStream stream, ILabel label, int labelIdx, LabelOutcome outcome, List<(string, JobStepOutcome, Uri)> stepData)
+		{
+			_logger.LogInformation("Sending Slack notification for job {JobId} outcome {Outcome} to {Name} ({UserId})", job.Id, outcome, user.Name, user.Id);
+
+			string? slackUserId = await GetSlackUserId(user);
+			if (slackUserId != null)
+			{
+				await SendLabelUpdateMessageAsync(slackUserId, stream, job, label, labelIdx, outcome, stepData);
+			}
+		}
+
+		Task SendLabelUpdateMessageAsync(string recipient, IStream stream, IJob job, ILabel label, int labelIdx, LabelOutcome outcome, List<(string, JobStepOutcome, Uri)> jobStepData)
+		{
+			Uri labelLink = new Uri($"{_settings.DashboardUrl}job/{job.Id}?label={labelIdx}");
+
+			Color outcomeColor = outcome == LabelOutcome.Failure ? BlockKitAttachmentColors.Error : outcome == LabelOutcome.Warnings ? BlockKitAttachmentColors.Warning : BlockKitAttachmentColors.Success;
+
+			BlockKitAttachment attachment = new BlockKitAttachment();
+			attachment.FallbackText = $"{stream.Name} - {GetJobChangeText(job)} - {job.Name} - Label {label.DashboardName} - {outcome}";
+			attachment.Color = outcomeColor;
+			attachment.Blocks.Add(new SectionBlock($"*<{labelLink}|{stream.Name} - {GetJobChangeText(job)} - {job.Name} - Label {label.DashboardName}>*"));
+			if (outcome == LabelOutcome.Success)
+			{
+				attachment.Blocks.Add(new SectionBlock($"*Label Succeeded*"));
+			}
+			else
+			{
+				List<string> failedStepStrings = new List<string>();
+				List<string> warningStepStrings = new List<string>();
+				foreach ((string Name, JobStepOutcome StepOutcome, Uri Link) jobStep in jobStepData)
+				{
+					string stepString = $"<{jobStep.Link}|{jobStep.Name}>";
+					if (jobStep.StepOutcome == JobStepOutcome.Failure)
+					{
+						failedStepStrings.Add(stepString);
+					}
+					else if (jobStep.StepOutcome == JobStepOutcome.Warnings)
+					{
+						warningStepStrings.Add(stepString);
+					}
+				}
+
+				if (failedStepStrings.Any())
+				{
+					attachment.Blocks.Add(new SectionBlock($"*Errors*\n{String.Join(", ", failedStepStrings)}"));
+				}
+				else if (warningStepStrings.Any())
+				{
+					attachment.Blocks.Add(new SectionBlock($"*Warnings*\n{String.Join(", ", warningStepStrings)}"));
+				}
+			}
+
+			return SendMessageAsync(recipient, attachments: new[] { attachment });
+		}
+
+		#endregion
+
+		#region Issues
+
+		/// <inheritdoc/>
+		public async Task NotifyIssueUpdatedAsync(IIssue issue)
+		{
+			// Do not send notifications for quarantined issues
+			if (issue.QuarantinedByUserId != null)
+			{
+				_logger.LogInformation("Skipping notifications for quarantined issue {IssueId}", issue.Id);
+				return;
+			}
+
+			using IDisposable scope = _logger.BeginScope("Slack notifications for issue {IssueId}", issue.Id);
+			_logger.LogInformation("Updating Slack notifications for issue {IssueId}", issue.Id);
+
+			IIssueDetails details = await _issueService.GetIssueDetailsAsync(issue);
+
+			WorkflowConfig? workflow = null;
+			if (details.Spans.Count > 0)
+			{
+				IIssueSpan span = details.Spans[0];
+
+				WorkflowId? workflowId = span.LastFailure.Annotations.WorkflowId;
+				if (workflowId != null)
+				{
+					IStream? stream = await _streamService.GetStreamAsync(span.StreamId);
+					if (stream != null && stream.Config.TryGetWorkflow(workflowId.Value, out workflow))
+					{
+						await CreateOrUpdateWorkflowThreadAsync(issue, span, details.Spans, workflow);
+					}
+				}
+			}
+
+			bool notifyOwner = workflow?.TriageChannel == null;
+			bool notifySuspects = issue.Promoted || details.Spans.Any(x => x.LastFailure.Annotations.NotifySubmitters ?? false);
+
+			HashSet<UserId> userIds = new HashSet<UserId>();
+			if (notifySuspects)
+			{
+				userIds.UnionWith(details.Suspects.Select(x => x.AuthorId));
+			}
+			if (notifyOwner && issue.OwnerId.HasValue)
+			{
+				userIds.Add(issue.OwnerId.Value);
+			}
+
+			HashSet<string> channels = new HashSet<string>(StringComparer.OrdinalIgnoreCase);
+
+			List<MessageStateDocument> existingMessages = await _messageStates.Find(x => x.EventId == GetIssueEventId(issue)).ToListAsync();
+			foreach (MessageStateDocument existingMessage in existingMessages)
+			{
+				if (existingMessage.UserId != null)
+				{
+					userIds.Add(existingMessage.UserId.Value);
+				}
+				else
+				{
+					channels.Add(existingMessage.Recipient);
+				}
+			}
+
+			foreach (IIssueSpan span in details.Spans)
+			{
+				IStream? stream = await _streamService.GetCachedStream(span.StreamId);
+				if (stream != null)
+				{
+					TemplateRef? templateRef;
+					if (stream.Templates.TryGetValue(span.TemplateRefId, out templateRef) && templateRef.TriageChannel != null)
+					{
+						channels.Add(templateRef.TriageChannel);
+					}
+					else if (stream.TriageChannel != null)
+					{
+						channels.Add(stream.TriageChannel);
+					}
+				}
+			}
+
+			if (userIds.Count > 0)
+			{
+				foreach (UserId userId in userIds)
+				{
+					IUser? user = await _userCollection.GetUserAsync(userId);
+					if (user == null)
+					{
+						_logger.LogWarning("Unable to find user {UserId}", userId);
+					}
+					else
+					{
+						await NotifyIssueUpdatedAsync(user, issue, details);
+					}
+				}
+			}
+
+			if (channels.Count > 0)
+			{
+				foreach (string channel in channels)
+				{
+					await SendIssueMessageAsync(channel, issue, details, null, defaultAllowMentions);
+				}
+			}
+		}
+
+		async Task AddReactionAsync(string channel, string ts, string name)
+		{
+			ReactionMessage message = new ReactionMessage();
+			message.Channel = channel;
+			message.Ts = ts;
+			message.Name = name;
+			await SendRequestAsync<SlackResponse>(AddReactionUrl, message);
+		}
+
+		async Task RemoveReactionAsync(string channel, string ts, string name)
+		{
+			ReactionMessage message = new ReactionMessage();
+			message.Channel = channel;
+			message.Ts = ts;
+			message.Name = name;
+			await SendRequestAsync<SlackResponse>(RemoveReactionUrl, message);
+		}
+
+		static string GetTriageThreadEventId(int issueId) => $"issue_triage_{issueId}";
+
+		async Task CreateOrUpdateWorkflowThreadAsync(IIssue issue, IIssueSpan span, IReadOnlyList<IIssueSpan> spans, WorkflowConfig workflow)
+		{
+			string? triageChannel = workflow.TriageChannel;
+			if (triageChannel != null)
+			{
+				Uri issueUrl = GetIssueUrl(issue, span.FirstFailure);
+
+				string eventId = GetTriageThreadEventId(issue.Id);
+
+				string text = $"{workflow.TriagePrefix}*Issue <{issueUrl}|{issue.Id}>*: {issue.Summary}{workflow.TriageSuffix}";
+				if (!spans.Any(x => x.NextSuccess == null))
+				{
+					text = $"~{text}~";
+				}
+
+				(MessageStateDocument state, bool isNew) = await SendOrUpdateMessageAsync(triageChannel, eventId, null, text);
+				if (isNew)
+				{
+					// Create the summary text
+					List<ILogEvent> events = new List<ILogEvent>();
+					List<ILogEventData> eventDataItems = new List<ILogEventData>();
+
+					if (span.FirstFailure.LogId != null)
+					{
+						LogId logId = span.FirstFailure.LogId.Value;
+						ILogFile? logFile = await _logFileService.GetLogFileAsync(logId);
+						if (logFile != null)
+						{
+							events = await _logFileService.FindEventsAsync(logFile, span.Id, 0, 50);
+							if (events.Any(x => x.Severity == EventSeverity.Error))
+							{
+								events.RemoveAll(x => x.Severity == EventSeverity.Warning);
+							}
+
+							List<string> eventStrings = new List<string>();
+							for (int idx = 0; idx < Math.Min(events.Count, 3); idx++)
+							{
+								ILogEventData data = await _logFileService.GetEventDataAsync(logFile, events[idx].LineIndex, events[idx].LineCount);
+								eventDataItems.Add(data);
+							}
+						}
+					}
+
+					StringBuilder summary = new StringBuilder();
+					summary.AppendLine($"From {FormatJobStep(span.FirstFailure, span.NodeName)}:");
+					foreach (ILogEventData eventDataItem in eventDataItems)
+					{
+						summary.AppendLine(QuoteText(eventDataItem.Message, TextObject.MaxLength));
+					}
+					if (events.Count > eventDataItems.Count)
+					{
+						summary.AppendLine("```...```");
+					}
+					string? summaryTs = await SendMessageToThread(triageChannel, state.Ts, summary.ToString());
+
+					// Permalink to the summary text so we link inside the thread rather than just to the original message
+					string? permalink = null;
+					if (summaryTs != null)
+					{
+						permalink = await GetPermalinkAsync(state.Channel, summaryTs);
+					}
+
+					// If it has an owner, show that
+					if (issue.OwnerId != null)
+					{
+						string mention = await FormatMentionAsync(issue.OwnerId.Value, workflow.AllowMentions);
+						await SendMessageToThread(triageChannel, state.Ts, $"Assigned to {mention}");
+					}
+					else
+					{
+						List<IIssueSuspect> suspects = await _issueService.Collection.FindSuspectsAsync(issue);
+						IGrouping<UserId, IIssueSuspect>[] suspectGroups = suspects.GroupBy(x => x.AuthorId).ToArray();
+
+						if (suspectGroups.Length > 0 && suspectGroups.Length <= workflow.MaxMentions)
+						{
+							List<string> suspectList = new List<string>();
+							foreach (IGrouping<UserId, IIssueSuspect> suspectGroup in suspectGroups)
+							{
+								string mention = await FormatMentionAsync(suspectGroup.Key, workflow.AllowMentions);
+								string changes = String.Join(", ", suspectGroup.Select(x => FormatChange(x.Change)));
+								suspectList.Add($"{mention} ({changes})");
+							}
+
+							string suspectMessage = $"Possibly {StringUtils.FormatList(suspectList, "or")}.";
+							await SendMessageToThread(triageChannel, state.Ts, suspectMessage);
+						}
+					}
+
+					await SetMessageTimestampAsync(state.Id, state.Channel, state.Ts, permalink);
+				}
+
+				if (issue.AcknowledgedAt != null)
+				{
+					await AddReactionAsync(state.Channel, state.Ts, "eyes");
+				}
+				else
+				{
+					await RemoveReactionAsync(state.Channel, state.Ts, "eyes");
+				}
+
+				if (issue.ResolvedAt != null)
+				{
+					await AddReactionAsync(state.Channel, state.Ts, "tick");
+				}
+				else
+				{
+					await RemoveReactionAsync(state.Channel, state.Ts, "tick");
+				}
+
+				if (issue.FixChange != null)
+				{
+					HashSet<StreamId> originStreams = new HashSet<StreamId>(issue.Streams.Where(x => x.MergeOrigin ?? false).Select(x => x.StreamId));
+
+					IIssueSpan? fixFailedSpan = null;
+					foreach (IIssueSpan originSpan in spans.Where(x => originStreams.Contains(x.StreamId)).OrderBy(x => x.StreamId.ToString()))
+					{
+						if (originSpan.LastFailure.Change > issue.FixChange.Value)
+						{
+							fixFailedSpan = originSpan;
+							break;
+						}
+					}
+
+					if (fixFailedSpan == null)
+					{
+						string fixedEventId = $"issue_{issue.Id}_fixed_{issue.FixChange}";
+						string fixedMessage = $"Marked as fixed in {FormatChange(issue.FixChange.Value)}";
+						await PostSingleMessageToThreadAsync(triageChannel, state.Ts, fixedEventId, fixedMessage);
+					}
+					else
+					{
+						string fixFailedEventId = $"issue_{issue.Id}_fixfailed_{issue.FixChange}";
+						string fixFailedMessage = $"Issue not fixed by {FormatChange(issue.FixChange.Value)}; see {FormatJobStep(fixFailedSpan.LastFailure, fixFailedSpan.NodeName)} in {fixFailedSpan.StreamName}.";
+						if (issue.OwnerId.HasValue)
+						{
+							string mention = await FormatMentionAsync(issue.OwnerId.Value, workflow.AllowMentions);
+							fixFailedMessage += $" ({mention})";
+						}
+						await PostSingleMessageToThreadAsync(triageChannel, state.Ts, fixFailedEventId, fixFailedMessage);
+					}
+
+					if (fixFailedSpan == null)
+					{
+						foreach (IIssueStream stream in issue.Streams)
+						{
+							if ((stream.MergeOrigin ?? false) && !(stream.ContainsFix ?? false))
+							{
+								string streamName = spans.FirstOrDefault(x => x.StreamId == stream.StreamId)?.StreamName ?? stream.StreamId.ToString();
+								string missingEventId = $"issue_{issue.Id}_fixmissing_{issue.FixChange}_{stream.StreamId}";
+								string missingMessage = $"Note: Fix may need manually merging to {streamName}";
+								await PostSingleMessageToThreadAsync(triageChannel, state.Ts, missingEventId, missingMessage);
+							}
+						}
+					}
+				}
+
+				if(issue.ExternalIssueKey != null)
+				{
+					string extIssueEventId = $"issue_{issue.Id}_ext_{issue.ExternalIssueKey}";
+					string extIssueMessage = $"Linked to issue {FormatExternalIssue(issue.ExternalIssueKey)}";
+					await PostSingleMessageToThreadAsync(triageChannel, state.Ts, extIssueEventId, extIssueMessage);
+				}
+			}
+		}
+
+		async Task PostSingleMessageToThreadAsync(string channel, string threadTs, string eventId, string message)
+		{
+			(MessageStateDocument state, bool isNew) = await AddOrUpdateMessageStateAsync(channel, eventId, null, "");
+			if (isNew)
+			{
+				string? ts = await SendMessageToThread(channel, threadTs, message);
+				if (ts != null)
+				{
+					await SetMessageTimestampAsync(state.Id, channel, ts);
+				}
+			}
+		}
+
+		async Task<string?> SendMessageToThread(string channel, string threadTs, string text)
+		{
+			BlockKitMessage message = new BlockKitMessage();
+			message.Channel = channel;
+			message.ThreadTs = threadTs;
+			message.Text = text;
+
+			PostMessageResponse response = await SendRequestAsync<PostMessageResponse>(PostMessageUrl, message);
+			return response.Ts;
+		}
+
+		async Task NotifyIssueUpdatedAsync(IUser user, IIssue issue, IIssueDetails details)
+		{
+			string? slackUserId = await GetSlackUserId(user);
+			if (slackUserId == null)
+			{
+				return;
+			}
+
+			await SendIssueMessageAsync(slackUserId, issue, details, user.Id, defaultAllowMentions);
+		}
+
+		Uri GetJobUrl(JobId jobId)
+		{
+			return new Uri(_settings.DashboardUrl, $"job/{jobId}");
+		}
+
+		Uri GetStepUrl(JobId jobId, SubResourceId stepId)
+		{
+			return new Uri(_settings.DashboardUrl, $"job/{jobId}?step={stepId}");
+		}
+
+		Uri GetIssueUrl(IIssue issue, IIssueStep step)
+		{
+			return new Uri(_settings.DashboardUrl, $"job/{step.JobId}?step={step.StepId}&issue={issue.Id}");
+		}
+
+		async Task SendIssueMessageAsync(string recipient, IIssue issue, IIssueDetails details, UserId? userId, bool allowMentions)
+		{
+			using IDisposable scope = _logger.BeginScope("SendIssueMessageAsync (User: {SlackUser}, Issue: {IssueId})", recipient, issue.Id);
+
+			BlockKitAttachment attachment = new BlockKitAttachment();
+			attachment.Color = Color.Red;
+
+			Uri issueUrl = _settings.DashboardUrl;
+			if (details.Steps.Count > 0)
+			{
+				issueUrl = GetIssueUrl(issue, details.Steps[0]);
+			}
+			attachment.Blocks.Add(new SectionBlock($"*<{issueUrl}|Issue #{issue.Id}: {issue.Summary}>*"));
+
+			string streamList = StringUtils.FormatList(details.Spans.Select(x => $"*{x.StreamName}*").Distinct(StringComparer.OrdinalIgnoreCase).OrderBy(x => x, StringComparer.OrdinalIgnoreCase));
+			attachment.Blocks.Add(new SectionBlock($"Occurring in {streamList}"));
+
+			IIssueSpan? lastSpan = details.Spans.OrderByDescending(x => x.LastFailure.StepTime).FirstOrDefault();
+			if (lastSpan != null && lastSpan.LastFailure.LogId != null)
+			{
+				LogId logId = lastSpan.LastFailure.LogId.Value;
+				ILogFile? logFile = await _logFileService.GetLogFileAsync(logId);
+				if(logFile != null)
+				{
+					List<ILogEvent> events = await _logFileService.FindEventsAsync(logFile, lastSpan.Id, 0, 20);
+					if (events.Any(x => x.Severity == EventSeverity.Error))
+					{
+						events.RemoveAll(x => x.Severity == EventSeverity.Warning);
+					}
+
+					for (int idx = 0; idx < Math.Min(events.Count, 3); idx++)
+					{
+						ILogEventData data = await _logFileService.GetEventDataAsync(logFile, events[idx].LineIndex, events[idx].LineCount);
+						attachment.Blocks.Add(QuoteBlock(data.Message));
+					}
+					if (events.Count > 3)
+					{
+						attachment.Blocks.Add(new SectionBlock("```...```"));
+					}
+				}
+			}
+
+			if (issue.FixChange != null)
+			{
+				IIssueStep? fixFailedStep = issue.FindFixFailedStep(details.Spans);
+
+				string text;
+				if (issue.FixChange.Value < 0)
+				{
+					text = ":tick: Marked as a systemic issue.";
+				}
+				else if (fixFailedStep != null)
+				{
+					Uri fixFailedUrl = new Uri(_settings.DashboardUrl, $"job/{fixFailedStep.JobId}?step={fixFailedStep.StepId}&issue={issue.Id}");
+					text = $":cross: Marked fixed in *CL {issue.FixChange.Value}*, but seen again at *<{fixFailedUrl}|CL {fixFailedStep.Change}>*";
+				}
+				else
+				{
+					text = $":tick: Marked fixed in *CL {issue.FixChange.Value}*.";
+				}
+				attachment.Blocks.Add(new SectionBlock(text));
+			}
+			else if (userId != null && issue.OwnerId == userId)
+			{
+				if (issue.AcknowledgedAt.HasValue)
+				{
+					attachment.Blocks.Add(new SectionBlock($":+1: Acknowledged at {FormatSlackTime(issue.AcknowledgedAt.Value)}"));
+				}
+				else
+				{
+					if (issue.NominatedById != null)
+					{
+						string mention = await FormatMentionAsync(issue.NominatedById.Value, allowMentions);
+						string text = $"You were nominated to fix this issue by {mention} at {FormatSlackTime(issue.NominatedAt ?? DateTime.UtcNow)}";
+						attachment.Blocks.Add(new SectionBlock(text));
+					}
+					else
+					{
+						List<int> changes = details.Suspects.Where(x => x.AuthorId == userId).Select(x => x.Change).OrderBy(x => x).ToList();
+						if (changes.Count > 0)
+						{
+							string text = $"Horde has determined that {StringUtils.FormatList(changes.Select(x => $"CL {x}"), "or")} is the most likely cause for this issue.";
+							attachment.Blocks.Add(new SectionBlock(text));
+						}
+					}
+
+					ActionsBlock actions = new ActionsBlock();
+					actions.AddButton("Acknowledge", value: $"issue_{issue.Id}_ack_{userId}", style: ActionButton.ButtonStyle.Primary);
+					actions.AddButton("Not Me", value: $"issue_{issue.Id}_decline_{userId}", style: ActionButton.ButtonStyle.Danger);
+					attachment.Blocks.Add(actions);
+				}
+			}
+			else if (issue.OwnerId != null)
+			{
+				string ownerMention = await FormatMentionAsync(issue.OwnerId.Value, allowMentions);
+				if (issue.AcknowledgedAt.HasValue)
+				{
+					attachment.Blocks.Add(new SectionBlock($":+1: Acknowledged by {ownerMention} at {FormatSlackTime(issue.AcknowledgedAt.Value)}"));
+				}
+				else if (issue.NominatedById == null)
+				{
+					attachment.Blocks.Add(new SectionBlock($"Assigned to {ownerMention}"));
+				}
+				else if (issue.NominatedById == userId)
+				{
+					attachment.Blocks.Add(new SectionBlock($"You nominated {ownerMention} to fix this issue."));
+				}
+				else
+				{
+					attachment.Blocks.Add(new SectionBlock($"{ownerMention} was nominated to fix this issue by {await FormatMentionAsync(issue.NominatedById.Value, allowMentions)}"));
+				}
+			}
+			else if (userId != null)
+			{
+				IIssueSuspect? suspect = details.Suspects.FirstOrDefault(x => x.AuthorId == userId);
+				if (suspect != null)
+				{
+					if (suspect.DeclinedAt != null)
+					{
+						attachment.Blocks.Add(new SectionBlock($":downvote: Declined at {FormatSlackTime(suspect.DeclinedAt.Value)}"));
+					}
+					else
+					{
+						attachment.Blocks.Add(new SectionBlock("Please check if any of your recently submitted changes have caused this issue."));
+
+						ActionsBlock actions = new ActionsBlock();
+						actions.AddButton("Will Fix", value: $"issue_{issue.Id}_accept_{userId}", style: ActionButton.ButtonStyle.Primary);
+						actions.AddButton("Not Me", value: $"issue_{issue.Id}_decline_{userId}", style: ActionButton.ButtonStyle.Danger);
+						attachment.Blocks.Add(actions);
+					}
+				}
+			}
+			else if (details.Suspects.Count > 0)
+			{
+				List<string> declinedLines = new List<string>();
+				foreach (IIssueSuspect suspect in details.Suspects)
+				{
+					if (!details.Issue.Promoted)
+					{
+						declinedLines.Add($"Possibly {await FormatNameAsync(suspect.AuthorId)} (CL {suspect.Change})");
+					}
+					else if (suspect.DeclinedAt == null)
+					{
+						declinedLines.Add($":heavy_minus_sign: Ignored by {await FormatNameAsync(suspect.AuthorId)} (CL {suspect.Change})");
+					}
+					else
+					{
+						declinedLines.Add($":downvote: Declined by {await FormatNameAsync(suspect.AuthorId)} at {FormatSlackTime(suspect.DeclinedAt.Value)} (CL {suspect.Change})");
+					}
+				}
+				attachment.Blocks.Add(new SectionBlock(String.Join("\n", declinedLines)));
+			}
+
+			if (IsRecipientAllowed(recipient, "issue update"))
+			{
+				await SendOrUpdateMessageAsync(recipient, GetIssueEventId(issue), userId, attachments: new[] { attachment });
+			}
+		}
+
+		static string GetIssueEventId(IIssue issue)
+		{
+			return $"issue_{issue.Id}";
+		}
+
+		async Task<string> FormatNameAsync(UserId userId)
+		{
+			IUser? user = await _userCollection.GetUserAsync(userId);
+			if (user == null)
+			{
+				return $"User {userId}";
+			}
+			return user.Name;
+		}
+
+		static string FormatChange(int change) => $"<ugs://change?number={change}|CL {change}>";
+
+		string FormatJobStep(IIssueStep step, string name) => $"<{GetJobUrl(step.JobId)}|{step.JobName}> / <{GetStepUrl(step.JobId, step.StepId)}|{name}>";
+
+		string FormatExternalIssue(string key)
+		{
+			string? url = _externalIssueService.GetIssueUrl(key);
+			if (url == null)
+			{
+				return key;
+			}
+			else
+			{
+				return $"<{_externalIssueService.GetIssueUrl(key)}|{key}>";
+			}
+		}
+
+		async Task<string> FormatMentionAsync(UserId userId, bool allowMentions)
+		{
+			IUser? user = await _userCollection.GetUserAsync(userId);
+			if (user == null)
+			{
+				return $"User {userId}";
+			}
+
+			string? slackUserId = await GetSlackUserId(user);
+			if (slackUserId == null)
+			{
+				return user.Login;
+			}
+
+			if (!_environment.IsProduction() || !allowMentions)
+			{
+				return $"{user.Name} [{slackUserId}]";
+			}
+
+			return $"<@{slackUserId}>";
+		}
+
+		async Task HandleIssueResponseAsync(EventPayload payload, ActionInfo action)
+		{
+			string? userName = payload.User?.UserName;
+			if(userName == null)
+			{
+				_logger.LogWarning("No user for message payload: {Payload}", payload);
+				return;
+			}
+			if (payload.ResponseUrl == null)
+			{
+				_logger.LogWarning("No response url for payload: {Payload}", payload);
+				return;
+			}
+
+			Match match = Regex.Match(action.Value ?? String.Empty, @"^issue_(\d+)_([a-zA-Z]+)_([a-fA-F0-9]{24})$");
+			if (!match.Success)
+			{
+				_logger.LogWarning("Could not match format of button action: {Action}", action.Value);
+				return;
+			}
+
+			int issueId = Int32.Parse(match.Groups[1].Value, NumberStyles.Integer, CultureInfo.InvariantCulture);
+			string verb = match.Groups[2].Value;
+			UserId userId = match.Groups[3].Value.ToObjectId<IUser>();
+			_logger.LogInformation("Issue {IssueId}: {Action} from {SlackUser} ({UserId})", issueId, verb, userName, userId);
+
+			if (String.Equals(verb, "ack", StringComparison.Ordinal))
+			{
+				await _issueService.UpdateIssueAsync(issueId, acknowledged: true);
+			}
+			else if (String.Equals(verb, "accept", StringComparison.Ordinal))
+			{
+				await _issueService.UpdateIssueAsync(issueId, ownerId: userId, acknowledged: true);
+			}
+			else if (String.Equals(verb, "decline", StringComparison.Ordinal))
+			{
+				await _issueService.UpdateIssueAsync(issueId, declinedById: userId);
+			}
+
+			IIssue? newIssue = await _issueService.Collection.GetIssueAsync(issueId);
+			if (newIssue != null)
+			{
+				IUser? user = await _userCollection.GetUserAsync(userId);
+				if (user != null)
+				{
+					string? recipient = await GetSlackUserId(user);
+					if (recipient != null)
+					{
+						IIssueDetails details = await _issueService.GetIssueDetailsAsync(newIssue);
+						await SendIssueMessageAsync(recipient, newIssue, details, userId, defaultAllowMentions);
+					}
+				}
+			}
+		}
+
+		static string FormatSlackTime(DateTimeOffset time)
+		{
+			return $"<!date^{time.ToUnixTimeSeconds()}^{{time}}|{time}>";
+		}
+
+		/// <inheritdoc/>
+		public async Task SendIssueReportAsync(IssueReport report)
+		{
+			if (report.Workflow.ReportChannel != null)
+			{
+				List<BlockBase> blocks = new List<BlockBase>();
+				blocks.Add(new HeaderBlock(AddEnvironmentAnnotation($"{report.Stream.Name}: {report.Time:d}")));
+				await SendMessageAsync(report.Workflow.ReportChannel, blocks: blocks.ToArray(), withEnvironment: false);
+
+				StringBuilder body = new StringBuilder();
+				if (report.Issues.Count == 0)
+				{
+					body.Append(":tick: No issues open.");
+				}
+				else
+				{
+					TimeSpan averageAge = TimeSpan.FromHours(report.Issues.Select(x => (report.Time - x.CreatedAt).TotalHours).Average());
+					body.Append($"*{report.Issues.Count} issues* currently open (average age {FormatReadableTimeSpan(averageAge)}):");
+				}
+				body.Append('\n');
+
+				foreach(IIssue issue in report.Issues.OrderByDescending(x => x.Id))
+				{
+					IIssueSpan? span = report.IssueSpans.FirstOrDefault(x => x.IssueId == issue.Id);
+
+					Uri issueUrl = _settings.DashboardUrl;
+					if(span != null)
+					{
+						IIssueStep lastFailure = span.LastFailure;
+						issueUrl = new Uri(issueUrl, $"job/{lastFailure.JobId}?step={lastFailure.StepId}&issue={issue.Id}");
+					}
+
+					string status = "*Unassigned*";
+					if(issue.OwnerId != null)
+					{
+						IUser? user = await _userCollection.GetCachedUserAsync(issue.OwnerId.Value);
+						if (user == null)
+						{
+							status = $"Assigned to user {issue.OwnerId.Value}";
+						}
+						else
+						{
+							status = $"Assigned to {user.Name}";
+						}
+						if(issue.AcknowledgedAt == null)
+						{
+							status = $"{status} (unacknowledged)";
+						}
+					}
+
+					if (issue.QuarantinedByUserId != null)
+					{
+						status = $"{status} - *Quarantined*";
+					}
+
+					body.Append($"\n\u2022 *Issue <{issueUrl}|{issue.Id}>");
+					
+					string? triageChannel = report.Workflow.TriageChannel;
+					if (triageChannel != null)
+					{
+						MessageStateDocument? state = await GetMessageStateAsync(triageChannel, GetTriageThreadEventId(issue.Id));
+						if (state != null && state.Permalink != null)
+						{
+							body.Append($" (<{state.Permalink}|Thread>)");
+						}
+					}
+
+					body.Append($"*: {issue.Summary} [{FormatReadableTimeSpan(report.Time - issue.CreatedAt)}]");
+
+					if (!String.IsNullOrEmpty(issue.ExternalIssueKey))
+					{
+						body.Append($" ({FormatExternalIssue(issue.ExternalIssueKey)})");
+					}
+
+					body.Append($" - {status}");
+				}
+
+				if (report.WorkflowStats.NumSteps > 0)
+				{
+					double totalPct = (report.WorkflowStats.NumPassingSteps * 100.0) / report.WorkflowStats.NumSteps;
+					body.Append($"\n\n{report.WorkflowStats.NumPassingSteps:n0} of {report.WorkflowStats.NumSteps:n0} (*{totalPct:0.0}%*) build steps succeeded since last status update.");
+				}
+
+				await SendMessageAsync(report.Workflow.ReportChannel, text: body.ToString(), withEnvironment: false);
+			}
+		}
+
+		static string FormatReadableTimeSpan(TimeSpan timeSpan)
+		{
+			double totalDays = timeSpan.TotalDays;
+			if (totalDays > 1.0)
+			{
+				return $"{totalDays:n0}d";
+			}
+
+			double totalHours = timeSpan.TotalHours;
+			if(totalHours > 1.0)
+			{
+				return $"{totalHours:n0}h";
+			}
+
+			return $"{timeSpan.TotalMinutes:n0}m";
+		}
+
+		#endregion
+
+		#region Stream updates
+
+		/// <inheritdoc/>
+		public async Task NotifyConfigUpdateFailureAsync(string errorMessage, string fileName, int? change = null, IUser? author = null, string? description = null)
+		{
+			_logger.LogInformation("Sending config update failure notification for {FileName} (change: {Change}, author: {UserId})", fileName, change ?? -1, author?.Id ?? UserId.Empty);
+
+			string? slackUserId = null;
+			if (author != null)
+			{
+				slackUserId = await GetSlackUserId(author);
+				if (slackUserId == null)
+				{
+					_logger.LogWarning("Unable to identify Slack user id for {UserId}", author.Id);
+				}
+				else
+				{
+					_logger.LogInformation("Mappsed user {UserId} to Slack user {SlackUserId}", author.Id, slackUserId);
+				}
+			}
+
+			if (slackUserId != null)
+			{
+				await SendConfigUpdateFailureMessageAsync(slackUserId, errorMessage, fileName, change, slackUserId, description);
+			}
+			if (_settings.UpdateStreamsNotificationChannel != null)
+			{
+				await SendConfigUpdateFailureMessageAsync($"#{_settings.UpdateStreamsNotificationChannel}", errorMessage, fileName, change, slackUserId, description);
+			}
+		}
+
+		private Task SendConfigUpdateFailureMessageAsync(string recipient, string errorMessage, string fileName, int? change = null, string? author = null, string? description = null)
+		{
+			Color outcomeColor = BlockKitAttachmentColors.Error;
+			BlockKitAttachment attachment = new BlockKitAttachment();
+			attachment.Color = outcomeColor;
+			attachment.FallbackText = $"Update Failure: {fileName}";
+
+			attachment.Blocks.Add(new HeaderBlock($"Config Update Failure :rip:", true));
+
+			attachment.Blocks.Add(new SectionBlock($"Horde was unable to update {fileName}"));
+			attachment.Blocks.Add(QuoteBlock(errorMessage));
+			if (change != null)
+			{
+				if (author != null)
+				{
+					attachment.Blocks.Add(new SectionBlock($"Possibly due to CL: {change.Value} by <@{author}>"));
+				}
+				else
+				{
+					attachment.Blocks.Add(new SectionBlock($"Possibly due to CL: {change.Value} - (Could not determine author from P4 user)"));
+				}
+				if (description != null)
+				{
+					attachment.Blocks.Add(QuoteBlock(description));
+				}
+			}
+
+			return SendMessageAsync(recipient, attachments: new[] { attachment });
+		}
+
+		#endregion
+
+		#region Stream update (file)
+
+		/// <inheritdoc/>
+		public async Task NotifyStreamUpdateFailedAsync(FileSummary file)
+		{
+			_logger.LogDebug("Sending stream update failure notification for {File}", file.DepotPath);
+			if (_settings.UpdateStreamsNotificationChannel != null)
+			{
+				await SendStreamUpdateFailureMessage($"#{_settings.UpdateStreamsNotificationChannel}", file);
+			}
+		}
+
+		/// <summary>
+		/// Creates a stream update failure message in relation to a file
+		/// </summary>
+		/// <param name="recipient"></param>
+		/// <param name="file">The file</param>
+		/// <returns></returns>
+		Task SendStreamUpdateFailureMessage(string recipient, FileSummary file)
+		{
+			Color outcomeColor = BlockKitAttachmentColors.Error;
+			BlockKitAttachment attachment = new BlockKitAttachment();
+			attachment.Color = outcomeColor;
+			attachment.FallbackText = $"{file.DepotPath} - Update Failure";
+
+			attachment.Blocks.Add(new HeaderBlock($"Stream Update Failure :rip:", true));
+
+			attachment.Blocks.Add(new SectionBlock($"<!here> Horde was unable to update {file.DepotPath}"));
+			if (file.Error != null)
+			{
+				attachment.Blocks.Add(QuoteBlock(file.Error));
+			}
+
+			return SendMessageAsync(recipient, attachments: new[] { attachment });
+		}
+
+		#endregion
+
+		#region Device notifications
+
+		/// <inheritdoc/>
+		public async Task NotifyDeviceServiceAsync(string message, IDevice? device = null, IDevicePool? pool = null, IStream? stream = null, IJob? job = null, IJobStep? step = null, INode? node = null, IUser? user = null)
+		{
+			string recipient = $"#{_settings.DeviceServiceNotificationChannel}";
+
+			if (user != null)
+			{
+				string? slackRecipient = await GetSlackUserId(user);
+
+				if (slackRecipient == null)
+				{
+					_logger.LogError("NotifyDeviceServiceAsync - Unable to send user slack notification, user {UserId} slack user id not found", user.Id);
+					return;
+				}
+
+				recipient = slackRecipient;
+			}
+
+			_logger.LogDebug("Sending device service notification to {Recipient}", recipient);
+
+			if (_settings.DeviceServiceNotificationChannel != null)
+			{
+				await SendDeviceServiceMessage(recipient, message, device, pool, stream, job, step, node, user);
+			}
+		}
+
+		/// <summary>
+		/// Creates a Slack message for a device service notification
+		/// </summary>
+		/// <param name="recipient"></param>
+		/// <param name="message"></param>
+		/// <param name="device"></param>
+		/// <param name="pool"></param>
+		/// <param name="stream"></param>
+		/// <param name="job">The job that contains the step that completed.</param>
+		/// <param name="step">The job step that completed.</param>
+		/// <param name="node">The node for the job step.</param>
+		/// <param name="user">The user to notify.</param>
+		private Task SendDeviceServiceMessage(string recipient, string message, IDevice? device = null, IDevicePool? pool = null, IStream? stream = null, IJob? job = null, IJobStep? step = null, INode? node = null, IUser? user = null)
+		{
+
+			if (user != null)
+			{
+				return SendMessageAsync(recipient, message);
+			}
+
+			// truncate message to avoid slack error on message length
+			if (message.Length > 150)
+			{
+				message = message.Substring(0, 146) + "...";
+			}
+
+			BlockKitAttachment attachment = new BlockKitAttachment();
+							
+			attachment.FallbackText = $"{message}";
+
+			if (device != null && pool != null)
+			{
+				attachment.FallbackText += $" - Device: {device.Name} Pool: {pool.Name}";
+			}
+				
+			attachment.Blocks.Add(new HeaderBlock(AddEnvironmentAnnotation(message), false));
+
+			if (stream != null && job != null && step != null && node != null)
+			{
+				Uri jobStepLink = new Uri($"{_settings.DashboardUrl}job/{job.Id}?step={step.Id}");
+				Uri jobStepLogLink = new Uri($"{_settings.DashboardUrl}log/{step.LogId}");
+
+				attachment.FallbackText += $" - {stream.Name} - {GetJobChangeText(job)} - {job.Name} - {node.Name}";
+				attachment.Blocks.Add(new SectionBlock($"*<{jobStepLink}|{stream.Name} - {GetJobChangeText(job)} - {job.Name} - {node.Name}>*"));
+				attachment.Blocks.Add(new SectionBlock($"<{jobStepLogLink}|View Job Step Log>"));
+			}
+			else
+			{
+				attachment.FallbackText += " - No job information (Gauntlet might need to be updated in stream)";
+				attachment.Blocks.Add(new SectionBlock("*No job information (Gauntlet might need to be updated in stream)*"));
+			}
+
+			return SendMessageAsync(recipient, attachments: new[] { attachment });
+
+		}
+
+		#endregion
+
+		static SectionBlock QuoteBlock(string text)
+		{
+			return new SectionBlock(QuoteText(text, TextObject.MaxLength));
+		}
+
+		static string QuoteText(string text, int maxLength)
+		{
+			maxLength -= 6;
+			if (text.Length > maxLength)
+			{
+				int length = text.LastIndexOf('\n', maxLength);
+				if (length == 0)
+				{
+					text = text.Substring(0, maxLength - 7) + "...\n...";
+				}
+				else
+				{
+					text = text.Substring(0, length + 1) + "...";
+				}
+			}
+			return $"```{text}```";
+		}
+
+		const int MaxJobStepEvents = 5;
+
+		static string GetJobChangeText(IJob job)
+		{
+			if (job.PreflightChange == 0)
+			{
+				return $"CL {job.Change}";
+			}
+			else
+			{
+				return $"Preflight CL {job.PreflightChange} against CL {job.Change}";
+			}
+		}
+
+		static bool ShouldUpdateUser(SlackUser? document)
+		{
+			if(document == null || document.Version < SlackUser.CurrentVersion)
+			{
+				return true;
+			}
+
+			TimeSpan expiryTime;
+			if (document.SlackUserId == null)
+			{
+				expiryTime = TimeSpan.FromMinutes(10.0);
+			}
+			else
+			{
+				expiryTime = TimeSpan.FromDays(1.0);
+			}
+			return document.Time + expiryTime < DateTime.UtcNow;
+		}
+
+		private async Task<string?> GetSlackUserId(IUser user)
+		{
+			return (await GetSlackUser(user))?.SlackUserId;
+		}
+
+		private async Task<SlackUser?> GetSlackUser(IUser user)
+		{
+			string? email = user.Email;
+			if (email == null)
+			{
+				_logger.LogWarning("Unable to find Slack user id for {UserId} ({Name}): No email address in user profile", user.Id, user.Name);
+				return null;
+			}
+
+			SlackUser? userDocument;
+			if (!_userCache.TryGetValue(email, out userDocument))
+			{
+				userDocument = await _slackUsers.Find(x => x.Id == user.Id).FirstOrDefaultAsync();
+				if (userDocument == null || ShouldUpdateUser(userDocument))
+				{
+					UserInfo? userInfo = await GetSlackUserInfoByEmail(email);
+					if (userDocument == null || userInfo != null)
+					{
+						userDocument = new SlackUser(user.Id, userInfo);
+						await _slackUsers.ReplaceOneAsync(x => x.Id == user.Id, userDocument, new ReplaceOptions { IsUpsert = true });
+					}
+				}
+				using (ICacheEntry entry = _userCache.CreateEntry(email))
+				{
+					entry.SlidingExpiration = TimeSpan.FromMinutes(10.0);
+					entry.Value = userDocument;
+				}
+			}
+
+			return userDocument;
+		}
+
+		private async Task<UserInfo?> GetSlackUserInfoByEmail(string email)
+		{
+			using HttpClient client = new HttpClient();
+
+			using HttpRequestMessage getUserIdRequest = new HttpRequestMessage(HttpMethod.Post, $"https://slack.com/api/users.lookupByEmail?email={email}");
+			getUserIdRequest.Headers.Add("Authorization", $"Bearer {_settings.SlackToken ?? ""}");
+
+			HttpResponseMessage responseMessage = await client.SendAsync(getUserIdRequest);
+			byte[] responseData = await responseMessage.Content.ReadAsByteArrayAsync();
+
+			UserResponse userResponse = JsonSerializer.Deserialize<UserResponse>(responseData)!;
+			if(!userResponse.Ok || userResponse.User == null)
+			{
+				_logger.LogWarning("Unable to find Slack user id for {Email}: {Response}", email, Encoding.UTF8.GetString(responseData));
+				return null;
+			}
+
+			return userResponse.User;
+		}
+
+		class SlackResponse
+		{
+			[JsonPropertyName("ok")]
+			public bool Ok { get; set; }
+
+			[JsonPropertyName("error")]
+			public string? Error { get; set; }
+		}
+
+		class PostMessageResponse : SlackResponse
+		{
+			[JsonPropertyName("channel")]
+			public string? Channel { get; set; }
+
+			[JsonPropertyName("ts")]
+			public string? Ts { get; set; }
+		}
+
+		class GetPermalinkRequest
+		{
+			[JsonPropertyName("channel")]
+			public string? Channel { get; set; }
+
+			[JsonPropertyName("message_ts")]
+			public string? MessageTs { get; set; }
+		}
+
+		class GetPermalinkResponse : SlackResponse
+		{
+			[JsonPropertyName("channel")]
+			public string? Channel { get; set; }
+
+			[JsonPropertyName("permalink")]
+			public string? Permalink { get; set; }
+		}
+
+		[return: NotNullIfNotNull("message")]
+		private string? AddEnvironmentAnnotation(string? message)
+		{
+			if (_environment.IsProduction())
+			{
+				return message;
+			}
+			else if (_environment.IsDevelopment())
+			{
+				return $"[DEV] {message}";
+			}
+			else if (_environment.IsStaging())
+			{
+				return $"[STAGING] {message}";
+			}
+			else
+			{
+				return $"[{_environment.EnvironmentName.ToUpperInvariant()}] {message}";
+			}
+		}
+
+		private async Task SendMessageAsync(string recipient, string? text = null, BlockBase[]? blocks = null, BlockKitAttachment[]? attachments = null, bool withEnvironment = true)
+		{
+			if (!IsRecipientAllowed(recipient, text))
+			{
+				return;
+			}
+
+			BlockKitMessage message = new BlockKitMessage();
+			message.Channel = recipient;
+			if (text != null)
+			{
+				message.Text = withEnvironment ? AddEnvironmentAnnotation(text) : text;
+			}
+			if (blocks != null)
+			{
+				message.Blocks.AddRange(blocks);
+			}
+			if (attachments != null)
+			{
+				message.Attachments.AddRange(attachments);
+			}
+
+			await SendRequestAsync<PostMessageResponse>(PostMessageUrl, message);
+		}
+
+		bool IsRecipientAllowed(string recipient, string? description)
+		{
+			if (_allowUsers != null && !_allowUsers.Contains(recipient) && !recipient.StartsWith("#", StringComparison.Ordinal))
+			{
+				_logger.LogDebug("Suppressing message to {Recipient}: {Description}", recipient, description);
+				return false;
+			}
+			return true;
+		}
+
+		private async Task<(MessageStateDocument, bool)> SendOrUpdateMessageAsync(string recipient, string eventId, UserId? userId, string? text = null, BlockBase[]? blocks = null, BlockKitAttachment[]? attachments = null)
+		{
+			BlockKitMessage message = new BlockKitMessage();
+			message.Text = AddEnvironmentAnnotation(text);
+
+			if (blocks != null)
+			{
+				message.Blocks.AddRange(blocks);
+			}
+			if (attachments != null)
+			{
+				message.Attachments.AddRange(attachments);
+			}
+
+			string requestDigest = ContentHash.MD5(JsonSerializer.Serialize(message, new JsonSerializerOptions { DefaultIgnoreCondition = JsonIgnoreCondition.WhenWritingNull })).ToString();
+
+			MessageStateDocument? prevState = await GetMessageStateAsync(recipient, eventId);
+			if (prevState != null && prevState.Digest == requestDigest)
+			{
+				return (prevState, false);
+			}
+
+			(MessageStateDocument state, bool isNew) = await AddOrUpdateMessageStateAsync(recipient, eventId, userId, requestDigest);
+			if (isNew)
+			{
+				_logger.LogInformation("Sending new slack message to {SlackUser} (msg: {MessageId})", recipient, state.Id);
+				message.Channel = recipient;
+				message.Ts = null;
+
+				PostMessageResponse response = await SendRequestAsync<PostMessageResponse>(PostMessageUrl, message);
+				if (String.IsNullOrEmpty(response.Ts) || String.IsNullOrEmpty(response.Channel))
+				{
+					_logger.LogWarning("Missing 'ts' or 'channel' field on slack response");
+				}
+
+				state.Channel = response.Channel ?? String.Empty;
+				state.Ts = response.Ts ?? String.Empty;
+
+				await SetMessageTimestampAsync(state.Id, state.Channel, state.Ts);
+			}
+			else if (!String.IsNullOrEmpty(state.Ts))
+			{
+				_logger.LogInformation("Updating existing slack message {MessageId} for user {SlackUser} ({Channel}, {MessageTs})", state.Id, recipient, state.Channel, state.Ts);
+				message.Channel = state.Channel;
+				message.Ts = state.Ts;
+				await SendRequestAsync<PostMessageResponse>(UpdateMessageUrl, message);
+			}
+			return (state, isNew);
+		}
+
+		private async Task<string?> GetPermalinkAsync(string channel, string messageTs)
+		{
+			string requestUrl = $"{GetPermalinkUrl}?channel={channel}&message_ts={messageTs}";
+			using (HttpRequestMessage sendMessageRequest = new HttpRequestMessage(HttpMethod.Get, requestUrl))
+			{
+				GetPermalinkResponse response = await SendRequestAsync<GetPermalinkResponse>(sendMessageRequest, null);
+				return response?.Permalink;
+			}
+		}
+
+		private async Task<TResponse> SendRequestAsync<TResponse>(string requestUrl, object request) where TResponse : SlackResponse
+		{
+			using (HttpRequestMessage sendMessageRequest = new HttpRequestMessage(HttpMethod.Post, requestUrl))
+			{
+				string requestJson = JsonSerializer.Serialize(request, new JsonSerializerOptions { DefaultIgnoreCondition = JsonIgnoreCondition.WhenWritingNull });
+				using (StringContent messageContent = new StringContent(requestJson, Encoding.UTF8, "application/json"))
+				{
+					sendMessageRequest.Content = messageContent;
+					return await SendRequestAsync<TResponse>(sendMessageRequest, requestJson);
+				}
+			}
+		}
+
+		static readonly string[] s_ignoreErrorCodes =
+		{
+			"no_reaction",
+			"already_reacted"
+		};
+
+		private async Task<TResponse> SendRequestAsync<TResponse>(HttpRequestMessage request, string? requestJson) where TResponse : SlackResponse
+		{
+			using (HttpClient client = new HttpClient())
+			{
+				request.Headers.Add("Authorization", $"Bearer {_settings.SlackToken ?? ""}");
+
+				HttpResponseMessage response = await client.SendAsync(request);
+				byte[] responseBytes = await response.Content.ReadAsByteArrayAsync();
+
+				TResponse responseObject = JsonSerializer.Deserialize<TResponse>(responseBytes)!;
+				if (!responseObject.Ok && !s_ignoreErrorCodes.Contains(responseObject.Error, StringComparer.Ordinal))
+				{
+					_logger.LogError("Failed to send Slack message ({Error}). Request: {Request}. Response: {Response}", responseObject.Error, requestJson, Encoding.UTF8.GetString(responseBytes));
+				}
+
+				return responseObject;
+			}
+		}
+
+		/// <inheritdoc/>
+		protected override async Task ExecuteAsync(CancellationToken stoppingToken)
+		{
+			if (!String.IsNullOrEmpty(_settings.SlackSocketToken))
+			{
+				while (!stoppingToken.IsCancellationRequested)
+				{
+					try
+					{
+						Uri? webSocketUrl = await GetWebSocketUrlAsync(stoppingToken);
+						if (webSocketUrl == null)
+						{
+							await Task.Delay(TimeSpan.FromSeconds(5.0), stoppingToken);
+							continue;
+						}
+						await HandleSocketAsync(webSocketUrl, stoppingToken);
+					}
+					catch (OperationCanceledException)
+					{
+						break;
+					}
+					catch (Exception ex)
+					{
+						_logger.LogError(ex, "Exception while updating Slack socket");
+						await Task.Delay(TimeSpan.FromSeconds(5.0), stoppingToken);
+					}
+				}
+			}
+		}
+
+		/// <summary>
+		/// Get the url for opening a websocket to Slack
+		/// </summary>
+		/// <param name="stoppingToken"></param>
+		/// <returns></returns>
+		private async Task<Uri?> GetWebSocketUrlAsync(CancellationToken stoppingToken)
+		{
+			using HttpClient client = new HttpClient();
+			client.DefaultRequestHeaders.Add("Authorization", $"Bearer {_settings.SlackSocketToken}");
+
+			using FormUrlEncodedContent content = new FormUrlEncodedContent(Array.Empty<KeyValuePair<string?, string?>>());
+			HttpResponseMessage response = await client.PostAsync(new Uri("https://slack.com/api/apps.connections.open"), content, stoppingToken);
+
+			byte[] responseData = await response.Content.ReadAsByteArrayAsync(stoppingToken);
+
+			SocketResponse socketResponse = JsonSerializer.Deserialize<SocketResponse>(responseData)!;
+			if (!socketResponse.Ok)
+			{
+				_logger.LogWarning("Unable to get websocket url: {Response}", Encoding.UTF8.GetString(responseData));
+				return null;
+			}
+
+			return socketResponse.Url;
+		}
+
+		/// <summary>
+		/// Handle the lifetime of a websocket connection
+		/// </summary>
+		/// <param name="socketUrl"></param>
+		/// <param name="stoppingToken"></param>
+		/// <returns></returns>
+		private async Task HandleSocketAsync(Uri socketUrl, CancellationToken stoppingToken)
+		{
+			using ClientWebSocket socket = new ClientWebSocket();
+			await socket.ConnectAsync(socketUrl, stoppingToken);
+
+			byte[] buffer = new byte[2048];
+			while (!stoppingToken.IsCancellationRequested)
+			{
+				// Read the next message
+				int length = 0;
+				for (; ; )
+				{
+					if (length == buffer.Length)
+					{
+						Array.Resize(ref buffer, buffer.Length + 2048);
+					}
+
+					WebSocketReceiveResult result = await socket.ReceiveAsync(new ArraySegment<byte>(buffer, length, buffer.Length - length), stoppingToken);
+					if (result.MessageType == WebSocketMessageType.Close)
+					{
+						return;
+					}
+					length += result.Count;
+
+					if (result.EndOfMessage)
+					{
+						break;
+					}
+				}
+
+				// Get the message data
+				_logger.LogInformation("Slack event: {Message}", Encoding.UTF8.GetString(buffer, 0, length));
+				EventMessage eventMessage = JsonSerializer.Deserialize<EventMessage>(buffer.AsSpan(0, length))!;
+
+				// Acknowledge the message
+				if (eventMessage.EnvelopeId != null)
+				{
+					object response = new { eventMessage.EnvelopeId };
+					await socket.SendAsync(JsonSerializer.SerializeToUtf8Bytes(response), WebSocketMessageType.Text, true, stoppingToken);
+				}
+
+				// Handle the message type
+				if (eventMessage.Type != null)
+				{
+					string type = eventMessage.Type;
+					if (type.Equals("disconnect", StringComparison.Ordinal))
+					{
+						break;
+					}
+					else if (type.Equals("interactive", StringComparison.Ordinal))
+					{
+						await HandleInteractionMessage(eventMessage);
+					}
+					else
+					{
+						_logger.LogDebug("Unhandled event type ({Type})", type);
+					}
+				}
+			}
+		}
+
+		/// <summary>
+		/// Handle a button being clicked
+		/// </summary>
+		/// <param name="message">The event message</param>
+		/// <returns></returns>
+		private async Task HandleInteractionMessage(EventMessage message)
+		{
+			if (message.Payload != null && message.Payload.User != null && message.Payload.User.UserName != null && String.Equals(message.Payload.Type, "block_actions", StringComparison.Ordinal))
+			{
+				foreach (ActionInfo action in message.Payload.Actions)
+				{
+					if (action.Value != null)
+					{
+						if (action.Value.StartsWith("issue_", StringComparison.Ordinal))
+						{
+							await HandleIssueResponseAsync(message.Payload, action);
+						}
+					}
+				}
+			}
+		}
+	}
+}
diff --git a/Engine/Source/Programs/Horde/Horde.Build/Perforce/CommitCollection.cs b/Engine/Source/Programs/Horde/Horde.Build/Perforce/CommitCollection.cs
new file mode 100644
--- /dev/null
+++ b/Engine/Source/Programs/Horde/Horde.Build/Perforce/CommitCollection.cs
@@ -0,0 +1,119 @@
+// Copyright Epic Games, Inc. All Rights Reserved.
+
+using System;
+using System.Collections.Generic;
+using System.Linq;
+using System.Threading.Tasks;
+using Horde.Build.Models;
+using Horde.Build.Server;
+using Horde.Build.Utilities;
+using MongoDB.Bson;
+using MongoDB.Bson.Serialization.Attributes;
+using MongoDB.Driver;
+
+namespace Horde.Build.Commits.Impl
+{
+	using CommitId = ObjectId<ICommit>;
+	using StreamId = StringId<IStream>;
+	using UserId = ObjectId<IUser>;
+
+	/// <summary>
+	/// Concrete implementation of ICommitCollection
+	/// </summary>
+	class CommitCollection : ICommitCollection
+	{
+		class Commit : ICommit
+		{
+			[BsonIgnoreIfDefault]
+			public CommitId Id { get; set; }
+
+			[BsonElement("s")]
+			public StreamId StreamId { get; set; }
+
+			[BsonElement("c")]
+			public int Change { get; set; }
+
+			[BsonElement("oc"), BsonIgnoreIfNull]
+			public int? OriginalChange { get; set; }
+
+			[BsonElement("a")]
+			public UserId AuthorId { get; set; }
+
+			[BsonElement("o"), BsonIgnoreIfNull]
+			public UserId? OwnerId { get; set; }
+
+			[BsonElement("d")]
+			public string Description { get; set; } = String.Empty;
+
+			[BsonElement("p")]
+			public string BasePath { get; set; } = String.Empty;
+
+			[BsonElement("t")]
+			public DateTime DateUtc { get; set; }
+
+			int ICommit.OriginalChange => OriginalChange ?? Change;
+			UserId ICommit.OwnerId => OwnerId ?? AuthorId;
+
+			public Commit()
+			{
+			}
+
+			public Commit(NewCommit newCommit)
+			{
+				Change = newCommit.Change;
+				if (newCommit.OriginalChange != newCommit.Change)
+				{
+					OriginalChange = newCommit.OriginalChange;
+				}
+				StreamId = newCommit.StreamId;
+				AuthorId = newCommit.AuthorId;
+				if (newCommit.OwnerId != newCommit.AuthorId)
+				{
+					OwnerId = newCommit.OwnerId;
+				}
+				Description = newCommit.Description;
+				BasePath = newCommit.BasePath;
+				DateUtc = newCommit.DateUtc;
+			}
+		}
+
+		readonly IMongoCollection<Commit> _commits;
+
+		/// <summary>
+		/// Constructor
+		/// </summary>
+		public CommitCollection(MongoService mongoService)
+		{
+			_commits = mongoService.GetCollection<Commit>("Commits", keys => keys.Ascending(x => x.StreamId).Descending(x => x.Change), unique: true);
+		}
+
+		/// <inheritdoc/>
+		public async Task<ICommit> AddOrReplaceAsync(NewCommit newCommit)
+		{
+			Commit commit = new Commit(newCommit);
+			FilterDefinition<Commit> filter = Builders<Commit>.Filter.Expr(x => x.StreamId == newCommit.StreamId && x.Change == newCommit.Change);
+			return await _commits.FindOneAndReplaceAsync(filter, new Commit(newCommit), new FindOneAndReplaceOptions<Commit> { IsUpsert = true, ReturnDocument = ReturnDocument.After });
+		}
+
+		/// <inheritdoc/>
+		public async Task<ICommit?> GetCommitAsync(CommitId id)
+		{
+			return await _commits.Find(x => x.Id == id).FirstOrDefaultAsync();
+		}
+
+		/// <inheritdoc/>
+		public async Task<List<ICommit>> FindCommitsAsync(StreamId streamId, int? minChange = null, int? maxChange = null, int? index = null, int? count = null)
+		{
+			FilterDefinition<Commit> filter = Builders<Commit>.Filter.Eq(x => x.StreamId, streamId);
+			if (maxChange != null)
+			{
+				filter &= Builders<Commit>.Filter.Lte(x => x.Change, maxChange.Value);
+			}
+			if (minChange != null)
+			{
+				filter &= Builders<Commit>.Filter.Lte(x => x.Change, minChange.Value);
+			}
+			return await _commits.Find(filter).SortByDescending(x => x.Change).Range(index, count).ToListAsync<Commit, ICommit>();
+		}
+	}
+}
