Horde: Finish re-org of Horde solution started a while ago.

diff --git a/Engine/Source/Programs/Horde/Horde.Build/Perforce/PerforceService.cs b/Engine/Source/Programs/Horde/Horde.Build/Perforce/PerforceService.cs
new file mode 100644
--- /dev/null
+++ b/Engine/Source/Programs/Horde/Horde.Build/Perforce/PerforceService.cs
@@ -0,0 +1,1710 @@
+// Copyright Epic Games, Inc. All Rights Reserved.
+
+using System;
+using System.Collections.Generic;
+using System.Diagnostics.CodeAnalysis;
+using System.Globalization;
+using System.IO;
+using System.Linq;
+using System.Net;
+using System.Reflection;
+using System.Runtime.InteropServices;
+using System.Text;
+using System.Text.RegularExpressions;
+using System.Threading.Tasks;
+using EpicGames.Core;
+using EpicGames.Perforce;
+using Horde.Build.Collections;
+using Horde.Build.Models;
+using Horde.Build.Server;
+using Horde.Build.Utilities;
+using Microsoft.Extensions.Caching.Memory;
+using Microsoft.Extensions.Logging;
+using Microsoft.Extensions.Options;
+using OpenTracing;
+using OpenTracing.Util;
+
+namespace Horde.Build.Services
+{
+	using P4 = Perforce.P4;
+
+	static class PerforceExtensions
+	{
+		static readonly FieldInfo s_field = typeof(P4.Changelist).GetField("_baseForm", BindingFlags.Instance | BindingFlags.NonPublic)!;
+
+		public static string GetPath(this P4.Changelist changelist)
+		{
+			P4.FormBase? formBase = s_field.GetValue(changelist) as P4.FormBase;
+			object? pathValue = null;
+			formBase?.TryGetValue("path", out pathValue);
+			return (pathValue as string) ?? "//...";
+		}
+	}
+
+	/// <summary>
+	/// P4API implementation of the Perforce service
+	/// </summary>
+	class PerforceService : IPerforceService, IDisposable
+	{
+		class CachedTicketInfo
+		{
+			public IPerforceServer _server;
+			public string _userName;
+			public P4.Credential _ticket;
+
+			public CachedTicketInfo(IPerforceServer server, string userName, P4.Credential ticket)
+			{
+				_server = server;
+				_userName = userName;
+				_ticket = ticket;
+			}
+		}
+
+		readonly PerforceLoadBalancer _loadBalancer;
+		readonly LazyCachedValue<Task<Globals>> _cachedGlobals;
+		readonly ILogger _logger;
+		readonly bool _haveNativePerforceApi;
+
+		// Useful overrides for local debugging with read-only data
+		readonly string? _perforceServerOverride;
+		readonly string? _perforceUserOverride;
+
+		/// <summary>
+		/// Object used for controlling access to the access user tickets
+		/// </summary>
+		static readonly object s_ticketLock = new object();
+
+		/// <summary>
+		/// Object used for controlling access to the p4 command output
+		/// </summary>
+		static readonly object s_p4LogLock = new object();
+
+		/// <summary>
+		/// Native -> managed debug logging callback
+		/// </summary>
+		readonly P4.P4CallBacks.LogMessageDelegate _logBridgeDelegate;
+
+		/// <summary>
+		/// The server settings
+		/// </summary>
+		readonly ServerSettings _settings;
+		readonly Dictionary<string, Dictionary<string, CachedTicketInfo>> _clusterTickets = new Dictionary<string, Dictionary<string, CachedTicketInfo>>(StringComparer.OrdinalIgnoreCase);
+		readonly IUserCollection _userCollection;
+		readonly MemoryCache _userCache = new MemoryCache(new MemoryCacheOptions { SizeLimit = 2000 });
+
+		/// <summary>
+		/// Constructor
+		/// </summary>
+		public PerforceService(PerforceLoadBalancer loadBalancer, MongoService mongoService, IUserCollection userCollection, IOptions<ServerSettings> settings, ILogger<PerforceService> logger)
+		{
+			_loadBalancer = loadBalancer;
+			_cachedGlobals = new LazyCachedValue<Task<Globals>>(() => mongoService.GetGlobalsAsync(), TimeSpan.FromSeconds(30.0));
+			_userCollection = userCollection;
+			_settings = settings.Value;
+			_logger = logger;
+
+			_logBridgeDelegate = new P4.P4CallBacks.LogMessageDelegate(LogBridgeMessage);
+
+			_haveNativePerforceApi = !(RuntimeInformation.IsOSPlatform(OSPlatform.OSX) && RuntimeInformation.OSArchitecture == Architecture.Arm64);
+			if (_haveNativePerforceApi)
+			{
+				P4.P4Debugging.SetBridgeLogFunction(_logBridgeDelegate);
+				P4.LogFile.SetLoggingFunction(LogPerforce);
+			}
+
+			if(settings.Value.UseLocalPerforceEnv)
+			{
+				IPerforceSettings perforceSettings = PerforceSettings.Default;
+				_perforceServerOverride = perforceSettings.ServerAndPort;
+				_perforceUserOverride = perforceSettings.UserName;
+			}
+		}
+
+		public void Dispose()
+		{
+			_userCache.Dispose();
+		}
+
+		public async ValueTask<IUser> FindOrAddUserAsync(string clusterName, string userName)
+		{
+			using IScope scope = GlobalTracer.Instance.BuildSpan("PerforceService.FindOrAddUserAsync").StartActive();
+			scope.Span.SetTag("ClusterName", clusterName);
+			scope.Span.SetTag("UserName", userName);
+			
+			IUser? user;
+			if (!_userCache.TryGetValue((clusterName, userName), out user))
+			{
+				user = await _userCollection.FindUserByLoginAsync(userName);
+				if (user == null)
+				{
+					PerforceUserInfo? userInfo = await GetUserInfoAsync(clusterName, userName);
+					user = await _userCollection.FindOrAddUserByLoginAsync(userName, userInfo?.FullName, userInfo?.Email);
+				}
+
+				using (ICacheEntry entry = _userCache.CreateEntry((clusterName, userName)))
+				{
+					entry.SetValue(user);
+					entry.SetSize(1);
+					entry.SetAbsoluteExpiration(TimeSpan.FromDays(1.0));
+				}
+			}
+			return user!;
+		}
+
+		async Task<PerforceCluster> GetClusterAsync(string? clusterName, string? serverAndPort = null)
+		{
+			Globals globals = await _cachedGlobals.GetCached();
+
+			PerforceCluster? cluster = globals.FindPerforceCluster(clusterName, serverAndPort);
+			if (cluster == null)
+			{
+				throw new Exception($"Unknown Perforce cluster '{clusterName}'");
+			}
+
+			return cluster;
+		}
+
+		async ValueTask<string> SelectServerAddressAsync(PerforceCluster cluster)
+		{
+			return (await SelectServer(cluster)).ServerAndPort;
+		}
+
+		async Task<IPerforceServer> SelectServer(PerforceCluster cluster)
+		{
+			using IScope scope = GlobalTracer.Instance.BuildSpan("PerforceService.SelectServer").StartActive();
+			scope.Span.SetTag("ClusterName", cluster.Name);
+			
+			IPerforceServer? server = await _loadBalancer.SelectServerAsync(cluster);
+			if (server == null)
+			{
+				throw new Exception($"Unable to select server from '{cluster.Name}'");
+			}
+			return server;
+		}
+
+		[SuppressMessage("Microsoft.Reliability", "CA2000:DisposeObjectsBeforeLosingScope")]
+		static P4.Repository CreateConnection(IPerforceServer server, string? userName, string? ticket)
+		{
+			using IScope scope = GlobalTracer.Instance.BuildSpan("PerforceService.CreateConnection").StartActive();
+			scope.Span.SetTag("Cluster", server.Cluster);
+			scope.Span.SetTag("ServerAndPort", server.ServerAndPort);
+			scope.Span.SetTag("UserName", userName);
+			
+			P4.Repository repository = new P4.Repository(new P4.Server(new P4.ServerAddress(server.ServerAndPort)));
+			try
+			{
+				P4.Connection connection = repository.Connection;
+				if (userName != null)
+				{
+					connection.UserName = userName;
+				}
+
+				P4.Options options = new P4.Options();
+				if (ticket != null)
+				{
+					options["Ticket"] = ticket;
+				}
+
+				// connect to the server
+				if (!connection.Connect(options))
+				{
+					throw new Exception("Unable to get P4 server connection");
+				}
+			}
+			catch
+			{
+				repository.Dispose();
+				throw;
+			}
+			return repository;
+		}
+
+		/// <summary>
+		/// 
+		/// </summary>
+		/// <param name="clusterName"></param>
+		/// <returns></returns>
+		public async Task<IPerforceConnection?> GetServiceUserConnection(string? clusterName)
+		{
+			using IScope scope = GlobalTracer.Instance.BuildSpan("PerforceService.GetServiceUserConnection").StartActive();
+			scope.Span.SetTag("ClusterName", clusterName);
+
+			PerforceCluster cluster = await GetClusterAsync(clusterName);
+
+			string? userName, password;
+			if (_perforceUserOverride != null)
+			{
+				userName = _perforceUserOverride;
+				password = null;
+			}
+			else if (cluster.ServiceAccount != null)
+			{
+				PerforceCredentials? credentials = cluster.Credentials.FirstOrDefault(x => x.UserName.Equals(cluster.ServiceAccount, StringComparison.OrdinalIgnoreCase));
+				if (credentials == null)
+				{
+					throw new Exception($"No credentials defined for {cluster.ServiceAccount} on {cluster.Name}");
+				}
+
+				userName = credentials.UserName;
+				password = credentials.Password;
+			}
+			else
+			{
+				userName = PerforceSettings.Default.UserName;
+				password = null;
+			}
+
+			string serverAndPort;
+			if (_perforceServerOverride != null)
+			{
+				serverAndPort = _perforceServerOverride;
+			}
+			else
+			{
+				serverAndPort = await SelectServerAddressAsync(cluster);
+			}
+
+			PerforceSettings settings = new PerforceSettings(serverAndPort, userName);
+			settings.Password = password;
+			settings.AppName = "Horde.Build";
+			settings.ClientName = "__DOES_NOT_EXIST__";
+			settings.PreferNativeClient = true;
+
+			return await PerforceConnection.CreateAsync(settings, _logger);
+		}
+
+		async Task<P4.Repository> GetServiceUserConnection(PerforceCluster cluster)
+		{
+			using IScope scope = GlobalTracer.Instance.BuildSpan("PerforceService.GetServiceUserConnectionAsync").StartActive();
+			scope.Span.SetTag("ClusterName", cluster.Name);
+			
+			IPerforceServer server = await SelectServer(cluster);
+			return GetServiceUserConnection(cluster, server);
+		}
+
+		P4.Repository GetServiceUserConnection(PerforceCluster cluster, IPerforceServer server)
+		{
+			using IScope scope = GlobalTracer.Instance.BuildSpan("PerforceService.GetServiceUserConnection").StartActive();
+			scope.Span.SetTag("ClusterName", cluster.Name);
+
+			if (cluster.Name == PerforceCluster.DefaultName && _settings.P4BridgeServiceUsername != null && _settings.P4BridgeServicePassword != null)
+			{
+				return CreateConnection(server, _settings.P4BridgeServiceUsername, _settings.P4BridgeServicePassword);
+			}
+
+			string? userName = null;
+			string? password = null;
+			if (cluster.ServiceAccount != null)
+			{
+				PerforceCredentials? credentials = cluster.Credentials.FirstOrDefault(x => x.UserName.Equals(cluster.ServiceAccount, StringComparison.OrdinalIgnoreCase));
+				if (credentials == null)
+				{
+					throw new Exception($"No credentials defined for {cluster.ServiceAccount} on {cluster.Name}");
+				}
+				userName = credentials.UserName;
+				password = credentials.Password;
+			}
+			return CreateConnection(server, userName, password);
+		}
+
+		async Task<CachedTicketInfo> GetImpersonateCredential(PerforceCluster cluster, string impersonateUser)
+		{
+			using IScope scope = GlobalTracer.Instance.BuildSpan("PerforceService.GetImpersonateCredential").StartActive();
+			scope.Span.SetTag("ClusterName", cluster.Name);
+			scope.Span.SetTag("ImpersonateUser", impersonateUser);
+			
+			if (!cluster.CanImpersonate)
+			{
+				throw new Exception($"Service account required to impersonate user {impersonateUser}");
+			}
+
+			CachedTicketInfo? ticketInfo = null;
+
+			Dictionary<string, CachedTicketInfo>? userTickets;
+			lock (s_ticketLock)
+			{
+				// Check if we have a ticket
+				if (!_clusterTickets.TryGetValue(cluster.Name, out userTickets))
+				{
+					userTickets = new Dictionary<string, CachedTicketInfo>(StringComparer.OrdinalIgnoreCase);
+					_clusterTickets[cluster.Name] = userTickets;
+				}
+				if (userTickets.TryGetValue(impersonateUser, out ticketInfo))
+				{
+					// if the credential expires within the next 15 minutes, refresh
+					TimeSpan time = new TimeSpan(0, 15, 0);
+					if (ticketInfo._ticket.Expires.Subtract(time) <= DateTime.UtcNow)
+					{
+						userTickets.Remove(impersonateUser);
+						ticketInfo = null;
+					}
+				}
+			}
+
+			if (ticketInfo == null)
+			{
+				IPerforceServer server = await SelectServer(cluster);
+
+				P4.Credential? credential;
+				using (P4.Repository repository = GetServiceUserConnection(cluster, server))
+				{
+					credential = repository.Connection.Login(null, new P4.LoginCmdOptions(P4.LoginCmdFlags.AllHosts | P4.LoginCmdFlags.DisplayTicket, null), impersonateUser);
+				}
+				if (credential == null)
+				{
+					throw new Exception($"GetImpersonateCredential - Unable to get impersonation credential for {impersonateUser} from {Dns.GetHostName()}");
+				}
+
+				ticketInfo = new CachedTicketInfo(server, impersonateUser, credential);
+
+				lock (s_ticketLock)
+				{
+					userTickets[impersonateUser] = ticketInfo;
+				}
+			}
+
+			return ticketInfo;
+		}
+
+		async Task<P4.Repository> GetImpersonatedConnection(PerforceCluster cluster, string impersonateUser)
+		{
+			using IScope scope = GlobalTracer.Instance.BuildSpan("PerforceService.GetImpersonatedConnection").StartActive();
+			scope.Span.SetTag("ClusterName", cluster.Name);
+			scope.Span.SetTag("ImpersonateUser", impersonateUser);
+			
+			CachedTicketInfo ticketInfo = await GetImpersonateCredential(cluster, impersonateUser);
+			return CreateConnection(ticketInfo._server, ticketInfo._userName, ticketInfo._ticket.Ticket);
+		}
+		
+		async Task<string?> GetPortFromChange(PerforceCluster cluster, int change)
+		{
+			using IScope scope = GlobalTracer.Instance.BuildSpan("PerforceService.GetPortFromChange").StartActive();
+			scope.Span.SetTag("ClusterName", cluster.Name);
+			scope.Span.SetTag("Change", change);
+
+			string? clientPort = null;
+
+			using (P4.Repository repository = await GetConnection(cluster, noClient: true))
+			{				
+				P4.Changelist changelist = repository.GetChangelist(change);
+				if (changelist == null || String.IsNullOrEmpty(changelist.ClientId))
+				{
+					throw new Exception($"GetPortFromChange - Unable to get changelist for {change}");
+				}
+
+				P4.Client client = repository.GetClient(changelist.ClientId);
+
+				if (client == null)
+				{
+					throw new Exception($"GetPortFromChange - Unable to get client for {change}");
+				}
+
+				// Check whether not restricted to a specific server
+				if (String.IsNullOrEmpty(client.ServerID))
+				{
+					return null;
+				}
+
+				List<string> args = new List<string> { "-o", client.ServerID };
+
+				using (P4.P4Command command = new P4.P4Command(repository, "server", true, args.ToArray()))
+				{
+					P4.P4CommandResult result = command.Run();
+
+					if (!result.Success || result.TaggedOutput == null || result.TaggedOutput.Count == 0)
+					{
+						throw new Exception($"Unable to get tagged output for {client.ServerID}");
+					}
+
+					foreach (P4.TaggedObject taggedObject in result.TaggedOutput)
+					{
+						if (taggedObject.TryGetValue("Address", out clientPort))
+						{
+							break;
+						}
+					}
+				}
+			}
+
+			if (String.IsNullOrEmpty(clientPort))
+			{
+				throw new Exception($"GetPortFromChange - Unable to get client port for {change}");
+			}
+
+			return clientPort;
+		}
+
+		[SuppressMessage("Microsoft.Reliability", "CA2000:DisposeObjectsBeforeLosingScope")]
+		async Task<P4.Repository> GetConnection(PerforceCluster cluster, string? stream = null, string? username = null, bool readOnly = true, bool createChange = false, int? clientFromChange = null, bool useClientFromChange = false, int? usePortFromChange = null, string? clientId = null, bool noClient = false)
+		{
+			using IScope scope = GlobalTracer.Instance.BuildSpan("PerforceService.GetConnection").StartActive();
+			scope.Span.SetTag("ClusterName", cluster.Name);
+			scope.Span.SetTag("Stream", stream);
+			scope.Span.SetTag("Username", username);
+			scope.Span.SetTag("NoClient", noClient);
+			scope.Span.SetTag("ClientId", clientId);
+
+			if (usePortFromChange.HasValue)
+			{
+				try
+				{
+					string? clientPort = await GetPortFromChange(cluster, usePortFromChange.Value);
+					if (clientPort != null)
+					{
+						PerforceCluster clientCluster = await GetClusterAsync(null, clientPort);
+						if (!String.Equals(clientCluster.Name, cluster.Name, StringComparison.OrdinalIgnoreCase))
+						{
+							_logger.LogInformation("Perforce: Overriding Cluster {Cluster} with Client Cluster {ClientCluster} for server restricted change {Change}", cluster.Name, clientCluster.Name, usePortFromChange.Value);
+							cluster = clientCluster;
+						}
+					}
+					else
+					{
+						_logger.LogInformation("Change {UsePortFromChange} is not restricted to a server", usePortFromChange.Value);
+					}
+				}
+				catch (Exception ex)
+				{
+					_logger.LogError(ex, "Unable to get client port for changelist {Change} using cluster {ClusterName}", usePortFromChange, cluster.Name);
+				}
+			}
+
+			P4.Repository repository;
+			if (username == null || !cluster.CanImpersonate || username.Equals(cluster.ServiceAccount, StringComparison.OrdinalIgnoreCase))
+			{
+				repository = await GetServiceUserConnection(cluster);
+			}
+			else
+			{
+				repository = await GetImpersonatedConnection(cluster, username);
+			}
+
+			if (!noClient)
+			{
+				try
+				{
+					if (clientId != null)
+					{
+						repository.Connection.SetClient(clientId);
+					}
+					else
+					{
+						P4.Client client = GetOrCreateClient(cluster.ServiceAccount, repository, stream, username, readOnly, createChange, clientFromChange, useClientFromChange);
+						repository.Connection.SetClient(client.Name);
+
+					}
+
+					if (useClientFromChange)
+					{
+						string? clientHost = repository.Connection.Client.Host;
+						if (!String.IsNullOrEmpty(clientHost))
+						{
+							repository.Connection.getP4Server().SetConnectionHost(clientHost);
+						}
+					}					
+				}
+				catch
+				{
+					repository.Dispose();
+					throw;
+				}
+			}
+
+			if (repository == null)
+			{
+				throw new Exception($"Unable to get connection for Stream:{stream} Username:{username} ReadOnly:{readOnly} CreateChange:{createChange} ClientFromChange:{clientFromChange} UseClientFromChange:{useClientFromChange} UsePortFromChange:{usePortFromChange}");
+			}
+
+			repository.Connection.CommandEcho += LogPerforceCommand;
+
+			return repository;
+
+		}
+
+		static string GetClientName(string? serviceUserName, string stream, bool readOnly, bool createChange, string? username = null)
+		{
+			string clientName = $"zzt-horde-p4bridge-{Dns.GetHostName()}-{serviceUserName ?? "default"}-{stream.Replace("/", "+", StringComparison.OrdinalIgnoreCase)}";
+
+			if (!readOnly)
+			{
+				clientName += "-write";
+			}
+
+			if (createChange)
+			{
+				clientName += "-create";
+			}
+
+			if (!String.IsNullOrEmpty(username))
+			{
+				clientName += "-" + username;
+			}
+
+			return clientName;
+
+		}
+
+		static void ResetClient(P4.Repository repository)
+		{
+			repository.Connection.Client.RevertFiles(new P4.Options(), new P4.DepotPath("//..."));
+
+			IList<P4.Changelist> changes = repository.GetChangelists(new P4.ChangesCmdOptions(P4.ChangesCmdFlags.None, repository.Connection.Client.Name, 100, P4.ChangeListStatus.Pending, null));
+			if (changes != null)
+			{
+				foreach (P4.Changelist change in changes)
+				{
+					repository.DeleteChangelist(change, new P4.Options());
+				}
+			}
+		}
+
+		static P4.Client GetOrCreateClient(string? serviceUserName, P4.Repository repository, string? stream, string? username = null, bool readOnly = true, bool createChange = false, int? clientFromChange = null, bool useClientFromChange = false)
+		{
+			using IScope scope = GlobalTracer.Instance.BuildSpan("PerforceService.GetOrCreateClient").StartActive();
+			scope.Span.SetTag("ServiceUserName", serviceUserName);
+			scope.Span.SetTag("Stream", stream);
+			scope.Span.SetTag("Username", username);
+			
+			P4.Client? client = null;
+			P4.Changelist? changelist = null;
+
+			string? clientHost = null;
+
+			if (clientFromChange.HasValue)
+			{
+				changelist = repository.GetChangelist(clientFromChange.Value);
+
+				if (changelist == null)
+				{
+					throw new Exception($"Unable to get changelist for client {clientFromChange}");
+				}
+
+				client = repository.GetClient(changelist.ClientId);
+
+				if (client == null)
+				{
+					throw new Exception($"Unable to get client for id {changelist.ClientId}");
+				}
+
+				stream = client.Stream;
+				username = client.OwnerName;
+				clientHost = client.Host;
+
+			}
+
+			if (!useClientFromChange)
+			{
+				if (stream == null)
+				{
+					throw new Exception("Stream required for client");
+
+				}
+				string clientName = GetClientName(serviceUserName, stream, readOnly, createChange, username);
+
+				IList<P4.Client>? clients = repository.GetClients(new P4.ClientsCmdOptions(P4.ClientsCmdFlags.None, username, clientName, 1, stream));
+
+				if (clients != null && clients.Count == 1)
+				{
+					client = clients[0];
+					if (!String.IsNullOrEmpty(client.Host))
+					{
+						client.Host = String.Empty;
+						repository.UpdateClient(client);
+					}
+				}
+				else
+				{
+
+					P4.Client newClient = new P4.Client();
+
+					newClient.Name = clientName;
+					newClient.Root = RuntimeInformation.IsOSPlatform(OSPlatform.Linux) ? $"/tmp/{clientName}/" : $"{Path.GetTempPath()}{clientName}\\";
+					newClient.Stream = stream;
+					if (readOnly)
+					{
+						newClient.ClientType = P4.ClientType.@readonly;
+					}
+
+					newClient.OwnerName = username ?? serviceUserName;
+
+					client = repository.CreateClient(newClient);
+
+					if (client == null)
+					{
+						throw new Exception($"Unable to create client for ${stream} : {username ?? serviceUserName}");
+					}
+				}
+			}
+
+			if (client == null)
+			{
+				throw new Exception($"Unable to create client for Stream:{stream} Username:{username} ReadOnly:{readOnly} CreateChange:{createChange} ClientFromChange:{clientFromChange} UseClientFromChange:{useClientFromChange}");
+			}
+
+			return client;
+
+		}
+
+		/// <summary>
+		/// Logs perforce commands 
+		/// </summary>
+		/// <param name="log">The p4 command log info</param>
+		static void LogPerforceCommand(string log)
+		{
+			lock (s_p4LogLock)
+			{
+				Serilog.Log.Information("Perforce: {Message}", log);
+			}
+		}
+
+		/// <summary>
+		/// Perforce logging funcion
+		/// </summary>
+		/// <param name="logLevel">The level whether error, warning, message, or debug</param>
+		/// <param name="source">The log source</param>
+		/// <param name="message">The log message</param>
+		static void LogPerforce(int logLevel, string source, string message)
+		{
+			lock (s_p4LogLock)
+			{
+				switch (logLevel)
+				{
+					case 0:
+					case 1:
+						Serilog.Log.Error("Perforce (Error): {Message} {Source} : ", message, source);
+						break;
+					case 2:
+						Serilog.Log.Warning("Perforce (Warning): {Message} {Source} : ", message, source);
+						break;
+					case 3:
+						Serilog.Log.Information("Perforce: {Message} {Source} : ", message, source);
+						break;
+					default:
+						Serilog.Log.Debug("Perforce (Debug): {Message} {Source} : ", message, source);
+						break;
+				}
+			}
+		}
+
+		static void LogBridgeMessage(int logLevel, string filename, int line, string message)
+		{
+
+			// Note, we do not get log level 4 unless it is defined in native code as it is very, very spammy (P4BridgeServer.cpp)
+
+			// remove the full path to the source, keep just the file name
+			string fileName = Path.GetFileName(filename);
+
+			string category = String.Format(CultureInfo.CurrentCulture, "P4Bridge({0}:{1})", fileName, line);
+
+			P4.LogFile.LogMessage(logLevel, category, message);
+		}
+
+		class StreamView : IStreamView
+		{
+			readonly P4.P4MapApi _mapApi;
+
+			public StreamView(P4.P4MapApi mapApi)
+			{
+				_mapApi = mapApi;
+			}
+
+			public void Dispose()
+			{
+				_mapApi.Dispose();
+			}
+
+			public bool TryGetStreamPath(string depotPath, out string streamPath)
+			{
+				streamPath = _mapApi.Translate(depotPath, P4.P4MapApi.Direction.LeftRight);
+				return streamPath != null;
+			}
+		}
+
+		/// <inheritdoc/>
+		public async Task<IStreamView> GetStreamViewAsync(string clusterName, string streamName)
+		{
+			using IScope scope = GlobalTracer.Instance.BuildSpan("PerforceService.GetStreamViewAsync").StartActive();
+			scope.Span.SetTag("ClusterName", clusterName);
+			scope.Span.SetTag("StreamName", streamName);
+			
+			PerforceCluster cluster = await GetClusterAsync(clusterName);
+			using (P4.Repository repository = await GetConnection(cluster, noClient: true))
+			{
+				P4.Stream stream = repository.GetStream(streamName, new P4.StreamCmdOptions(P4.StreamCmdFlags.View, null, null));
+
+				P4.P4MapApi? mapApi = null;
+				try
+				{
+					mapApi = new P4.P4MapApi(repository.Server.Metadata.UnicodeEnabled);
+					foreach (P4.MapEntry entry in stream.View)
+					{
+						P4.P4MapApi.Type mapType;
+						switch (entry.Type)
+						{
+							case P4.MapType.Include:
+								mapType = P4.P4MapApi.Type.Include;
+								break;
+							case P4.MapType.Exclude:
+								mapType = P4.P4MapApi.Type.Exclude;
+								break;
+							case P4.MapType.Overlay:
+								mapType = P4.P4MapApi.Type.Overlay;
+								break;
+							default:
+								throw new Exception($"Invalid map type: {entry.Type}");
+						}
+						mapApi.Insert(entry.Left.ToString(), entry.Right.ToString(), mapType);
+					}
+
+					StreamView view = new StreamView(mapApi);
+					mapApi = null;
+					return view;
+				}
+				finally
+				{
+					mapApi?.Dispose();
+				}
+			}
+		}
+
+		static int GetSyncRevision(string path, P4.FileAction headAction, int headRev)
+		{
+			switch (headAction)
+			{
+				case P4.FileAction.None:
+				case P4.FileAction.Add:
+				case P4.FileAction.Branch:
+				case P4.FileAction.MoveAdd:
+				case P4.FileAction.Edit:
+				case P4.FileAction.Integrate:
+					return headRev;
+				case P4.FileAction.Delete:
+				case P4.FileAction.MoveDelete:
+				case P4.FileAction.Purge:
+					return -1;
+				default:
+					throw new Exception($"Unrecognized P4 file change type '{headAction}' for file {path}#{headRev}");
+			}
+		}
+
+		/// <summary>
+		/// Creates a <see cref="ChangeFile"/> from a <see cref="P4.FileMetaData"/>
+		/// </summary>
+		/// <param name="relativePath"></param>
+		/// <param name="metaData"></param>
+		/// <returns></returns>
+		static ChangeFile CreateChangeFile(string relativePath, P4.FileMetaData metaData)
+		{
+			using IScope scope = GlobalTracer.Instance.BuildSpan("PerforceService.CreateChangeFile (FileMetaData)").StartActive();
+			scope.Span.SetTag("RelativePath", relativePath);
+			
+			int revision = GetSyncRevision(metaData.DepotPath.Path, metaData.HeadAction, metaData.HeadRev);
+			Md5Hash? digest = String.IsNullOrEmpty(metaData.Digest) ? (Md5Hash?)null : Md5Hash.Parse(metaData.Digest);
+			return new ChangeFile(relativePath, metaData.DepotPath.Path, revision, metaData.FileSize, digest, (metaData.Type ?? metaData.HeadType).ToString());
+		}
+
+		/// <summary>
+		/// Creates a <see cref="ChangeFile"/> from a <see cref="P4.FileMetaData"/>
+		/// </summary>
+		/// <param name="relativePath"></param>
+		/// <param name="metaData"></param>
+		/// <returns></returns>
+		static ChangeFile CreateChangeFile(string relativePath, P4.ShelvedFile metaData)
+		{
+			using IScope scope = GlobalTracer.Instance.BuildSpan("PerforceService.CreateChangeFile (ShelvedFile)").StartActive();
+			scope.Span.SetTag("RelativePath", relativePath);
+			
+			int revision = GetSyncRevision(metaData.Path.Path, metaData.Action, metaData.Revision);
+			Md5Hash? digest = String.IsNullOrEmpty(metaData.Digest) ? (Md5Hash?)null : Md5Hash.Parse(metaData.Digest);
+			return new ChangeFile(relativePath, metaData.Path.Path, revision, metaData.Size, digest, metaData.Type.ToString());
+		}
+
+		/// <inheritdoc/>
+		public async Task<List<ChangeSummary>> GetChangesAsync(string clusterName, int? minChange, int? maxChange, int maxResults)
+		{
+			using IScope scope = GlobalTracer.Instance.BuildSpan("PerforceService.GetChangesAsync").StartActive();
+			scope.Span.SetTag("ClusterName", clusterName);
+			scope.Span.SetTag("MinChange", minChange ?? -1);
+			scope.Span.SetTag("MaxResults", maxResults);
+			
+			PerforceCluster cluster = await GetClusterAsync(clusterName);
+			using (P4.Repository repository = await GetConnection(cluster, noClient: true))
+			{
+				P4.ChangesCmdOptions options = new P4.ChangesCmdOptions(P4.ChangesCmdFlags.IncludeTime | P4.ChangesCmdFlags.FullDescription, null, maxResults, P4.ChangeListStatus.Submitted, null);
+
+				IList<P4.Changelist> changelists = repository.GetChangelists(options, new P4.FileSpec(new P4.DepotPath(GetFilter("//...", minChange, maxChange)), null, null, null));
+
+				List<ChangeSummary> changes = new List<ChangeSummary>();
+				if (changelists != null)
+				{
+					foreach (P4.Changelist changelist in changelists)
+					{
+						IUser user = await FindOrAddUserAsync(clusterName, changelist.OwnerName);
+						changes.Add(new ChangeSummary(changelist.Id, user, changelist.GetPath(), changelist.Description));
+					}
+				}
+				return changes;
+			}
+		}
+
+		/// <inheritdoc/>
+		public async Task<List<ChangeSummary>> GetChangesAsync(string clusterName, string streamName, int? minChange, int? maxChange, int results, string? impersonateUser)
+		{
+			using IScope scope = GlobalTracer.Instance.BuildSpan("PerforceService.GetChangesAsync").StartActive();
+			scope.Span.SetTag("ClusterName", clusterName);
+			scope.Span.SetTag("MinChange", minChange ?? -1);
+			scope.Span.SetTag("MaxChange", maxChange ?? -1);
+			scope.Span.SetTag("Results", results);
+			scope.Span.SetTag("ImpersonateUser", impersonateUser);
+			
+			PerforceCluster cluster = await GetClusterAsync(clusterName);
+			using (P4.Repository repository = await GetConnection(cluster, streamName, impersonateUser))
+			{
+
+				P4.ChangesCmdOptions options = new P4.ChangesCmdOptions(P4.ChangesCmdFlags.IncludeTime | P4.ChangesCmdFlags.FullDescription, null, results, P4.ChangeListStatus.Submitted, null);
+
+				string filter = GetFilter($"//{repository.Connection.Client.Name}/...", minChange, maxChange);
+
+				P4.FileSpec fileSpec = new P4.FileSpec(new P4.DepotPath(filter), null, null, null);
+
+				IList<P4.Changelist> changelists = repository.GetChangelists(options, fileSpec);
+
+				List<ChangeSummary> changes = new List<ChangeSummary>();
+
+				if (changelists != null)
+				{
+					foreach (P4.Changelist changelist in changelists)
+					{
+						IUser user = await FindOrAddUserAsync(clusterName, changelist.OwnerName);
+						changes.Add(new ChangeSummary(changelist.Id, user, changelist.GetPath(), changelist.Description));
+					}
+				}
+
+				return changes;
+			}
+		}
+
+		/// <inheritdoc/>
+		public async Task<ChangeDetails> GetChangeDetailsAsync(string clusterName, string streamName, int changeNumber)
+		{
+			using IScope scope = GlobalTracer.Instance.BuildSpan("PerforceService.GetChangeDetailsAsync").StartActive();
+			scope.Span.SetTag("ClusterName", clusterName);
+			scope.Span.SetTag("StreamName", streamName);
+			scope.Span.SetTag("ChangeNumber", changeNumber);
+			
+			PerforceCluster cluster = await GetClusterAsync(clusterName);
+			using (P4.Repository repository = await GetConnection(cluster, noClient: true))
+			{
+				List<ChangeDetails> results = new List<ChangeDetails>();
+
+				List<string> args = new List<string> { "-s", $"{changeNumber}" };
+
+				using (P4.P4Command command = new P4.P4Command(repository, "describe", true, args.ToArray()))
+				{
+					P4.P4CommandResult result = command.Run();
+
+					if (!result.Success || result.TaggedOutput == null || result.TaggedOutput.Count == 0)
+					{
+						throw new Exception("Unable to get changes");
+					}
+
+					List<P4.Changelist> changelists = new List<P4.Changelist>() { };
+
+					bool dstMismatch = false;
+					string offset = String.Empty;
+
+					P4.Server server = repository.Server;
+					if (server != null && server.Metadata != null)
+					{
+						offset = server.Metadata.DateTimeOffset;
+						dstMismatch = P4.FormBase.DSTMismatch(server.Metadata);
+					}
+
+					P4.TaggedObject taggedObject = result.TaggedOutput[0];
+					{
+						List<ChangeFile> files = new List<ChangeFile>();
+
+						P4.Changelist change = new P4.Changelist();
+						change.FromChangeCmdTaggedOutput(taggedObject, false, offset, dstMismatch);
+
+						foreach (P4.FileMetaData describeFile in change.Files)
+						{
+							string? relativePath;
+							if (TryGetStreamRelativePath(describeFile.DepotPath.Path, streamName, out relativePath))
+							{
+								files.Add(CreateChangeFile(relativePath, describeFile));
+							}
+						}
+
+						IUser user = await FindOrAddUserAsync(clusterName, change.OwnerName);
+						return new ChangeDetails(change.Id, user, change.GetPath(), change.Description, files, change.ModifiedDate);
+					}
+				}
+			}
+		}
+
+		/// <inheritdoc/>
+		public async Task<CheckShelfResult> CheckShelfAsync(string clusterName, string streamName, int changeNumber, string? impersonateUser)
+		{
+			using IScope scope = GlobalTracer.Instance.BuildSpan("PerforceService.CheckPreflightAsync").StartActive();
+			scope.Span.SetTag("ClusterName", clusterName);
+			scope.Span.SetTag("StreamName", streamName);
+			scope.Span.SetTag("ChangeNumber", changeNumber);
+			scope.Span.SetTag("ImpersonateUser", impersonateUser);
+
+			PerforceCluster cluster = await GetClusterAsync(clusterName);
+			using (P4.Repository repository = await GetConnection(cluster, stream: streamName, username: impersonateUser))
+			{
+				P4.Changelist change = repository.GetChangelist(changeNumber, new P4.DescribeCmdOptions(P4.DescribeChangelistCmdFlags.Omit | P4.DescribeChangelistCmdFlags.Shelved, 0, 0));
+				if(change == null)
+				{
+					return CheckShelfResult.NoChange;
+				}
+				if (change.ShelvedFiles == null || change.ShelvedFiles.Count == 0)
+				{
+					return CheckShelfResult.NoShelvedFiles;
+				}
+
+				P4.Stream stream = repository.GetStream(streamName, new P4.StreamCmdOptions(P4.StreamCmdFlags.View, null, null));
+				ViewMap map = new ViewMap(stream.View);
+
+				bool bHasMappedFile = false;
+				bool bHasUnmappedFile = false;
+				foreach (P4.ShelvedFile shelvedFile in change.ShelvedFiles)
+				{
+					if (map.TryMapFile(shelvedFile.Path.Path, Utf8StringComparer.OrdinalIgnoreCase, out Utf8String _))
+					{
+						bHasMappedFile = true;
+					}
+					else
+					{
+						bHasUnmappedFile = true;
+					}
+				}
+
+				if (bHasUnmappedFile)
+				{
+					return bHasMappedFile ? CheckShelfResult.MixedStream : CheckShelfResult.WrongStream;
+				}
+			}
+			return CheckShelfResult.Ok;
+		}
+
+		/// <inheritdoc/>
+		public async Task<List<ChangeDetails>> GetChangeDetailsAsync(string clusterName, string streamName, IReadOnlyList<int> changeNumbers, string? impersonateUser)
+		{
+			using IScope scope = GlobalTracer.Instance.BuildSpan("PerforceService.GetChangeDetailsAsync").StartActive();
+			scope.Span.SetTag("ClusterName", clusterName);
+			scope.Span.SetTag("StreamName", streamName);
+			scope.Span.SetTag("ChangeNumbers.Count", changeNumbers.Count);
+			scope.Span.SetTag("ImpersonateUser", impersonateUser);
+			
+			PerforceCluster cluster = await GetClusterAsync(clusterName);
+			using (P4.Repository repository = await GetConnection(cluster, stream: streamName, username: impersonateUser))
+			{
+				List<ChangeDetails> results = new List<ChangeDetails>();
+
+				List<string> args = new List<string> { "-s", "-S" };
+				args.AddRange(changeNumbers.Select(change => change.ToString(CultureInfo.InvariantCulture)));
+
+				using (P4.P4Command command = new P4.P4Command(repository, "describe", true, args.ToArray()))
+				{
+					P4.P4CommandResult result = command.Run();
+
+					if (!result.Success)
+					{
+						throw new Exception("Unable to get changes");
+					}
+
+					if (result.TaggedOutput == null || result.TaggedOutput.Count <= 0)
+					{
+						return results;
+					}
+
+					List<P4.Changelist> changelists = new List<P4.Changelist>() { };
+
+					bool dstMismatch = false;
+					string offset = String.Empty;
+
+					P4.Server server = repository.Server;
+					if (server != null && server.Metadata != null)
+					{
+						offset = server.Metadata.DateTimeOffset;
+						dstMismatch = P4.FormBase.DSTMismatch(server.Metadata);
+					}
+
+					foreach (P4.TaggedObject taggedObject in result.TaggedOutput)
+					{
+						List<ChangeFile> files = new List<ChangeFile>();
+
+						P4.Changelist change = new P4.Changelist();
+						change.FromChangeCmdTaggedOutput(taggedObject, true, offset, dstMismatch);
+
+						foreach (P4.FileMetaData describeFile in change.Files)
+						{
+							string? relativePath;
+							if (TryGetStreamRelativePath(describeFile.DepotPath.Path, streamName, out relativePath))
+							{
+								files.Add(CreateChangeFile(relativePath, describeFile));
+							}
+						}
+
+						if (change.ShelvedFiles != null && change.ShelvedFiles.Count > 0)
+						{
+							foreach (P4.ShelvedFile shelvedFile in change.ShelvedFiles)
+							{
+								string? relativePath;
+								if (TryGetStreamRelativePath(shelvedFile.Path.ToString(), streamName, out relativePath))
+								{
+									files.Add(CreateChangeFile(relativePath, shelvedFile));
+								}
+							}
+						}
+
+						IUser user = await FindOrAddUserAsync(clusterName, change.OwnerName);
+						results.Add(new ChangeDetails(change.Id, user, change.GetPath(), change.Description, files, change.ModifiedDate));
+					}
+				}
+
+				return results;
+			}
+		}
+
+		/// <inheritdoc/>
+		public async Task<List<FileSummary>> FindFilesAsync(string clusterName, IEnumerable<string> paths)
+		{
+			using IScope scope = GlobalTracer.Instance.BuildSpan("PerforceService.FindFilesAsync").StartActive();
+			scope.Span.SetTag("ClusterName", clusterName);
+			
+			List<FileSummary> results = new List<FileSummary>();
+
+			PerforceCluster cluster = await GetClusterAsync(clusterName);
+			using (P4.Repository repository = await GetConnection(cluster, noClient: true))
+			{
+				P4.GetFileMetaDataCmdOptions options = new P4.GetFileMetaDataCmdOptions(P4.GetFileMetadataCmdFlags.ExcludeClientData, "", "", 0, "", "", "");
+
+				IList<P4.FileMetaData>? files = repository.GetFileMetaData(options, paths.Select(path => 
+				{ 
+					P4.FileSpec fileSpec = new P4.FileSpec(new P4.DepotPath(path)); 
+					return fileSpec; 
+				}).ToArray());
+
+				if (files == null)
+				{
+					files = new List<P4.FileMetaData>();
+				}
+
+				foreach (string path in paths)
+				{
+					P4.FileMetaData? meta = files.FirstOrDefault(file => file.DepotPath?.Path == path);
+
+					if (meta == null)
+					{
+						results.Add(new FileSummary(path, false, 0));
+					}
+					else
+					{
+						results.Add(new FileSummary(path, meta.HeadAction != P4.FileAction.Delete && meta.HeadAction != P4.FileAction.MoveDelete, meta.HeadChange));
+					}
+				}
+
+				return results;
+			}
+		}
+
+		/// <inheritdoc/>
+		public async Task<byte[]> PrintAsync(string clusterName, string depotPath)
+		{
+			using IScope scope = GlobalTracer.Instance.BuildSpan("PerforceService.PrintAsync").StartActive();
+			scope.Span.SetTag("ClusterName", clusterName);
+			scope.Span.SetTag("DepotPath", depotPath);
+			
+			if (depotPath.EndsWith("...", StringComparison.OrdinalIgnoreCase) || depotPath.EndsWith("*", StringComparison.OrdinalIgnoreCase))
+			{
+				throw new Exception("PrintAsync requires exactly one file to be specified");
+			}
+
+			PerforceCluster cluster = await GetClusterAsync(clusterName);
+			using (P4.Repository repository = await GetConnection(cluster, noClient: true))
+			{
+				using (P4.P4Command command = new P4.P4Command(repository, "print", false, new string[] { "-q", depotPath }))
+				{
+					P4.P4CommandResult result = command.Run();
+
+					if (result.BinaryOutput != null)
+					{
+						return result.BinaryOutput;
+					}
+
+					return Encoding.Default.GetBytes(result.TextOutput);
+				}
+			}
+		}
+
+		/// <inheritdoc/>
+		public async Task<int> DuplicateShelvedChangeAsync(string clusterName, int shelvedChange)
+		{
+			using IScope scope = GlobalTracer.Instance.BuildSpan("PerforceService.DuplicateShelvedChangeAsync").StartActive();
+			scope.Span.SetTag("ClusterName", clusterName);
+			scope.Span.SetTag("ShelvedChange", shelvedChange);
+
+			string? changeOwner = null;
+
+			// Get the owner of the shelf
+			PerforceCluster cluster = await GetClusterAsync(clusterName);
+			using (P4.Repository repository = await GetConnection(cluster, noClient: true))
+			{
+				List<string> args = new List<string> { "-S", shelvedChange.ToString(CultureInfo.InvariantCulture) };
+
+				using (P4.P4Command command = new P4.P4Command(repository, "describe", true, args.ToArray()))
+				{
+					P4.P4CommandResult result = command.Run();
+
+					if (!result.Success)
+					{
+						throw new Exception($"Unable to get change {shelvedChange}");
+					}
+
+					if (result.TaggedOutput == null || result.TaggedOutput.Count != 1)
+					{
+
+						throw new Exception($"Unable to get tagged output for change: {shelvedChange}");
+					}
+
+					P4.Changelist changelist = new P4.Changelist();
+					changelist.FromChangeCmdTaggedOutput(result.TaggedOutput[0], true, String.Empty, false);
+					changeOwner = changelist.OwnerName;
+
+				}
+			}
+
+			if (String.IsNullOrEmpty(changeOwner))
+			{
+				throw new Exception($"Unable to get owner for shelved change {shelvedChange}");
+			}
+
+			using (P4.Repository repository = await GetConnection(cluster, readOnly: false, usePortFromChange: shelvedChange, clientFromChange: shelvedChange, useClientFromChange: false, username: changeOwner))
+			{
+				return ReshelveChange(repository, shelvedChange);
+			}
+
+			throw new Exception($"Unable to duplicate shelve change {shelvedChange}");
+		}
+
+		/// <inheritdoc/>
+		public async Task DeleteShelvedChangeAsync(string clusterName, int shelvedChange)
+		{
+			using IScope scope = GlobalTracer.Instance.BuildSpan("PerforceService.DeleteShelvedChangeAsync").StartActive();
+			scope.Span.SetTag("ClusterName", clusterName);
+			scope.Span.SetTag("ShelvedChange", shelvedChange);
+			
+			PerforceCluster cluster = await GetClusterAsync(clusterName);
+			using (P4.Repository repository = await GetConnection(cluster, noClient: true))
+			{
+
+				P4.Changelist? changelist = repository.GetChangelist(shelvedChange, new P4.DescribeCmdOptions(P4.DescribeChangelistCmdFlags.Omit, 0, 0));
+
+				if (changelist == null)
+				{
+					throw new Exception($"Unable to get shelved changelist for delete: {shelvedChange}");
+				}
+
+				string clientId = changelist.ClientId;
+
+				if (String.IsNullOrEmpty(clientId))
+				{
+					throw new Exception($"Unable to get shelved changelist client id for delete: {shelvedChange}");
+				}
+
+				using (P4.Repository clientRepository = await GetConnection(cluster, changelist.Stream, changelist.OwnerName, clientId: clientId, useClientFromChange: true, usePortFromChange: shelvedChange))
+				{
+
+					if (changelist.Shelved)
+					{
+						IList<P4.FileSpec> files = clientRepository.Connection.Client.ShelveFiles(new P4.ShelveFilesCmdOptions(P4.ShelveFilesCmdFlags.Delete, null, shelvedChange));
+					}
+
+					clientRepository.DeleteChangelist(changelist, null);
+				}
+			}
+		}
+
+		/// <inheritdoc/>
+		public async Task UpdateChangelistDescription(string clusterName, int change, string description)
+		{
+			using IScope scope = GlobalTracer.Instance.BuildSpan("PerforceService.UpdateChangelistDescription").StartActive();
+			scope.Span.SetTag("ClusterName", clusterName);
+			scope.Span.SetTag("Change", change);
+
+			try
+			{
+				PerforceCluster cluster = await GetClusterAsync(clusterName);
+
+				using (P4.Repository repository = await GetConnection(cluster, noClient: true))
+				{
+					P4.Changelist changelist = repository.GetChangelist(change);
+
+					repository.Connection.Disconnect();
+
+					// the client must exist for the change list, otherwise will fail (for example, CreateNewChangeAsync deletes the client before returning)
+					using (P4.Repository updateRepository = await GetConnection(cluster, clientFromChange: change, useClientFromChange: true, usePortFromChange: change, username: changelist.OwnerName))
+					{
+						P4.Changelist updatedChangelist = updateRepository.GetChangelist(change);
+						updatedChangelist.Description = description;
+						updateRepository.UpdateChangelist(updatedChangelist);
+					}
+				}
+			}
+			catch (Exception ex)
+			{
+				LogPerforce(1, "", $"Unable to update Changelist for CL {change} to ${description}, {ex.Message}");
+			}
+		}
+
+		/// <inheritdoc/>
+		public async Task<int> CreateNewChangeAsync(string clusterName, string streamName, string filePath, string description)
+		{
+			using IScope scope = GlobalTracer.Instance.BuildSpan("PerforceService.CreateNewChangeAsync").StartActive();
+			scope.Span.SetTag("ClusterName", clusterName);
+			scope.Span.SetTag("StreamName", streamName);
+			scope.Span.SetTag("FilePath", filePath);
+			
+			PerforceCluster cluster = await GetClusterAsync(clusterName);
+			P4.SubmitResults? submitResults = null;
+
+			using (P4.Repository repository = await GetConnection(cluster, stream: streamName, readOnly: false, createChange: true))
+			{
+				P4.Client client = repository.Connection.Client;
+
+				string workspaceFilePath = $"//{client.Name}/{filePath.TrimStart('/')}";
+				string diskFilePath = $"{client.Root + filePath.TrimStart('/')}";
+
+				P4.FileSpec workspaceFileSpec = new P4.FileSpec(new P4.ClientPath(workspaceFilePath));
+
+				IList<P4.File> files = repository.GetFiles(new P4.FilesCmdOptions(P4.FilesCmdFlags.None, 1), workspaceFileSpec);
+
+				P4.Changelist? submitChangelist = null;
+				P4.DepotPath? depotPath = null;
+
+				const int MaxRetries = 5;
+				int retry = 0;
+
+				for (; ; )
+				{
+					ResetClient(repository);
+
+					depotPath = null;
+
+					if (retry == MaxRetries)
+					{
+						break;
+					}
+
+					try
+					{
+						if (files == null || files.Count == 0)
+						{
+							// File does not exist, create it
+							string? directoryName = Path.GetDirectoryName(diskFilePath);
+
+							if (String.IsNullOrEmpty(directoryName))
+							{
+								throw new Exception($"Invalid directory name for local client file, disk file path: {diskFilePath}");
+							}
+
+							// Create the directory
+							if (!Directory.Exists(directoryName))
+							{
+								Directory.CreateDirectory(directoryName);
+
+								if (!Directory.Exists(directoryName))
+								{
+									throw new Exception($"Unable to create directrory: {directoryName}");
+								}
+							}
+
+							// Create the file
+							if (!File.Exists(diskFilePath))
+							{
+								using (FileStream fileStream = File.OpenWrite(diskFilePath))
+								{
+									fileStream.Close();
+								}
+
+								if (!File.Exists(diskFilePath))
+								{
+									throw new Exception($"Unable to create local change file: {diskFilePath}");
+								}
+							}
+
+							IList<P4.FileSpec> depotFiles = client.AddFiles(new P4.Options(), workspaceFileSpec);
+
+							if (depotFiles == null || depotFiles.Count != 1)
+							{
+								throw new Exception($"Unable to add local change file,  local: {diskFilePath} : workspace: {workspaceFileSpec}");
+							}
+
+							depotPath = depotFiles[0].DepotPath;
+
+						}
+						else
+						{
+							IList<P4.FileSpec> syncResults = client.SyncFiles(new P4.SyncFilesCmdOptions(P4.SyncFilesCmdFlags.Force), workspaceFileSpec);
+
+							if (syncResults == null || syncResults.Count != 1)
+							{
+								throw new Exception($"Unable to sync file, workspace: {workspaceFileSpec}");
+							}
+
+							IList<P4.FileSpec> editResults = client.EditFiles(new P4.FileSpec[] { new P4.FileSpec(syncResults[0].DepotPath) }, new P4.Options());
+
+							if (editResults == null || editResults.Count != 1)
+							{
+								throw new Exception($"Unable to edit file, workspace: {workspaceFileSpec}");
+							}
+
+							depotPath = editResults[0].DepotPath;
+
+						}
+
+						if (depotPath == null || String.IsNullOrEmpty(depotPath.Path))
+						{
+							throw new Exception($"Unable to get depot path for: {workspaceFileSpec}");
+						}
+
+						// create a new change
+						P4.Changelist changelist = new P4.Changelist();
+						changelist.Description = description;
+						changelist.Files.Add(new P4.FileMetaData(new P4.FileSpec(depotPath)));
+						submitChangelist = repository.CreateChangelist(changelist);
+
+						if (submitChangelist == null)
+						{
+							throw new Exception($"Unable to create a changelist for: {depotPath}");
+						}
+
+						submitResults = submitChangelist.Submit(null);
+
+						if (submitResults == null)
+						{
+							throw new Exception($"Unable to submit changelist for: {depotPath}");
+						}
+
+						break;
+					}
+					catch
+					{
+						retry++;
+						continue;
+					}
+				}
+
+				string clientName = client.Name;
+				try
+				{
+					repository.DeleteClient(client, new P4.Options());
+				}
+				catch
+				{
+					_logger.LogError("Unable to delete client {ClientName}", clientName);
+				}
+			}
+
+			if (submitResults == null)
+			{
+				throw new Exception($"Unable to submit change for {streamName} {filePath}");
+			}
+
+			return submitResults.ChangeIdAfterSubmit;
+		}
+
+		/// <inheritdoc/>
+		public async Task<(int? Change, string Message)> SubmitShelvedChangeAsync(string clusterName, int change, int originalChange)
+		{
+			using IScope scope = GlobalTracer.Instance.BuildSpan("PerforceService.SubmitShelvedChangeAsync").StartActive();
+			scope.Span.SetTag("ClusterName", clusterName);
+			scope.Span.SetTag("Change", change);
+			scope.Span.SetTag("OriginalChange", originalChange);
+			
+			PerforceCluster cluster = await GetClusterAsync(clusterName);
+			using (P4.Repository repository = await GetConnection(cluster, noClient: true))
+			{
+
+				List<string> args = new List<string> { "-S", change.ToString(CultureInfo.InvariantCulture), originalChange.ToString(CultureInfo.InvariantCulture) };
+
+				using (P4.P4Command command = new P4.P4Command(repository, "describe", true, args.ToArray()))
+				{
+					P4.P4CommandResult result = command.Run();
+
+					if (!result.Success)
+					{
+						return (null, $"Unable to get change {change}");
+
+					}
+
+					if (result.TaggedOutput == null || result.TaggedOutput.Count != 2)
+					{
+
+						return (null, $"Unable to get tagged output for change: {change} and original change: {originalChange}");
+					}
+
+					bool dstMismatch = false;
+					string offset = String.Empty;
+
+					P4.Server server = repository.Server;
+					if (server != null && server.Metadata != null)
+					{
+						offset = server.Metadata.DateTimeOffset;
+						dstMismatch = P4.FormBase.DSTMismatch(server.Metadata);
+					}
+
+					P4.Changelist changelist = new P4.Changelist();
+					changelist.FromChangeCmdTaggedOutput(result.TaggedOutput[0], true, offset, dstMismatch);
+
+					P4.Changelist originalChangelist = new P4.Changelist();
+					originalChangelist.FromChangeCmdTaggedOutput(result.TaggedOutput[1], true, offset, dstMismatch);
+
+					if (originalChangelist.ShelvedFiles.Count != changelist.ShelvedFiles.Count)
+					{
+						return (null, $"Mismatched number of shelved files for change: {change} and original change: {originalChange}");
+					}
+
+					if (originalChangelist.ShelvedFiles.Count == 0)
+					{
+						return (null, $"No shelved file for change: {change} and original change: {originalChange}");
+					}
+
+					foreach (P4.ShelvedFile shelvedFile in changelist.ShelvedFiles)
+					{
+						P4.ShelvedFile? found = originalChangelist.ShelvedFiles.FirstOrDefault(original => original.Digest == shelvedFile.Digest && original.Action == shelvedFile.Action);
+
+						if (found == null)
+						{
+							return (null, $"Mismatch in shelved file digest or action for {shelvedFile.Path}");
+						}
+					}
+
+					repository.Connection.Disconnect();
+
+					using (P4.Repository submitRepository = await GetConnection(cluster, readOnly: false, clientFromChange: change, useClientFromChange: true, usePortFromChange: change, username: changelist.OwnerName))
+					{
+						// we might not need a client here, possibly -e below facilitates this, check!
+						using (P4.P4Command submitCommand = new P4.P4Command(submitRepository, "submit", null, true, new string[] { "-e", change.ToString(CultureInfo.InvariantCulture) }))
+						{
+							try
+							{
+								result = submitCommand.Run();
+							}
+							catch (Exception ex)
+							{
+								return (null, $"Submit command failed: {ex.Message}");
+							}
+
+							if (!result.Success)
+							{
+								string message = (result.ErrorList != null && result.ErrorList.Count > 0) ? result.ErrorList[0].ErrorMessage : "Unknown error, no errors in list";
+								return (null, $"Unable to submit {change}, {message}");
+							}
+
+							int? submittedChangeId = null;
+
+							foreach (P4.TaggedObject taggedObject in result.TaggedOutput)
+							{
+								string? submitted;
+								if (taggedObject.TryGetValue("submittedChange", out submitted))
+								{
+									submittedChangeId = Int32.Parse(submitted, CultureInfo.InvariantCulture);
+								}
+							}
+
+							if (submittedChangeId == null)
+							{
+								return (null, $"Submit command succeeded, though unable to parse submitted change number");
+							}
+
+							return (submittedChangeId, changelist.Description);
+
+						}
+					}
+				}
+			}
+
+			throw new Exception($"Unable to get shelve change: {change} for original change: {originalChange}");
+		}
+
+		/// <inheritdoc/>		
+		public async Task<PerforceUserInfo?> GetUserInfoAsync(string clusterName, string userName)
+		{
+			using IScope scope = GlobalTracer.Instance.BuildSpan("PerforceService.GetUserInfoAsync").StartActive();
+			scope.Span.SetTag("ClusterName", clusterName);
+			scope.Span.SetTag("UserName", userName);
+			
+			PerforceCluster cluster = await GetClusterAsync(clusterName);
+			using (P4.Repository repository = await GetConnection(cluster, noClient: true))
+			{
+				P4.User user = repository.GetUser(userName, new P4.UserCmdOptions(P4.UserCmdFlags.Output));
+
+				if (user == null)
+				{
+					return null;
+				}
+
+				return new PerforceUserInfo { Login = userName, FullName = user.FullName, Email = user.EmailAddress };
+			}
+		}
+
+		static int ReshelveChange(P4.Repository repository, int change)
+		{
+			using IScope scope = GlobalTracer.Instance.BuildSpan("PerforceService.ReshelveChange").StartActive();
+			scope.Span.SetTag("Change", change);
+
+			bool edgeServer = false;
+			string? value;
+			if (repository.Server.Metadata.RawData.TryGetValue("serverServices", out value))
+			{
+				if (value == "edge-server")
+				{
+					edgeServer = true;
+				}
+			}
+
+			List<string> arguments = new List<string>();
+
+			// promote shelf if we're on an edge server
+			if (edgeServer)
+			{
+				arguments.Add("-p");
+			}
+
+			arguments.AddRange(new string[] { "-s", change.ToString(CultureInfo.InvariantCulture), "-f" });
+
+			using (P4.P4Command cmd = new P4.P4Command(repository, "reshelve", false, arguments.ToArray()))
+			{
+				P4.P4CommandResult results = cmd.Run();
+
+				if (results.Success)
+				{
+
+					if (results.InfoOutput.Count == 0)
+					{
+						Serilog.Log.Logger.Information("Perforce: Unexpected info output when reshelving change");
+						throw new Exception("Unexpected info output when reshelving change");
+					}
+
+					bool error = true;
+					int reshelvedChange = 0;
+					string message = results.InfoOutput[^1].Message;
+					Match changeMatch = Regex.Match(message, @"Change (\d+) files shelved");
+
+					if (changeMatch.Success && changeMatch.Groups.Count == 2)
+					{
+						if (Int32.TryParse(changeMatch.Groups[1].Value, out reshelvedChange))
+						{
+							error = false;
+						}
+					}
+
+					if (error)
+					{
+						Serilog.Log.Logger.Information("Perforce: Unable to parse cl for reshelf: {Message}", message);
+						throw new Exception($"Unable to parse cl for reshelf: {message}");
+					}
+
+					return reshelvedChange;
+				}
+			}
+
+			Serilog.Log.Logger.Information("Perforce: General reshelving failure");
+
+			throw new Exception($"Unable to reshelve CL {change}");
+
+		}
+
+		/// <inheritdoc/>
+		public virtual async Task<int> GetCodeChangeAsync(string clusterName, string streamName, int change)
+		{
+			using IScope scope = GlobalTracer.Instance.BuildSpan("PerforceService.GetCodeChangeAsync").StartActive();
+			scope.Span.SetTag("ClusterName", clusterName);
+			scope.Span.SetTag("StreamName", streamName);
+			scope.Span.SetTag("Change", change);
+			
+			int maxChange = change;
+			for (; ; )
+			{
+				// Query for the changes before this point
+				List<ChangeSummary> changes = await GetChangesAsync(clusterName, streamName, null, maxChange, 10, null);
+				Serilog.Log.Logger.Information("Finding last code change in {Stream} before {MaxChange}: {NumResults}", streamName, maxChange, changes.Count);
+				if (changes.Count == 0)
+				{
+					return 0;
+				}
+
+				// Get the details for them
+				List<ChangeDetails> detailsList = await GetChangeDetailsAsync(clusterName, streamName, changes.ConvertAll(x => x.Number), null);
+				foreach (ChangeDetails details in detailsList.OrderByDescending(x => x.Number))
+				{
+					ChangeContentFlags contentFlags = details.GetContentFlags();
+					Serilog.Log.Logger.Information("Change {Change} = {Flags}", details.Number, contentFlags.ToString());
+					if ((details.GetContentFlags() & ChangeContentFlags.ContainsCode) != 0)
+					{
+						return details.Number;
+					}
+				}
+
+				// Loop round again, adjusting our maximum changelist number
+				maxChange = changes.Min(x => x.Number) - 1;
+			}
+		}
+
+		/// <summary>
+		/// Gets the wildcard filter for a particular range of changes
+		/// </summary>
+		/// <param name="basePath">Base path to find files for</param>
+		/// <param name="minChange">Minimum changelist number to query</param>
+		/// <param name="maxChange">Maximum changelist number to query</param>
+		/// <returns>Filter string</returns>
+		public static string GetFilter(string basePath, int? minChange, int? maxChange)
+		{
+			StringBuilder filter = new StringBuilder(basePath);
+			if (minChange != null && maxChange != null)
+			{
+				filter.Append(CultureInfo.InvariantCulture, $"@{minChange},{maxChange}");
+			}
+			else if (minChange != null)
+			{
+				filter.Append(CultureInfo.InvariantCulture, $"@>={minChange}");
+			}
+			else if (maxChange != null)
+			{
+				filter.Append(CultureInfo.InvariantCulture, $"@<={maxChange}");
+			}
+			return filter.ToString();
+		}
+
+		/// <summary>
+		/// Gets a stream-relative path from a depot path
+		/// </summary>
+		/// <param name="depotFile">The depot file to check</param>
+		/// <param name="streamName">Name of the stream</param>
+		/// <param name="relativePath">Receives the relative path to the file</param>
+		/// <returns>True if the stream-relative path was returned</returns>
+		public static bool TryGetStreamRelativePath(string depotFile, string streamName, [NotNullWhen(true)] out string? relativePath)
+		{
+			if (depotFile.StartsWith(streamName, StringComparison.OrdinalIgnoreCase) && depotFile.Length > streamName.Length && depotFile[streamName.Length] == '/')
+			{
+				relativePath = depotFile.Substring(streamName.Length);
+				return true;
+			}
+
+			Match match = Regex.Match(depotFile, "^//[^/]+/[^/]+(/.*)$");
+			if (match.Success)
+			{
+				relativePath = match.Groups[1].Value;
+				return true;
+			}
+
+			relativePath = null;
+			return false;
+		}
+	}
+}
diff --git a/Engine/Source/Programs/Horde/Horde.Build/Projects/ProjectCollection.cs b/Engine/Source/Programs/Horde/Horde.Build/Projects/ProjectCollection.cs
new file mode 100644
--- /dev/null
+++ b/Engine/Source/Programs/Horde/Horde.Build/Projects/ProjectCollection.cs
@@ -0,0 +1,170 @@
+// Copyright Epic Games, Inc. All Rights Reserved.
+
+using System;
+using System.Collections.Generic;
+using System.Diagnostics.CodeAnalysis;
+using System.Linq;
+using System.Threading.Tasks;
+using Horde.Build.Acls;
+using Horde.Build.Models;
+using Horde.Build.Server;
+using Horde.Build.Utilities;
+using MongoDB.Bson.Serialization.Attributes;
+using MongoDB.Driver;
+
+namespace Horde.Build.Collections.Impl
+{
+	using ProjectId = StringId<IProject>;
+
+	/// <summary>
+	/// Wraps functionality for manipulating projects
+	/// </summary>
+	public class ProjectCollection : IProjectCollection
+	{
+		/// <summary>
+		/// Represents a project
+		/// </summary>
+		class ProjectDocument : IProject
+		{
+			public const int DefaultOrder = 128;
+
+			[BsonRequired, BsonId]
+			public ProjectId Id { get; set; }
+
+			[BsonRequired]
+			public string Name { get; set; }
+
+			public string? ConfigPath { get; set; }
+			public string? ConfigRevision { get; set; }
+
+			public int Order { get; set; } = DefaultOrder;
+			public Acl? Acl { get; set; }
+			public List<StreamCategory> Categories { get; set; } = new List<StreamCategory>();
+
+			[BsonIgnoreIfDefault, BsonDefaultValue(false)]
+			public bool Deleted { get; set; }
+
+			IReadOnlyList<StreamCategory> IProject.Categories => Categories;
+
+			[BsonConstructor]
+			private ProjectDocument()
+			{
+				Name = null!;
+			}
+
+			public ProjectDocument(ProjectId id, string name)
+			{
+				Id = id;
+				Name = name;
+			}
+		}
+
+		/// <summary>
+		/// Logo for a project
+		/// </summary>
+		class ProjectLogoDocument : IProjectLogo
+		{
+			public ProjectId Id { get; set; }
+
+			public string Path { get; set; } = String.Empty;
+			public string Revision { get; set; } = String.Empty;
+			public string MimeType { get; set; } = String.Empty;
+
+			public byte[] Data { get; set; } = Array.Empty<byte>();
+		}
+
+		/// <summary>
+		/// Projection of a project definition to just include permissions info
+		/// </summary>
+		[SuppressMessage("Design", "CA1812: Class is never instantiated")]
+		class ProjectPermissions : IProjectPermissions
+		{
+			/// <summary>
+			/// ACL for the project
+			/// </summary>
+			public Acl? Acl { get; set; }
+
+			/// <summary>
+			/// Projection to extract the permissions info from the project
+			/// </summary>
+			public static readonly ProjectionDefinition<ProjectDocument> Projection = Builders<ProjectDocument>.Projection.Include(x => x.Acl);
+		}
+
+		/// <summary>
+		/// Collection of project documents
+		/// </summary>
+		readonly IMongoCollection<ProjectDocument> _projects;
+
+		/// <summary>
+		/// Collection of project logo documents
+		/// </summary>
+		readonly IMongoCollection<ProjectLogoDocument> _projectLogos;
+
+		/// <summary>
+		/// Constructor
+		/// </summary>
+		/// <param name="mongoService">The database service instance</param>
+		public ProjectCollection(MongoService mongoService)
+		{
+			_projects = mongoService.GetCollection<ProjectDocument>("Projects");
+			_projectLogos = mongoService.GetCollection<ProjectLogoDocument>("ProjectLogos");
+		}
+
+		/// <inheritdoc/>
+		public async Task<IProject?> AddOrUpdateAsync(ProjectId id, string configPath, string revision, int order, ProjectConfig config)
+		{
+			ProjectDocument newProject = new ProjectDocument(id, config.Name);
+			newProject.ConfigPath = configPath;
+			newProject.ConfigRevision = revision;
+			newProject.Order = order;
+			newProject.Categories = config.Categories.ConvertAll(x => new StreamCategory(x));
+			newProject.Acl = Acl.Merge(new Acl(), config.Acl);
+
+			await _projects.FindOneAndReplaceAsync<ProjectDocument>(x => x.Id == id, newProject, new FindOneAndReplaceOptions<ProjectDocument> { IsUpsert = true });
+			return newProject;
+		}
+
+		/// <inheritdoc/>
+		public async Task<List<IProject>> FindAllAsync()
+		{
+			List<ProjectDocument> results = await _projects.Find(x => !x.Deleted).ToListAsync();
+			return results.OrderBy(x => x.Order).ThenBy(x => x.Name).Select<ProjectDocument, IProject>(x => x).ToList();
+		}
+
+		/// <inheritdoc/>
+		public async Task<IProject?> GetAsync(ProjectId projectId)
+		{
+			return await _projects.Find<ProjectDocument>(x => x.Id == projectId).FirstOrDefaultAsync();
+		}
+
+		/// <inheritdoc/>
+		public async Task<IProjectLogo?> GetLogoAsync(ProjectId projectId)
+		{
+			return await _projectLogos.Find(x => x.Id == projectId).FirstOrDefaultAsync();
+		}
+
+		/// <inheritdoc/>
+		public async Task SetLogoAsync(ProjectId projectId, string logoPath, string logoRevision, string mimeType, byte[] data)
+		{
+			ProjectLogoDocument logo = new ProjectLogoDocument();
+			logo.Id = projectId;
+			logo.Path = logoPath;
+			logo.Revision = logoRevision;
+			logo.MimeType = mimeType;
+			logo.Data = data;
+			await _projectLogos.ReplaceOneAsync(x => x.Id == projectId, logo, new ReplaceOptions { IsUpsert = true });
+		}
+
+		/// <inheritdoc/>
+		public async Task<IProjectPermissions?> GetPermissionsAsync(ProjectId projectId)
+		{
+			return await _projects.Find<ProjectDocument>(x => x.Id == projectId).Project<ProjectPermissions>(ProjectPermissions.Projection).FirstOrDefaultAsync();
+		}
+
+		/// <inheritdoc/>
+		public async Task DeleteAsync(ProjectId projectId)
+		{
+			await _projects.UpdateOneAsync(x => x.Id == projectId, Builders<ProjectDocument>.Update.Set(x => x.Deleted, true));
+		}
+	}
+}
diff --git a/Engine/Source/Programs/Horde/Horde.Build/Projects/ProjectService.cs b/Engine/Source/Programs/Horde/Horde.Build/Projects/ProjectService.cs
new file mode 100644
--- /dev/null
+++ b/Engine/Source/Programs/Horde/Horde.Build/Projects/ProjectService.cs
@@ -0,0 +1,152 @@
+﻿// Copyright Epic Games, Inc. All Rights Reserved.
+
+using System.Collections.Generic;
+using System.Security.Claims;
+using System.Threading.Tasks;
+using Horde.Build.Acls;
+using Horde.Build.Collections;
+using Horde.Build.Models;
+using Horde.Build.Utilities;
+
+namespace Horde.Build.Services
+{
+	using ProjectId = StringId<IProject>;
+
+	/// <summary>
+	/// Cache of information about job ACLs
+	/// </summary>
+	public class ProjectPermissionsCache : GlobalPermissionsCache
+	{
+		/// <summary>
+		/// Map of project id to permissions for that project
+		/// </summary>
+		public Dictionary<ProjectId, IProjectPermissions?> Projects { get; } = new Dictionary<ProjectId, IProjectPermissions?>();
+	}
+
+	/// <summary>
+	/// Wraps functionality for manipulating projects
+	/// </summary>
+	public class ProjectService
+	{
+		/// <summary>
+		/// The ACL service
+		/// </summary>
+		readonly AclService _aclService;
+
+		/// <summary>
+		/// Collection of project documents
+		/// </summary>
+		readonly IProjectCollection _projects;
+
+		/// <summary>
+		/// Accessor for the collection of project documents
+		/// </summary>
+		public IProjectCollection Collection => _projects;
+
+		/// <summary>
+		/// Constructor
+		/// </summary>
+		/// <param name="aclService">The ACL service</param>
+		/// <param name="projects">Collection of project documents</param>
+		public ProjectService(AclService aclService, IProjectCollection projects)
+		{
+			_aclService = aclService;
+			_projects = projects;
+		}
+
+		/// <summary>
+		/// Gets all the available projects
+		/// </summary>
+		/// <returns>List of project documents</returns>
+		public Task<List<IProject>> GetProjectsAsync()
+		{
+			return _projects.FindAllAsync();
+		}
+
+		/// <summary>
+		/// Gets a project by ID
+		/// </summary>
+		/// <param name="projectId">Unique id of the project</param>
+		/// <returns>The project document</returns>
+		public Task<IProject?> GetProjectAsync(ProjectId projectId)
+		{
+			return _projects.GetAsync(projectId);
+		}
+
+		/// <summary>
+		/// Gets a project's permissions info by ID
+		/// </summary>
+		/// <param name="projectId">Unique id of the project</param>
+		/// <returns>The project document</returns>
+		public Task<IProjectPermissions?> GetProjectPermissionsAsync(ProjectId projectId)
+		{
+			return _projects.GetPermissionsAsync(projectId);
+		}
+
+		/// <summary>
+		/// Deletes a project by id
+		/// </summary>
+		/// <param name="projectId">Unique id of the project</param>
+		public async Task DeleteProjectAsync(ProjectId projectId)
+		{
+			await _projects.DeleteAsync(projectId);
+		}
+
+		/// <summary>
+		/// Determines if the user is authorized to perform an action on a particular project
+		/// </summary>
+		/// <param name="acl">Acl for the project to check</param>
+		/// <param name="action">The action being performed</param>
+		/// <param name="user">The principal to authorize</param>
+		/// <param name="cache">Cache for the scope table</param>
+		/// <returns>True if the action is authorized</returns>
+		private Task<bool> AuthorizeAsync(Acl? acl, AclAction action, ClaimsPrincipal user, GlobalPermissionsCache? cache)
+		{
+			bool? result = acl?.Authorize(action, user);
+			if (result == null)
+			{
+				return _aclService.AuthorizeAsync(action, user, cache);
+			}
+			else
+			{
+				return Task.FromResult(result.Value);
+			}
+		}
+
+		/// <summary>
+		/// Determines if the user is authorized to perform an action on a particular project
+		/// </summary>
+		/// <param name="project">The project to check</param>
+		/// <param name="action">The action being performed</param>
+		/// <param name="user">The principal to authorize</param>
+		/// <param name="cache">Cache for the scope table</param>
+		/// <returns>True if the action is authorized</returns>
+		public Task<bool> AuthorizeAsync(IProject project, AclAction action, ClaimsPrincipal user, GlobalPermissionsCache? cache)
+		{
+			return AuthorizeAsync(project.Acl, action, user, cache);
+		}
+
+		/// <summary>
+		/// Determines if the user is authorized to perform an action on a particular project
+		/// </summary>
+		/// <param name="projectId">The project id to check</param>
+		/// <param name="action">The action being performed</param>
+		/// <param name="user">The principal to authorize</param>
+		/// <param name="cache">Cache for project permissions</param>
+		/// <returns>True if the action is authorized</returns>
+		public async Task<bool> AuthorizeAsync(ProjectId projectId, AclAction action, ClaimsPrincipal user, ProjectPermissionsCache? cache)
+		{
+			IProjectPermissions? permissions;
+			if (cache == null)
+			{
+				permissions = await GetProjectPermissionsAsync(projectId);
+			}
+			else if (!cache.Projects.TryGetValue(projectId, out permissions))
+			{
+				permissions = await GetProjectPermissionsAsync(projectId);
+				cache.Projects.Add(projectId, permissions);
+			}
+			return await AuthorizeAsync(permissions?.Acl, action, user, cache);
+		}
+	}
+}
diff --git a/Engine/Source/Programs/Horde/Horde.Build/Projects/ProjectsController.Messages.cs b/Engine/Source/Programs/Horde/Horde.Build/Projects/ProjectsController.Messages.cs
new file mode 100644
--- /dev/null
+++ b/Engine/Source/Programs/Horde/Horde.Build/Projects/ProjectsController.Messages.cs
@@ -0,0 +1,179 @@
+// Copyright Epic Games, Inc. All Rights Reserved.
+
+using System.Collections.Generic;
+using System.ComponentModel.DataAnnotations;
+using Horde.Build.Models;
+
+namespace Horde.Build.Api
+{
+	/// <summary>
+	/// Information about a category to display for a stream
+	/// </summary>
+	public class CreateProjectCategoryRequest
+	{
+		/// <summary>
+		/// Name of this category
+		/// </summary>
+		[Required]
+		public string Name { get; set; }
+
+		/// <summary>
+		/// Index of the row to display this category on
+		/// </summary>
+		public int Row { get; set; }
+
+		/// <summary>
+		/// Whether to show this category on the nav menu
+		/// </summary>
+		public bool ShowOnNavMenu { get; set; }
+
+		/// <summary>
+		/// Patterns for stream names to include
+		/// </summary>
+		public List<string> IncludePatterns { get; set; } = new List<string>();
+
+		/// <summary>
+		/// Patterns for stream names to exclude
+		/// </summary>
+		public List<string> ExcludePatterns { get; set; } = new List<string>();
+
+		/// <summary>
+		/// Constructor
+		/// </summary>
+		public CreateProjectCategoryRequest(string name)
+		{
+			Name = name;
+		}
+	}
+
+	/// <summary>
+	/// Information about a stream within a project
+	/// </summary>
+	public class GetProjectStreamResponse
+	{
+		/// <summary>
+		/// The stream id
+		/// </summary>
+		public string Id { get; set; }
+
+		/// <summary>
+		/// The stream name
+		/// </summary>
+		public string Name { get; set; }
+
+		/// <summary>
+		/// Constructor
+		/// </summary>
+		/// <param name="id">The unique stream id</param>
+		/// <param name="name">The stream name</param>
+		public GetProjectStreamResponse(string id, string name)
+		{
+			Id = id;
+			Name = name;
+		}
+	}
+
+	/// <summary>
+	/// Information about a category to display for a stream
+	/// </summary>
+	public class GetProjectCategoryResponse
+	{
+		/// <summary>
+		/// Heading for this column
+		/// </summary>
+		public string Name { get; set; }
+
+		/// <summary>
+		/// Index of the row to display this category on
+		/// </summary>
+		public int Row { get; set; }
+
+		/// <summary>
+		/// Whether to show this category on the nav menu
+		/// </summary>
+		public bool ShowOnNavMenu { get; set; }
+
+		/// <summary>
+		/// Patterns for stream names to include
+		/// </summary>
+		public List<string> IncludePatterns { get; set; } = new List<string>();
+
+		/// <summary>
+		/// Patterns for stream names to exclude
+		/// </summary>
+		public List<string> ExcludePatterns { get; set; } = new List<string>();
+
+		/// <summary>
+		/// Streams to include in this category
+		/// </summary>
+		public List<string> Streams { get; set; } = new List<string>();
+
+		/// <summary>
+		/// Constructor
+		/// </summary>
+		/// <param name="streamCategory">The category to construct from</param>
+		public GetProjectCategoryResponse(StreamCategory streamCategory)
+		{
+			Name = streamCategory.Name;
+			Row = streamCategory.Row;
+			ShowOnNavMenu = streamCategory.ShowOnNavMenu;
+			IncludePatterns = streamCategory.IncludePatterns;
+			ExcludePatterns = streamCategory.ExcludePatterns;
+		}
+	}
+
+	/// <summary>
+	/// Response describing a project
+	/// </summary>
+	public class GetProjectResponse
+	{
+		/// <summary>
+		/// Unique id of the project
+		/// </summary>
+		public string Id { get; set; }
+
+		/// <summary>
+		/// Name of the project
+		/// </summary>
+		public string Name { get; set; }
+
+		/// <summary>
+		/// Order to display this project on the dashboard
+		/// </summary>
+		public int Order { get; set; }
+
+		/// <summary>
+		/// List of streams that are in this project
+		/// </summary>
+		public List<GetProjectStreamResponse>? Streams { get; set; }
+
+		/// <summary>
+		/// List of stream categories to display
+		/// </summary>
+		public List<GetProjectCategoryResponse>? Categories { get; set; }
+
+		/// <summary>
+		/// Custom permissions for this object
+		/// </summary>
+		public GetAclResponse? Acl { get; set; }
+
+		/// <summary>
+		/// Constructor
+		/// </summary>
+		/// <param name="id">Unique id of the project</param>
+		/// <param name="name">Name of the project</param>
+		/// <param name="order">Order to show this project on the dashboard</param>
+		/// <param name="streams">List of streams to display</param>
+		/// <param name="categories">List of stream categories to display</param>
+		/// <param name="acl">Custom permissions for this object</param>
+		public GetProjectResponse(string id, string name, int order, List<GetProjectStreamResponse>? streams, List<GetProjectCategoryResponse>? categories, GetAclResponse? acl)
+		{
+			Id = id;
+			Name = name;
+			Order = order;
+			Streams = streams;
+			Categories = categories;
+			Acl = acl;
+		}
+	}
+}
diff --git a/Engine/Source/Programs/Horde/Horde.Build/Projects/ProjectsController.cs b/Engine/Source/Programs/Horde/Horde.Build/Projects/ProjectsController.cs
new file mode 100644
--- /dev/null
+++ b/Engine/Source/Programs/Horde/Horde.Build/Projects/ProjectsController.cs
@@ -0,0 +1,152 @@
+﻿// Copyright Epic Games, Inc. All Rights Reserved.
+
+using System.Collections.Generic;
+using System.Threading.Tasks;
+using Horde.Build.Acls;
+using Horde.Build.Api;
+using Horde.Build.Models;
+using Horde.Build.Services;
+using Horde.Build.Utilities;
+using Microsoft.AspNetCore.Authorization;
+using Microsoft.AspNetCore.Mvc;
+
+namespace Horde.Build.Controllers
+{
+	using ProjectId = StringId<IProject>;
+
+	/// <summary>
+	/// Controller for the /api/v1/projects endpoint
+	/// </summary>
+	[ApiController]
+	[Authorize]
+	[Route("[controller]")]
+	public class ProjectsController : HordeControllerBase
+	{
+		/// <summary>
+		/// Singleton instance of the project service
+		/// </summary>
+		private readonly ProjectService _projectService;
+
+		/// <summary>
+		/// Singleton instance of the stream service
+		/// </summary>
+		private readonly StreamService _streamService;
+
+		/// <summary>
+		/// Constructor
+		/// </summary>
+		/// <param name="projectService">The project service</param>
+		/// <param name="streamService">The stream service</param>
+		public ProjectsController(ProjectService projectService, StreamService streamService)
+		{
+			_projectService = projectService;
+			_streamService = streamService;
+		}
+
+		/// <summary>
+		/// Query all the projects
+		/// </summary>
+		/// <param name="includeStreams">Whether to include streams in the response</param>
+		/// <param name="includeCategories">Whether to include categories in the response</param>
+		/// <param name="filter">Filter for the properties to return</param>
+		/// <returns>Information about all the projects</returns>
+		[HttpGet]
+		[Route("/api/v1/projects")]
+		[ProducesResponseType(typeof(List<GetProjectResponse>), 200)]
+		public async Task<ActionResult<List<object>>> GetProjectsAsync([FromQuery(Name = "Streams")] bool includeStreams = false, [FromQuery(Name = "Categories")] bool includeCategories = false, [FromQuery] PropertyFilter? filter = null)
+		{
+			List<IProject> projects = await _projectService.GetProjectsAsync();
+			ProjectPermissionsCache permissionsCache = new ProjectPermissionsCache();
+
+			List<IStream>? streams = null;
+			if (includeStreams || includeCategories)
+			{
+				streams = await _streamService.GetStreamsAsync();
+			}
+
+			List<object> responses = new List<object>();
+			foreach (IProject project in projects)
+			{
+				if (await _projectService.AuthorizeAsync(project, AclAction.ViewProject, User, permissionsCache))
+				{
+					bool bIncludeAcl = await _projectService.AuthorizeAsync(project, AclAction.ViewPermissions, User, permissionsCache);
+					responses.Add(project.ToResponse(includeStreams, includeCategories, streams, bIncludeAcl).ApplyFilter(filter));
+				}
+			}
+			return responses;
+		}
+
+		/// <summary>
+		/// Retrieve information about a specific project
+		/// </summary>
+		/// <param name="projectId">Id of the project to get information about</param>
+		/// <param name="filter">Filter for the properties to return</param>
+		/// <returns>Information about the requested project</returns>
+		[HttpGet]
+		[Route("/api/v1/projects/{projectId}")]
+		[ProducesResponseType(typeof(List<GetProjectResponse>), 200)]
+		public async Task<ActionResult<object>> GetProjectAsync(ProjectId projectId, [FromQuery] PropertyFilter? filter = null)
+		{
+			IProject? project = await _projectService.GetProjectAsync(projectId);
+			if (project == null)
+			{
+				return NotFound(projectId);
+			}
+
+			ProjectPermissionsCache cache = new ProjectPermissionsCache();
+			if (!await _projectService.AuthorizeAsync(project, AclAction.ViewProject, User, cache))
+			{
+				return Forbid(AclAction.ViewProject, projectId);
+			}
+
+			bool bIncludeStreams = PropertyFilter.Includes(filter, nameof(GetProjectResponse.Streams));
+			bool bIncludeCategories = PropertyFilter.Includes(filter, nameof(GetProjectResponse.Categories));
+
+			List<IStream>? visibleStreams = null;
+			if (bIncludeStreams || bIncludeCategories)
+			{
+				visibleStreams = new List<IStream>();
+
+				List<IStream> streams = await _streamService.GetStreamsAsync(project.Id);
+				foreach (IStream stream in streams)
+				{
+					if (await _streamService.AuthorizeAsync(stream, AclAction.ViewStream, User, cache))
+					{
+						visibleStreams.Add(stream);
+					}
+				}
+			}
+
+			bool bIncludeAcl = await _projectService.AuthorizeAsync(project, AclAction.ViewPermissions, User, cache);
+			return project.ToResponse(bIncludeStreams, bIncludeCategories, visibleStreams, bIncludeAcl).ApplyFilter(filter);
+		}
+
+		/// <summary>
+		/// Retrieve information about a specific project
+		/// </summary>
+		/// <param name="projectId">Id of the project to get information about</param>
+		/// <returns>Information about the requested project</returns>
+		[HttpGet]
+		[Route("/api/v1/projects/{projectId}/logo")]
+		public async Task<ActionResult<object>> GetProjectLogoAsync(ProjectId projectId)
+		{
+			IProject? project = await _projectService.GetProjectAsync(projectId);
+			if (project == null)
+			{
+				return NotFound(projectId);
+			}
+			if (!await _projectService.AuthorizeAsync(project, AclAction.ViewProject, User, null))
+			{
+				return Forbid(AclAction.ViewProject, projectId);
+			}
+
+			IProjectLogo? projectLogo = await _projectService.Collection.GetLogoAsync(projectId);
+			if (projectLogo == null)
+			{
+				return NotFound();
+			}
+
+			return new FileContentResult(projectLogo.Data, projectLogo.MimeType);
+		}
+	}
+}
diff --git a/Engine/Source/Programs/Horde/Horde.Build/Secrets/CredentialsController.Messages.cs b/Engine/Source/Programs/Horde/Horde.Build/Secrets/CredentialsController.Messages.cs
new file mode 100644
--- /dev/null
+++ b/Engine/Source/Programs/Horde/Horde.Build/Secrets/CredentialsController.Messages.cs
@@ -0,0 +1,105 @@
+// Copyright Epic Games, Inc. All Rights Reserved.
+
+using System.Collections.Generic;
+using System.ComponentModel.DataAnnotations;
+using Horde.Build.Models;
+
+namespace Horde.Build.Api
+{
+	/// <summary>
+	/// Parameters to create a new credential
+	/// </summary>
+	public class CreateCredentialRequest
+	{
+		/// <summary>
+		/// Name for the new credential
+		/// </summary>
+		[Required]
+		public string Name { get; set; } = null!;
+
+		/// <summary>
+		/// Properties for the new credential
+		/// </summary>
+		public Dictionary<string, string>? Properties { get; set; }
+	}
+
+	/// <summary>
+	/// Response from creating a new credential
+	/// </summary>
+	public class CreateCredentialResponse
+	{
+		/// <summary>
+		/// Unique id for the new credential
+		/// </summary>
+		public string Id { get; set; }
+
+		/// <summary>
+		/// Constructor
+		/// </summary>
+		/// <param name="id">Unique id for the new credential</param>
+		public CreateCredentialResponse(string id)
+		{
+			Id = id;
+		}
+	}
+
+	/// <summary>
+	/// Parameters to update a credential
+	/// </summary>
+	public class UpdateCredentialRequest
+	{
+		/// <summary>
+		/// Optional new name for the credential
+		/// </summary>
+		public string? Name { get; set; }
+
+		/// <summary>
+		/// Properties to update for the credential. Properties set to null will be removed.
+		/// </summary>
+		public Dictionary<string, string?>? Properties { get; set; }
+
+		/// <summary>
+		/// Custom permissions for this object
+		/// </summary>
+		public UpdateAclRequest? Acl { get; set; }
+	}
+
+	/// <summary>
+	/// Response describing a credential
+	/// </summary>
+	public class GetCredentialResponse
+	{
+		/// <summary>
+		/// Unique id of the credential
+		/// </summary>
+		public string Id { get; set; }
+
+		/// <summary>
+		/// Name of the credential
+		/// </summary>
+		public string Name { get; set; }
+
+		/// <summary>
+		/// Properties for the credential
+		/// </summary>
+		public Dictionary<string, string> Properties { get; set; }
+
+		/// <summary>
+		/// Custom permissions for this object
+		/// </summary>
+		public GetAclResponse? Acl { get; set; }
+
+		/// <summary>
+		/// Constructor
+		/// </summary>
+		/// <param name="credential">The credential to construct from</param>
+		/// <param name="bIncludeAcl">Whether to include the ACL in the response</param>
+		public GetCredentialResponse(Credential credential, bool bIncludeAcl)
+		{
+			Id = credential.Id.ToString();
+			Name = credential.Name;
+			Properties = credential.Properties;
+			Acl = (bIncludeAcl && credential.Acl != null)? new GetAclResponse(credential.Acl) : null;
+		}
+	}
+}
diff --git a/Engine/Source/Programs/Horde/Horde.Build/Secrets/CredentialsController.cs b/Engine/Source/Programs/Horde/Horde.Build/Secrets/CredentialsController.cs
new file mode 100644
--- /dev/null
+++ b/Engine/Source/Programs/Horde/Horde.Build/Secrets/CredentialsController.cs
@@ -0,0 +1,182 @@
+﻿// Copyright Epic Games, Inc. All Rights Reserved.
+
+using System.Collections.Generic;
+using System.Threading.Tasks;
+using Horde.Build.Acls;
+using Horde.Build.Api;
+using Horde.Build.Models;
+using Horde.Build.Services;
+using Horde.Build.Utilities;
+using Microsoft.AspNetCore.Authorization;
+using Microsoft.AspNetCore.Mvc;
+using MongoDB.Bson;
+
+namespace Horde.Build.Controllers
+{
+	/// <summary>
+	/// Controller for the /api/v1/credentials endpoint
+	/// </summary>
+	[ApiController]
+	[Authorize]
+	[Route("[controller]")]
+	public class CredentialsController : ControllerBase
+	{
+		/// <summary>
+		/// Singleton instance of the ACL service
+		/// </summary>
+		private readonly AclService _aclService;
+
+		/// <summary>
+		/// Singleton instance of the credential service
+		/// </summary>
+		private readonly CredentialService _credentialService;
+
+		/// <summary>
+		/// Constructor
+		/// </summary>
+		/// <param name="aclService">The ACL service</param>
+		/// <param name="credentialService">The credential service</param>
+		public CredentialsController(AclService aclService, CredentialService credentialService)
+		{
+			_aclService = aclService;
+			_credentialService = credentialService;
+		}
+
+		/// <summary>
+		/// Creates a new credential
+		/// </summary>
+		/// <param name="create">Parameters for the new credential.</param>
+		/// <returns>Http result code</returns>
+		[HttpPost]
+		[Route("/api/v1/credentials")]
+		public async Task<ActionResult<CreateCredentialResponse>> CreateCredentialAsync([FromBody] CreateCredentialRequest create)
+		{
+			if(!await _aclService.AuthorizeAsync(AclAction.CreateCredential, User))
+			{
+				return Forbid();
+			}
+
+			Credential newCredential = await _credentialService.CreateCredentialAsync(create.Name, create.Properties);
+			return new CreateCredentialResponse(newCredential.Id.ToString());
+		}
+
+		/// <summary>
+		/// Query all the credentials
+		/// </summary>
+		/// <param name="name">Id of the credential to get information about</param>
+		/// <param name="filter">Filter for the properties to return</param>
+		/// <returns>Information about all the credentials</returns>
+		[HttpGet]
+		[Route("/api/v1/credentials")]
+		[ProducesResponseType(typeof(List<GetCredentialResponse>), 200)]
+		public async Task<ActionResult<object>> FindCredentialAsync([FromQuery] string? name = null, [FromQuery] PropertyFilter? filter = null)
+		{
+			if (!await _aclService.AuthorizeAsync(AclAction.ListCredentials, User))
+			{
+				return Forbid();
+			}
+
+			List<Credential> credentials = await _credentialService.FindCredentialsAsync(name);
+			GlobalPermissionsCache cache = new GlobalPermissionsCache();
+
+			List<object> responses = new List<object>();
+			foreach (Credential credential in credentials)
+			{
+				if (await _credentialService.AuthorizeAsync(credential, AclAction.ViewCredential, User, cache))
+				{
+					bool bIncludeAcl = await _credentialService.AuthorizeAsync(credential, AclAction.ViewPermissions, User, cache);
+					responses.Add(new GetCredentialResponse(credential, bIncludeAcl).ApplyFilter(filter));
+				}
+			}
+			return responses;
+		}
+
+		/// <summary>
+		/// Retrieve information about a specific credential
+		/// </summary>
+		/// <param name="credentialId">Id of the credential to get information about</param>
+		/// <param name="filter">Filter for properties to return</param>
+		/// <returns>Information about the requested credential</returns>
+		[HttpGet]
+		[Route("/api/v1/credentials/{credentialId}")]
+		[ProducesResponseType(typeof(GetCredentialResponse), 200)]
+		public async Task<ActionResult<object>> GetCredentialAsync(string credentialId, [FromQuery] PropertyFilter? filter = null)
+		{
+			ObjectId projectIdValue = credentialId.ToObjectId();
+
+			Credential? credential = await _credentialService.GetCredentialAsync(projectIdValue);
+			if (credential == null)
+			{
+				return NotFound();
+			}
+
+			GlobalPermissionsCache cache = new GlobalPermissionsCache();
+			if (!await _credentialService.AuthorizeAsync(credential, AclAction.ViewCredential, User, cache))
+			{
+				return Forbid();
+			}
+
+			bool bIncludeAcl = await _credentialService.AuthorizeAsync(credential, AclAction.ViewPermissions, User, cache);
+			return new GetCredentialResponse(credential, bIncludeAcl).ApplyFilter(filter);
+		}
+
+		/// <summary>
+		/// Update a credential's properties.
+		/// </summary>
+		/// <param name="credentialId">Id of the credential to update</param>
+		/// <param name="update">Items on the credential to update</param>
+		/// <returns>Http result code</returns>
+		[HttpPut]
+		[Route("/api/v1/credentials/{credentialId}")]
+		public async Task<ActionResult> UpdateCredentialAsync(string credentialId, [FromBody] UpdateCredentialRequest update)
+		{
+			ObjectId credentialIdValue = credentialId.ToObjectId();
+
+			Credential? credential = await _credentialService.GetCredentialAsync(credentialIdValue);
+			if(credential == null)
+			{
+				return NotFound();
+			}
+
+			GlobalPermissionsCache cache = new GlobalPermissionsCache();
+			if (!await _credentialService.AuthorizeAsync(credential, AclAction.UpdateCredential, User, cache))
+			{
+				return Forbid();
+			}
+			if (update.Acl != null && !await _credentialService.AuthorizeAsync(credential, AclAction.ChangePermissions, User, cache))
+			{
+				return Forbid();
+			}
+
+			await _credentialService.UpdateCredentialAsync(credentialIdValue, update.Name, update.Properties, Acl.Merge(credential.Acl, update.Acl));
+			return new OkResult();
+		}
+
+		/// <summary>
+		/// Delete a credential
+		/// </summary>
+		/// <param name="credentialId">Id of the credential to delete</param>
+		/// <returns>Http result code</returns>
+		[HttpDelete]
+		[Route("/api/v1/credentials/{credentialId}")]
+		public async Task<ActionResult> DeleteCredentialAsync(string credentialId)
+		{
+			ObjectId credentialIdValue = credentialId.ToObjectId();
+
+			Credential? credential = await _credentialService.GetCredentialAsync(credentialIdValue);
+			if (credential == null)
+			{
+				return NotFound();
+			}
+			if (!await _credentialService.AuthorizeAsync(credential, AclAction.DeleteCredential, User, null))
+			{
+				return Forbid();
+			}
+			if (!await _credentialService.DeleteCredentialAsync(credentialIdValue))
+			{
+				return NotFound();
+			}
+			return new OkResult();
+		}
+	}
+}
diff --git a/Engine/Source/Programs/Horde/Horde.Build/Server/AccountController.cs b/Engine/Source/Programs/Horde/Horde.Build/Server/AccountController.cs
new file mode 100644
--- /dev/null
+++ b/Engine/Source/Programs/Horde/Horde.Build/Server/AccountController.cs
@@ -0,0 +1,147 @@
+﻿// Copyright Epic Games, Inc. All Rights Reserved.
+
+using System;
+using System.Globalization;
+using System.Net;
+using System.Text;
+using System.Threading.Tasks;
+using Horde.Build.Acls;
+using Horde.Build.Authentication;
+using Microsoft.AspNetCore.Authentication;
+using Microsoft.AspNetCore.Authentication.Cookies;
+using Microsoft.AspNetCore.Authentication.OpenIdConnect;
+using Microsoft.AspNetCore.Mvc;
+using Microsoft.Extensions.Options;
+
+namespace Horde.Build.Controllers
+{
+	/// <summary>
+	/// Controller managing account status
+	/// </summary>
+	[ApiController]
+	[Route("[controller]")]
+	public class AccountController : Controller
+	{
+		/// <summary>
+		/// Style sheet for HTML responses
+		/// </summary>
+		const string StyleSheet =
+			"body { font-family: 'Segoe UI', 'Roboto', arial, sans-serif; } " +
+			"p { margin:20px; font-size:13px; } " +
+			"h1 { margin:20px; font-size:32px; font-weight:200; } " +
+			"table { margin:10px 20px; } " +
+			"td { margin:5px; font-size:13px; }";
+
+		/// <summary>
+		/// The ACL service singleton
+		/// </summary>
+		readonly AclService _aclService;
+
+		/// <summary>
+		/// Authentication scheme in use
+		/// </summary>
+		readonly string _authenticationScheme;
+
+		/// <summary>
+		/// Constructor
+		/// </summary>
+		/// <param name="aclService">ACL service instance</param>
+		/// <param name="serverSettings">Server settings</param>
+		public AccountController(AclService aclService, IOptionsMonitor<ServerSettings> serverSettings)
+		{
+			_aclService = aclService;
+			_authenticationScheme = GetAuthScheme(serverSettings.CurrentValue.AuthMethod);
+		}
+
+		/// <summary>
+		/// Get auth scheme name for a given auth method
+		/// </summary>
+		/// <param name="method">Authentication method</param>
+		/// <returns>Name of authentication scheme</returns>
+		public static string GetAuthScheme(AuthMethod method)
+		{
+			return method switch
+			{
+				AuthMethod.Anonymous => AnonymousAuthenticationHandler.AuthenticationScheme,
+				AuthMethod.Okta => OktaDefaults.AuthenticationScheme,
+				AuthMethod.OpenIdConnect => OpenIdConnectDefaults.AuthenticationScheme,
+				_ => throw new ArgumentOutOfRangeException(nameof(method), method, null)
+			};
+		}
+
+		/// <summary>
+		/// Gets the current login status
+		/// </summary>
+		/// <returns>The current login state</returns>
+		[HttpGet]
+		[Route("/account")]
+		public async Task<ActionResult> State()
+		{
+			StringBuilder content = new StringBuilder();
+			content.Append($"<html><style>{StyleSheet}</style><h1>Horde Server</h1>");
+			if (User.Identity?.IsAuthenticated ?? false)
+			{
+				content.Append(CultureInfo.InvariantCulture, $"<p>User <b>{User.Identity?.Name}</b> is logged in. <a href=\"/account/logout\">Log out</a></p>");
+				if (await _aclService.AuthorizeAsync(AclAction.AdminWrite, User))
+				{
+					content.Append("<p>");
+					content.Append("<a href=\"/api/v1/admin/token\">Get bearer token</a><br/>");
+					content.Append("<a href=\"/api/v1/admin/registrationtoken\">Get agent registration token</a><br/>");
+					content.Append("<a href=\"/api/v1/admin/softwaretoken\">Get agent software upload token</a><br/>");
+					content.Append("<a href=\"/api/v1/admin/softwaredownloadtoken\">Get agent software download token</a><br/>");
+					content.Append("<a href=\"/api/v1/admin/configtoken\">Get configuration token</a><br/>");
+					content.Append("<a href=\"/api/v1/admin/chainedjobtoken\">Get chained job token</a><br/>");
+					content.Append("</p>");
+				}
+				content.Append(CultureInfo.InvariantCulture, $"<p>Claims for {User.Identity?.Name}:");
+				content.Append("<table>");
+				foreach (System.Security.Claims.Claim claim in User.Claims)
+				{
+					content.Append(CultureInfo.InvariantCulture, $"<tr><td>{claim.Type}</td><td>{claim.Value}</td></tr>");
+				}
+				content.Append("</table>");
+				content.Append("</p>");
+
+				content.Append(CultureInfo.InvariantCulture, $"<p>Built from Perforce</p>");
+			}
+			else
+			{
+				content.Append("<p><a href=\"/account/login\"><b>Login with OAuth2</b></a></p>");
+			}
+			content.Append("</html>");
+			return new ContentResult { ContentType = "text/html", StatusCode = (int)HttpStatusCode.OK, Content = content.ToString() };
+		}
+
+		/// <summary>
+		/// Login to the server
+		/// </summary>
+		/// <returns>Http result</returns>
+		[HttpGet]
+		[Route("/account/login")]
+		public IActionResult Login()
+		{
+			return new ChallengeResult(_authenticationScheme, new AuthenticationProperties { RedirectUri = "/account" });
+		}
+
+		/// <summary>
+		/// Logout of the current account
+		/// </summary>
+		/// <returns>Http result</returns>
+		[HttpGet]
+		[Route("/account/logout")]
+		public async Task<IActionResult> Logout()
+		{
+			await HttpContext.SignOutAsync(CookieAuthenticationDefaults.AuthenticationScheme);
+			try
+			{
+				await HttpContext.SignOutAsync(_authenticationScheme);
+			}
+			catch
+			{
+			}
+
+			string content = $"<html><style>{StyleSheet}</style><body onload=\"setTimeout(function(){{ window.location = '/account'; }}, 2000)\"><p>User has been logged out. Returning to login page.</p></body></html>";
+			return new ContentResult { ContentType = "text/html", StatusCode = (int)HttpStatusCode.OK, Content = content };
+		}
+	}
+}
