Horde: Finish re-org of Horde solution started a while ago.

diff --git a/Engine/Source/Programs/Horde/Horde.Build/Server/AdminController.cs b/Engine/Source/Programs/Horde/Horde.Build/Server/AdminController.cs
new file mode 100644
--- /dev/null
+++ b/Engine/Source/Programs/Horde/Horde.Build/Server/AdminController.cs
@@ -0,0 +1,241 @@
+// Copyright Epic Games, Inc. All Rights Reserved.
+
+using System;
+using System.Collections.Generic;
+using System.Linq;
+using System.Security.Claims;
+using System.Threading.Tasks;
+using Horde.Build.Acls;
+using Horde.Build.Services;
+using Horde.Build.Utilities;
+using Microsoft.AspNetCore.Authorization;
+using Microsoft.AspNetCore.Mvc;
+using Microsoft.Extensions.Options;
+
+namespace Horde.Build.Controllers
+{
+	/// <summary>
+	/// Object containing settings for the server
+	/// </summary>
+	public class AdminSettings
+	{
+		/// <summary>
+		/// The default perforce server
+		/// </summary>
+		public string? DefaultServerAndPort { get; set; }
+
+		/// <summary>
+		/// The default perforce username
+		/// </summary>
+		public string? DefaultUserName { get; set; }
+
+		/// <summary>
+		/// The default perforce password
+		/// </summary>
+		public string? DefaultPassword { get; set; }
+	}
+
+	/// <summary>
+	/// The conform limit value
+	/// </summary>
+	public class ConformSettings
+	{
+		/// <summary>
+		/// Maximum number of conforms allowed at once
+		/// </summary>
+		public int MaxCount { get; set; }
+	}
+
+	/// <summary>
+	/// Controller managing account status
+	/// </summary>
+	[ApiController]
+	[Authorize]
+	[Route("[controller]")]
+	public class AdminController : HordeControllerBase
+	{
+		readonly AclService _aclService;
+		readonly UpgradeService _upgradeService;
+		readonly IOptionsMonitor<ServerSettings> _settings;
+
+		/// <summary>
+		/// Constructor
+		/// </summary>
+		/// <param name="aclService">The ACL service singleton</param>
+		/// <param name="upgradeService">The upgrade service singelton</param>
+		/// <param name="settings">Server settings</param>
+		public AdminController(AclService aclService, UpgradeService upgradeService, IOptionsMonitor<ServerSettings> settings)
+		{
+			_aclService = aclService;
+			_upgradeService = upgradeService;
+			_settings = settings;
+		}
+
+		/// <summary>
+		/// Upgrade the database to the latest schema
+		/// </summary>
+		/// <param name="fromVersion">The schema version to upgrade from.</param>
+		[HttpPost]
+		[Route("/api/v1/admin/upgradeschema")]
+		public async Task<ActionResult> UpgradeSchemaAsync([FromQuery] int? fromVersion = null)
+		{
+			if (!await _aclService.AuthorizeAsync(AclAction.AdminWrite, User))
+			{
+				return Forbid(AclAction.AdminWrite);
+			}
+
+			await _upgradeService.UpgradeSchemaAsync(fromVersion);
+			return Ok();
+		}
+
+		/// <summary>
+		/// Issues a token for the given roles. Issues a token for the current user if not specified.
+		/// </summary>
+		/// <returns>Administrative settings for the server</returns>
+		[HttpGet]
+		[Route("/api/v1/admin/token")]
+		public async Task<ActionResult<string>> GetTokenAsync()
+		{
+			if (!await _aclService.AuthorizeAsync(AclAction.IssueBearerToken, User))
+			{
+				return Forbid(AclAction.IssueBearerToken);
+			}
+
+			return _aclService.IssueBearerToken(User.Claims, GetDefaultExpiryTime());
+		}
+
+		/// <summary>
+		/// Issues a token for the given roles. Issues a token for the current user if not specified.
+		/// </summary>
+		/// <param name="roles">Roles for the new token</param>
+		/// <returns>Administrative settings for the server</returns>
+		[HttpGet]
+		[Route("/api/v1/admin/roletoken")]
+		public async Task<ActionResult<string>> GetRoleTokenAsync([FromQuery] string roles)
+		{
+			if (!await _aclService.AuthorizeAsync(AclAction.AdminWrite, User))
+			{
+				return Forbid(AclAction.AdminWrite);
+			}
+
+			List<Claim> claims = new List<Claim>();
+			claims.AddRange(roles.Split('+').Select(x => new Claim(ClaimTypes.Role, x)));
+
+			return _aclService.IssueBearerToken(claims, GetDefaultExpiryTime());
+		}
+
+		/// <summary>
+		/// Issues a token for the given roles. Issues a token for the current user if not specified.
+		/// </summary>
+		/// <returns>Administrative settings for the server</returns>
+		[HttpGet]
+		[Route("/api/v1/admin/registrationtoken")]
+		public async Task<ActionResult<string>> GetRegistrationTokenAsync()
+		{
+			if (!await _aclService.AuthorizeAsync(AclAction.AdminWrite, User))
+			{
+				return Forbid(AclAction.AdminWrite);
+			}
+
+			List<AclClaim> claims = new List<AclClaim>();
+			claims.Add(new AclClaim(ClaimTypes.Name, User.Identity?.Name ?? "Unknown"));
+			claims.Add(AclService.AgentRegistrationClaim);
+
+			return _aclService.IssueBearerToken(claims, null);
+		}
+
+		/// <summary>
+		/// Issues a token valid to upload new versions of the agent software.
+		/// </summary>
+		/// <returns>Administrative settings for the server</returns>
+		[HttpGet]
+		[Route("/api/v1/admin/softwaretoken")]
+		public async Task<ActionResult<string>> GetSoftwareTokenAsync()
+		{
+			if (!await _aclService.AuthorizeAsync(AclAction.AdminWrite, User))
+			{
+				return Forbid(AclAction.AdminWrite);
+			}
+
+			List<AclClaim> claims = new List<AclClaim>();
+			claims.Add(new AclClaim(ClaimTypes.Name, User.Identity?.Name ?? "Unknown"));
+			claims.Add(AclService.UploadSoftwareClaim);
+
+			return _aclService.IssueBearerToken(claims, null);
+		}
+
+		/// <summary>
+		/// Issues a token valid to download new versions of the agent software.
+		/// </summary>
+		/// <returns>Administrative settings for the server</returns>
+		[HttpGet]
+		[Route("/api/v1/admin/softwaredownloadtoken")]
+		public async Task<ActionResult<string>> GetSoftwareDownloadTokenAsync()
+		{
+			if (!await _aclService.AuthorizeAsync(AclAction.AdminRead, User))
+			{
+				return Forbid(AclAction.AdminRead);
+			}
+
+			List<AclClaim> claims = new List<AclClaim>();
+			claims.Add(new AclClaim(ClaimTypes.Name, User.Identity?.Name ?? "Unknown"));
+			claims.Add(AclService.DownloadSoftwareClaim);
+
+			return _aclService.IssueBearerToken(claims, null);
+		}
+
+		/// <summary>
+		/// Issues a token valid to configure streams and projects
+		/// </summary>
+		/// <returns>Administrative settings for the server</returns>
+		[HttpGet]
+		[Route("/api/v1/admin/configtoken")]
+		public async Task<ActionResult<string>> GetConfigToken()
+		{
+			if (!await _aclService.AuthorizeAsync(AclAction.AdminRead, User))
+			{
+				return Forbid(AclAction.AdminRead);
+			}
+
+			List<AclClaim> claims = new List<AclClaim>();
+			claims.Add(new AclClaim(ClaimTypes.Name, User.Identity?.Name ?? "Unknown"));
+			claims.Add(AclService.ConfigureProjectsClaim);
+
+			return _aclService.IssueBearerToken(claims, null);
+		}
+
+		/// <summary>
+		/// Issues a token valid to start chained jobs
+		/// </summary>
+		/// <returns>Administrative settings for the server</returns>
+		[HttpGet]
+		[Route("/api/v1/admin/chainedjobtoken")]
+		public async Task<ActionResult<string>> GetChainedJobToken()
+		{
+			if (!await _aclService.AuthorizeAsync(AclAction.AdminRead, User))
+			{
+				return Forbid(AclAction.AdminRead);
+			}
+
+			List<AclClaim> claims = new List<AclClaim>();
+			//Claims.Add(new AclClaim(ClaimTypes.Name, User.Identity.Name ?? "Unknown"));
+			claims.Add(AclService.StartChainedJobClaim);
+
+			return _aclService.IssueBearerToken(claims, null);
+		}
+
+		/// <summary>
+		/// Gets the default expiry time for a token
+		/// </summary>
+		/// <returns></returns>
+		private TimeSpan? GetDefaultExpiryTime()
+		{
+			TimeSpan? expiryTime = null;
+			if (_settings.CurrentValue.JwtExpiryTimeHours != -1)
+			{
+				expiryTime = TimeSpan.FromHours(_settings.CurrentValue.JwtExpiryTimeHours);
+			}
+			return expiryTime;
+		}
+	}
+}
diff --git a/Engine/Source/Programs/Horde/Horde.Build/Server/ConsistencyService.cs b/Engine/Source/Programs/Horde/Horde.Build/Server/ConsistencyService.cs
new file mode 100644
--- /dev/null
+++ b/Engine/Source/Programs/Horde/Horde.Build/Server/ConsistencyService.cs
@@ -0,0 +1,73 @@
+// Copyright Epic Games, Inc. All Rights Reserved.
+
+using System;
+using System.Collections.Generic;
+using System.Linq;
+using System.Threading;
+using System.Threading.Tasks;
+using Horde.Build.Collections;
+using Horde.Build.Models;
+using Horde.Build.Utilities;
+using HordeCommon;
+using Microsoft.Extensions.Hosting;
+using Microsoft.Extensions.Logging;
+
+namespace Horde.Build.Services
+{
+	using SessionId = ObjectId<ISession>;
+
+	/// <summary>
+	/// Service which checks the database for consistency and fixes up any errors
+	/// </summary>
+	class ConsistencyService : IHostedService, IDisposable
+	{
+		readonly ISessionCollection _sessionCollection;
+		readonly ILeaseCollection _leaseCollection;
+		readonly ITicker _ticker;
+		readonly ILogger<ConsistencyService> _logger;
+
+		/// <summary>
+		/// Constructor
+		/// </summary>
+		public ConsistencyService(ISessionCollection sessionCollection, ILeaseCollection leaseCollection, IClock clock, ILogger<ConsistencyService> logger)
+		{
+			_sessionCollection = sessionCollection;
+			_leaseCollection = leaseCollection;
+			_ticker = clock.AddSharedTicker<ConsistencyService>(TimeSpan.FromMinutes(20.0), TickLeaderAsync, logger);
+			_logger = logger;
+		}
+
+		/// <inheritdoc/>
+		public Task StartAsync(CancellationToken cancellationToken) => _ticker.StartAsync();
+
+		/// <inheritdoc/>
+		public Task StopAsync(CancellationToken cancellationToken) => _ticker.StopAsync();
+
+		/// <inheritdoc/>
+		public void Dispose() => _ticker.Dispose();
+
+		/// <summary>
+		/// Poll for inconsistencies in the database
+		/// </summary>
+		/// <param name="stoppingToken">Stopping token</param>
+		/// <returns>Async task</returns>
+		async ValueTask TickLeaderAsync(CancellationToken stoppingToken)
+		{
+			List<ISession> sessions = await _sessionCollection.FindActiveSessionsAsync();
+			Dictionary<SessionId, ISession> sessionIdToInstance = sessions.ToDictionary(x => x.Id, x => x);
+
+			// Find any leases that are still running when their session has terminated
+			List<ILease> leases = await _leaseCollection.FindActiveLeasesAsync();
+			foreach (ILease lease in leases)
+			{
+				if (!sessionIdToInstance.ContainsKey(lease.SessionId))
+				{
+					ISession? session = await _sessionCollection.GetAsync(lease.SessionId);
+					DateTime finishTime = session?.FinishTime ?? DateTime.UtcNow;
+					_logger.LogWarning("Setting finish time for lease {LeaseId} to {FinishTime}", lease.Id, finishTime);
+					await _leaseCollection.TrySetOutcomeAsync(lease.Id, finishTime, LeaseOutcome.Cancelled, null);
+				}
+			}
+		}
+	}
+}
diff --git a/Engine/Source/Programs/Horde/Horde.Build/Server/DashboardController.Messages.cs b/Engine/Source/Programs/Horde/Horde.Build/Server/DashboardController.Messages.cs
new file mode 100644
--- /dev/null
+++ b/Engine/Source/Programs/Horde/Horde.Build/Server/DashboardController.Messages.cs
@@ -0,0 +1,44 @@
+// Copyright Epic Games, Inc. All Rights Reserved.
+
+namespace Horde.Build.Api
+{
+	/// <summary>
+	/// Setting information required by dashboard
+	/// </summary>
+	public class GetDashboardConfigResponse
+	{
+		/// <summary>
+		/// The name of the external issue service
+		/// </summary>
+		public string? ExternalIssueServiceName { get; set; }
+
+		/// <summary>
+		/// The url of the external issue service
+		/// </summary>
+		public string? ExternalIssueServiceUrl { get; set; }
+
+		/// <summary>
+		/// The url of the perforce swarm installation
+		/// </summary>
+		public string? PerforceSwarmUrl { get; set; }
+
+		/// <summary>
+		/// Help email address that users can contact with issues
+		/// </summary>
+		public string? HelpEmailAddress { get; set; }
+
+		/// <summary>
+		/// Help slack channel that users can use for issues
+		/// </summary>
+		public string? HelpSlackChannel { get; set; }
+
+		/// <summary>
+		/// Response constructor
+		/// </summary>
+		public GetDashboardConfigResponse()
+		{
+		}
+
+	}
+}
+
diff --git a/Engine/Source/Programs/Horde/Horde.Build/Server/DebugController.cs b/Engine/Source/Programs/Horde/Horde.Build/Server/DebugController.cs
new file mode 100644
--- /dev/null
+++ b/Engine/Source/Programs/Horde/Horde.Build/Server/DebugController.cs
@@ -0,0 +1,895 @@
+// Copyright Epic Games, Inc. All Rights Reserved.
+
+//#define ENABLE_PUBLIC_DEBUG_CONTROLLER
+//#define ENABLE_SECURE_DEBUG_CONTROLLER
+
+#if ENABLE_PUBLIC_DEBUG_CONTROLLER
+
+using Horde.Build.Collections;
+using Horde.Build.Models;
+using Horde.Build.Utilities;
+
+namespace Horde.Build.Controllers
+{
+	/// <summary>
+	/// Public endpoints for the debug controller
+	/// </summary>
+	[ApiController]
+	public class PublicDebugController : ControllerBase
+	{
+		/// <summary>
+		/// The connection tracker service singleton
+		/// </summary>
+		RequestTrackerService RequestTrackerService;
+
+		IHostApplicationLifetime ApplicationLifetime;
+
+		IDogStatsd DogStatsd;
+
+		/// <summary>
+		/// Constructor
+		/// </summary>
+		/// <param name="RequestTrackerService"></param>
+		/// <param name="ApplicationLifetime"></param>
+		/// <param name="DogStatsd"></param>
+		public PublicDebugController(RequestTrackerService RequestTrackerService, IHostApplicationLifetime ApplicationLifetime, IDogStatsd DogStatsd)
+		{
+			RequestTrackerService = RequestTrackerService;
+			ApplicationLifetime = ApplicationLifetime;
+			DogStatsd = DogStatsd;
+		}
+
+		/// <summary>
+		/// Prints all the headers for the incoming request
+		/// </summary>
+		/// <returns>Http result</returns>
+		[HttpGet]
+		[Route("/api/v1/debug/headers")]
+		public IActionResult GetRequestHeaders()
+		{
+			StringBuilder Content = new StringBuilder();
+			Content.AppendLine("<html><body><pre>");
+			foreach (KeyValuePair<string, StringValues> Pair in HttpContext.Request.Headers)
+			{
+				foreach (string Value in Pair.Value)
+				{
+					Content.AppendLine(HttpUtility.HtmlEncode($"{Pair.Key}: {Value}"));
+				}
+			}
+			Content.Append("</pre></body></html>");
+			return new ContentResult { ContentType = "text/html", StatusCode = (int)HttpStatusCode.OK, Content = Content.ToString() };
+		}
+
+		/// <summary>
+		/// Waits specified number of milliseconds and then returns a response
+		/// Used for testing timeouts proxy settings.
+		/// </summary>
+		/// <returns>Http result</returns>
+		[HttpGet]
+		[Route("/api/v1/debug/wait")]
+		public async Task<ActionResult> GetAndWait([FromQuery] int WaitTimeMs = 1000)
+		{
+			await Task.Delay(WaitTimeMs);
+			string Content = $"Waited {WaitTimeMs} ms. " + new Random().Next(0, 10000000);
+			return new ContentResult { ContentType = "text/plain", StatusCode = (int)HttpStatusCode.OK, Content = Content };
+		}
+
+		/// <summary>
+		/// Waits specified number of milliseconds and then throws an exception
+		/// Used for testing graceful shutdown and interruption of outstanding requests.
+		/// </summary>
+		/// <returns>Http result</returns>
+		[HttpGet]
+		[Route("/api/v1/debug/exception")]
+		public async Task<ActionResult> ThrowException([FromQuery] int WaitTimeMs = 0)
+		{
+			await Task.Delay(WaitTimeMs);
+			throw new Exception("Test exception triggered by debug controller!");
+		}
+
+		/// <summary>
+		/// Trigger an increment of a DogStatsd metric
+		/// </summary>
+		/// <returns>Http result</returns>
+		[HttpGet]
+		[Route("/api/v1/debug/metric")]
+		public ActionResult TriggerMetric([FromQuery] int Value = 10)
+		{
+			DogStatsd.Increment("hordeMetricTest", Value);
+			return Ok("Incremented metric 'hordeMetricTest' Type: " + DogStatsd.GetType());
+		}
+
+		/// <summary>
+		/// Display metrics related to the .NET runtime
+		/// </summary>
+		/// <returns>Http result</returns>
+		[HttpGet]
+		[Route("/api/v1/debug/dotnet-metrics")]
+		public ActionResult DotNetMetrics()
+		{
+			ThreadPool.GetMaxThreads(out int MaxWorkerThreads, out int MaxIoThreads);
+			ThreadPool.GetAvailableThreads(out int FreeWorkerThreads, out int FreeIoThreads);
+			ThreadPool.GetMinThreads(out int MinWorkerThreads, out int MinIoThreads);
+
+			int BusyIoThreads = MaxIoThreads - FreeIoThreads;
+			int BusyWorkerThreads = MaxWorkerThreads - FreeWorkerThreads;
+
+			StringBuilder Content = new StringBuilder();
+			Content.AppendLine("Threads:");
+			Content.AppendLine("-------------------------------------------------------------");
+			Content.AppendLine("Worker busy={0,-5} free={1,-5} min={2,-5} max={3,-5}", BusyWorkerThreads, FreeWorkerThreads, MinWorkerThreads, MaxWorkerThreads);
+			Content.AppendLine("  IOCP busy={0,-5} free={1,-5} min={2,-5} max={3,-5}", BusyIoThreads, FreeIoThreads, MinIoThreads, MaxWorkerThreads);
+
+
+			NumberFormatInfo Nfi = (NumberFormatInfo)CultureInfo.InvariantCulture.NumberFormat.Clone();
+			Nfi.NumberGroupSeparator = " ";
+
+			string FormatBytes(long Number)
+			{
+				return (Number / 1024 / 1024).ToString("#,0", Nfi) + " MB";
+			}
+
+			GCMemoryInfo GcMemoryInfo = GC.GetGCMemoryInfo();
+			Content.AppendLine("");
+			Content.AppendLine("");
+			Content.AppendLine("Garbage collection (GC):");
+			Content.AppendLine("-------------------------------------------------------------");
+			Content.AppendLine("              Latency mode: " + GCSettings.LatencyMode);
+			Content.AppendLine("              Is server GC: " + GCSettings.IsServerGC);
+			Content.AppendLine("              Total memory: " + FormatBytes(GC.GetTotalMemory(false)));
+			Content.AppendLine("           Total allocated: " + FormatBytes(GC.GetTotalAllocatedBytes(false)));
+			Content.AppendLine("                 Heap size: " + FormatBytes(GcMemoryInfo.HeapSizeBytes));
+			Content.AppendLine("                Fragmented: " + FormatBytes(GcMemoryInfo.FragmentedBytes));
+			Content.AppendLine("               Memory Load: " + FormatBytes(GcMemoryInfo.MemoryLoadBytes));
+			Content.AppendLine("    Total available memory: " + FormatBytes(GcMemoryInfo.TotalAvailableMemoryBytes));
+			Content.AppendLine("High memory load threshold: " + FormatBytes(GcMemoryInfo.HighMemoryLoadThresholdBytes));
+
+			return Ok(Content.ToString());
+		}
+
+		/// <summary>
+		/// Force a full GC of all generations
+		/// </summary>
+		/// <returns>Prints time taken in ms</returns>
+		[HttpGet]
+		[Route("/api/v1/debug/force-gc")]
+		public ActionResult ForceTriggerGc()
+		{
+			Stopwatch Timer = new Stopwatch();
+			Timer.Start();
+			GC.Collect();
+			Timer.Stop();
+			return Ok($"Time taken: {Timer.Elapsed.TotalMilliseconds} ms");
+		}
+
+		/// <summary>
+		/// Lists requests in progress
+		/// </summary>
+		/// <returns>HTML result</returns>
+		[HttpGet]
+		[Route("/api/v1/debug/requests-in-progress")]
+		public ActionResult GetRequestsInProgress()
+		{
+			StringBuilder Content = new StringBuilder();
+			Content.AppendLine("<html><body>");
+			Content.AppendLine("<h1>Requests in progress</h1>");
+			Content.AppendLine("<table border=\"1\">");
+			Content.AppendLine("<tr>");
+			Content.AppendLine("<th>Request Trace ID</th>");
+			Content.AppendLine("<th>Path</th>");
+			Content.AppendLine("<th>Started At</th>");
+			Content.AppendLine("<th>Age</th>");
+			Content.AppendLine("</tr>");
+
+			List<KeyValuePair<string, TrackedRequest>> Requests = RequestTrackerService.GetRequestsInProgress().ToList();
+			Requests.Sort((A, B) => A.Value.StartedAt.CompareTo(B.Value.StartedAt));
+
+			foreach (KeyValuePair<string, TrackedRequest> Entry in Requests)
+			{
+				Content.Append("<tr>");
+				Content.AppendLine($"<td>{Entry.Key}</td>");
+				Content.AppendLine($"<td>{Entry.Value.Request.Path}</td>");
+				Content.AppendLine($"<td>{Entry.Value.StartedAt}</td>");
+				Content.AppendLine($"<td>{Entry.Value.GetTimeSinceStartInMs()} ms</td>");
+				Content.Append("</tr>");
+			}
+			Content.Append("</table>\n</body>\n</html>");
+
+			return new ContentResult { ContentType = "text/html", StatusCode = (int)HttpStatusCode.OK, Content = Content.ToString() };
+		}
+
+		/*
+		// Used during development only
+		[HttpGet]
+		[Route("/api/v1/debug/stop")]
+		public ActionResult StopApp()
+		{
+			Task.Run(async () =>
+			{
+				await Task.Delay(100);
+				ApplicationLifetime.StopApplication();
+			});
+			
+			return new ContentResult { ContentType = "text/plain", StatusCode = (int)HttpStatusCode.OK, Content = "App stopping..." };
+		}
+		/**/
+	}
+#endif
+#if ENABLE_SECURE_DEBUG_CONTROLLER
+	/// <summary>
+	/// Controller managing account status
+	/// </summary>
+	[ApiController]
+	[Authorize]
+	public class SecureDebugController : ControllerBase
+	{
+		/// <summary>
+		/// The ACL service singleton
+		/// </summary>
+		AclService AclService;
+
+		/// <summary>
+		/// The database service instance
+		/// </summary>
+		DatabaseService DatabaseService;
+
+		/// <summary>
+		/// The job task source singelton
+		/// </summary>
+		JobTaskSource JobTaskSource;
+
+		/// <summary>
+		/// The globals singleton
+		/// </summary>
+		ISingletonDocument<Globals> GlobalsSingleton;
+
+		/// <summary>
+		/// Collection of pool documents
+		/// </summary>
+		IPoolCollection PoolCollection;
+
+		/// <summary>
+		/// Collection of project documents
+		/// </summary>
+		IProjectCollection ProjectCollection;
+
+		/// <summary>
+		/// Collection of agent documents
+		/// </summary>
+		IAgentCollection AgentCollection;
+
+		/// <summary>
+		/// Collection of session documents
+		/// </summary>
+		ISessionCollection SessionCollection;
+
+		/// <summary>
+		/// Collection of lease documents
+		/// </summary>
+		ILeaseCollection LeaseCollection;
+
+		/// <summary>
+		/// Collection of template documents
+		/// </summary>
+		ITemplateCollection TemplateCollection;
+
+		/// <summary>
+		/// Collection of stream documents
+		/// </summary>
+		IStreamCollection StreamCollection;
+
+		/// <summary>
+		/// The graph collection singleton
+		/// </summary>
+		IGraphCollection GraphCollection;
+
+		/// <summary>
+		/// The log file collection singleton
+		/// </summary>
+		ILogFileCollection LogFileCollection;
+
+		/// <summary>
+		/// Perforce client
+		/// </summary>
+		IPerforceService Perforce;
+
+		/// <summary>
+		/// Fleet manager
+		/// </summary>
+		IFleetManager FleetManager;
+
+		/// <summary>
+		/// Settings
+		/// </summary>
+		private IOptionsMonitor<ServerSettings> Settings;
+
+		/// <summary>
+		/// Constructor
+		/// </summary>
+		/// <param name="AclService">The ACL service singleton</param>
+		/// <param name="DatabaseService">The database service instance</param>
+		/// <param name="JobTaskSource">The dispatch service singleton</param>
+		/// <param name="GlobalsSingleton">The globals singleton</param>
+		/// <param name="PoolCollection">Collection of pool documents</param>
+		/// <param name="ProjectCollection">Collection of project documents</param>
+		/// <param name="AgentCollection">Collection of agent documents</param>
+		/// <param name="SessionCollection">Collection of session documents</param>
+		/// <param name="LeaseCollection">Collection of lease documents</param>
+		/// <param name="TemplateCollection">Collection of template documents</param>
+		/// <param name="StreamCollection">Collection of stream documents</param>
+		/// <param name="GraphCollection">The graph collection</param>
+		/// <param name="LogFileCollection">The log file collection</param>
+		/// <param name="Perforce">Perforce client</param>
+		/// <param name="FleetManager">The default fleet manager</param>
+		/// <param name="Settings">Settings</param>
+		public SecureDebugController(AclService AclService, DatabaseService DatabaseService, JobTaskSource JobTaskSource, ISingletonDocument<Globals> GlobalsSingleton, IPoolCollection PoolCollection, IProjectCollection ProjectCollection, IAgentCollection AgentCollection, ISessionCollection SessionCollection, ILeaseCollection LeaseCollection, ITemplateCollection TemplateCollection, IStreamCollection StreamCollection, IGraphCollection GraphCollection, ILogFileCollection LogFileCollection, IPerforceService Perforce, IFleetManager FleetManager, IOptionsMonitor<ServerSettings> Settings)
+		{
+			AclService = AclService;
+			DatabaseService = DatabaseService;
+			JobTaskSource = JobTaskSource;
+			GlobalsSingleton = GlobalsSingleton;
+			PoolCollection = PoolCollection;
+			ProjectCollection = ProjectCollection;
+			AgentCollection = AgentCollection;
+			SessionCollection = SessionCollection;
+			LeaseCollection = LeaseCollection;
+			TemplateCollection = TemplateCollection;
+			StreamCollection = StreamCollection;
+			GraphCollection = GraphCollection;
+			LogFileCollection = LogFileCollection;
+			Perforce = Perforce;
+			FleetManager = FleetManager;
+			Settings = Settings;
+		}
+
+		/// <summary>
+		/// Prints all the environment variables
+		/// </summary>
+		/// <returns>Http result</returns>
+		[HttpGet]
+		[Route("/api/v1/debug/environment")]
+		public async Task<ActionResult> GetServerEnvVars()
+		{
+			if (!await AclService.AuthorizeAsync(AclAction.AdminRead, User))
+			{
+				return Forbid();
+			}
+
+			StringBuilder Content = new StringBuilder();
+			Content.AppendLine("<html><body><pre>");
+			foreach (System.Collections.DictionaryEntry? Pair in System.Environment.GetEnvironmentVariables())
+			{
+				if (Pair != null)
+				{
+					Content.AppendLine(HttpUtility.HtmlEncode($"{Pair.Value.Key}={Pair.Value.Value}"));
+				}
+			}
+			Content.Append("</pre></body></html>");
+			return new ContentResult { ContentType = "text/html", StatusCode = (int)HttpStatusCode.OK, Content = Content.ToString() };
+		}
+
+		/// <summary>
+		/// Returns diagnostic information about the current state of the queue
+		/// </summary>
+		/// <returns>Information about the queue</returns>
+		[HttpGet]
+		[Route("/api/v1/debug/queue")]
+		public async Task<ActionResult<object>> GetQueueStatusAsync()
+		{
+			if (!await AclService.AuthorizeAsync(AclAction.AdminRead, User))
+			{
+				return Forbid();
+			}
+
+			return JobTaskSource.GetStatus();
+		}
+
+		/// <summary>
+		/// Returns the complete config Horde uses
+		/// </summary>
+		/// <returns>Information about the config</returns>
+		[HttpGet]
+		[Route("/api/v1/debug/config")]
+		public async Task<ActionResult> GetConfig()
+		{
+			if (!await AclService.AuthorizeAsync(AclAction.AdminRead, User))
+			{
+				return Forbid();
+			}
+
+			JsonSerializerOptions Options = new JsonSerializerOptions
+			{
+				WriteIndented = true
+			};
+
+			return Ok(JsonSerializer.Serialize(Settings.CurrentValue, Options));
+		}
+
+		/// <summary>
+		/// Queries for all graphs
+		/// </summary>
+		/// <returns>The graph definitions</returns>
+		[HttpGet]
+		[Route("/api/v1/debug/graphs")]
+		[ProducesResponseType(200, Type = typeof(GetGraphResponse))]
+		public async Task<ActionResult<List<object>>> GetGraphsAsync([FromQuery] int? Index = null, [FromQuery] int? Count = null, [FromQuery] PropertyFilter? Filter = null)
+		{
+			if (!await AclService.AuthorizeAsync(AclAction.AdminRead, User))
+			{
+				return Forbid();
+			}
+
+			List<IGraph> Graphs = await GraphCollection.FindAllAsync(null, Index, Count);
+			return Graphs.ConvertAll(x => new GetGraphResponse(x).ApplyFilter(Filter));
+		}
+
+		/// <summary>
+		/// Queries for a particular graph by hash
+		/// </summary>
+		/// <returns>The graph definition</returns>
+		[HttpGet]
+		[Route("/api/v1/debug/graphs/{GraphId}")]
+		[ProducesResponseType(200, Type = typeof(GetGraphResponse))]
+		public async Task<ActionResult<object>> GetGraphAsync(string GraphId, [FromQuery] PropertyFilter? Filter = null)
+		{
+			if (!await AclService.AuthorizeAsync(AclAction.AdminRead, User))
+			{
+				return Forbid();
+			}
+
+			IGraph Graph = await GraphCollection.GetAsync(ContentHash.Parse(GraphId));
+			return new GetGraphResponse(Graph).ApplyFilter(Filter);
+		}
+
+		/// <summary>
+		/// Query all the job templates.
+		/// </summary>
+		/// <param name="Filter">Filter for properties to return</param>
+		/// <returns>Information about all the job templates</returns>
+		[HttpGet]
+		[Route("/api/v1/debug/templates")]
+		[ProducesResponseType(typeof(List<GetTemplateResponse>), 200)]
+		public async Task<ActionResult<object>> GetTemplatesAsync([FromQuery] PropertyFilter? Filter = null)
+		{
+			if (!await AclService.AuthorizeAsync(AclAction.AdminRead, User))
+			{
+				return Forbid();
+			}
+
+			List<ITemplate> Templates = await TemplateCollection.FindAllAsync();
+			return Templates.ConvertAll(x => new GetTemplateResponse(x).ApplyFilter(Filter));
+		}
+
+		/// <summary>
+		/// Retrieve information about a specific job template.
+		/// </summary>
+		/// <param name="TemplateHash">Id of the template to get information about</param>
+		/// <param name="Filter">List of properties to return</param>
+		/// <returns>Information about the requested template</returns>
+		[HttpGet]
+		[Route("/api/v1/debug/templates/{TemplateHash}")]
+		[ProducesResponseType(typeof(GetTemplateResponse), 200)]
+		public async Task<ActionResult<object>> GetTemplateAsync(string TemplateHash, [FromQuery] PropertyFilter? Filter = null)
+		{
+			ContentHash TemplateHashValue = ContentHash.Parse(TemplateHash);
+			if (!await AclService.AuthorizeAsync(AclAction.AdminRead, User))
+			{
+				return Forbid();
+			}
+
+			ITemplate? Template = await TemplateCollection.GetAsync(TemplateHashValue);
+			if (Template == null)
+			{
+				return NotFound();
+			}
+			return Template.ApplyFilter(Filter);
+		}
+
+		/// <summary>
+		/// Retrieve metadata about a specific log file
+		/// </summary>
+		/// <param name="LogFileId">Id of the log file to get information about</param>
+		/// <param name="Filter">Filter for the properties to return</param>
+		/// <returns>Information about the requested project</returns>
+		[HttpGet]
+		[Route("/api/v1/debug/logs/{LogFileId}")]
+		public async Task<ActionResult<object>> GetLogAsync(LogId LogFileId, [FromQuery] PropertyFilter? Filter = null)
+		{
+			if (!await AclService.AuthorizeAsync(AclAction.AdminRead, User))
+			{
+				return Forbid();
+			}
+
+			ILogFile? LogFile = await LogFileCollection.GetLogFileAsync(LogFileId);
+			if (LogFile == null)
+			{
+				return NotFound();
+			}
+
+			return LogFile.ApplyFilter(Filter);
+		}
+
+		/// <summary>
+		/// Populate the database with test data
+		/// </summary>
+		/// <returns>Async task</returns>
+		[HttpGet]
+		[Route("/api/v1/debug/collections/{Name}")]
+		public async Task<ActionResult<object>> GetDocumentsAsync(string Name, [FromQuery] string? Filter = null, [FromQuery] int Index = 0, [FromQuery] int Count = 10)
+		{
+			if (!await AclService.AuthorizeAsync(AclAction.AdminRead, User))
+			{
+				return Forbid();
+			}
+
+			IMongoCollection<Dictionary<string, object>> Collection = DatabaseService.GetCollection<Dictionary<string, object>>(Name);
+			List<Dictionary<string, object>> Documents = await Collection.Find(Filter ?? "{}").Skip(Index).Limit(Count).ToListAsync();
+			return Documents;
+		}
+
+		/// <summary>
+		/// Create P4 login ticket for the specified username
+		/// </summary>
+		/// <returns>Perforce ticket</returns>
+		[HttpGet]
+		[Route("/api/v1/debug/p4ticket")]
+		public async Task<ActionResult<string>> CreateTicket([FromQuery] string? Username = null)
+		{
+			if (!await AclService.AuthorizeAsync(AclAction.AdminRead, User))
+			{
+				return Forbid();
+			}
+
+			if (Username == null)
+			{
+				return BadRequest();
+			}
+
+			return await Perforce.CreateTicket(PerforceCluster.DefaultName, Username);
+		}
+
+		/// <summary>
+		/// Debugs perforce service commands
+		/// </summary>
+		/// <returns>Perforce ticket</returns>
+		[HttpGet]
+		[Route("/api/v1/debug/perforce/stress")]
+		public async Task<ActionResult<string>> GetPerforceStress()
+		{
+			if (!await AclService.AuthorizeAsync(AclAction.AdminWrite, User))
+			{
+				return Forbid();
+			}
+
+			Dictionary<string, string> Results = new Dictionary<string, string>();
+			const string DevBuild = "//UE4/Dev-Build";
+			string? PerforceUser = User.GetPerforceUser();
+
+			object UpdateLock = new object();
+
+			// ParallelTest.ForAsync doesn't seem to handle async lambdas
+			IEnumerable<Task> Tasks = Enumerable.Range(0, 16).Select(async (Idx) => {
+
+				int LatestChange = await Perforce.GetLatestChangeAsync(PerforceCluster.DefaultName, DevBuild, PerforceUser);
+				int CodeChange = await Perforce.GetCodeChangeAsync(PerforceCluster.DefaultName, DevBuild, LatestChange);				
+
+				lock (UpdateLock)
+				{
+					Results.Add($"Parallel Perforce Result {Idx}", $"LatestChange: {LatestChange} - CodeChange: {CodeChange}");
+				}
+			});
+
+			await Task.WhenAll(Tasks);			
+
+			Tasks = Enumerable.Range(0, 16).Select(async (Idx) => {
+
+				int LatestChange = await Perforce.GetLatestChangeAsync(PerforceCluster.DefaultName, DevBuild, PerforceUser);
+				int CodeChange = await Perforce.GetCodeChangeAsync(PerforceCluster.DefaultName, DevBuild, LatestChange);				
+
+				lock (UpdateLock)
+				{
+					Results.Add($"Parallel Perforce Result {Idx + 16}", $"LatestChange: {LatestChange} - CodeChange: {CodeChange}");
+				}
+			});
+
+
+			await Task.WhenAll(Tasks);
+
+			StringBuilder Content = new StringBuilder();
+			Content.AppendLine("<html><body><pre>");
+			foreach (KeyValuePair<string, string> Pair in Results)
+			{
+				Content.AppendLine(HttpUtility.HtmlEncode($"{Pair.Key}={Pair.Value}"));
+			}
+			Content.Append("</pre></body></html>");
+			return new ContentResult { ContentType = "text/html", StatusCode = (int)HttpStatusCode.OK, Content = Content.ToString() };
+
+		}
+
+		/// <summary>
+		/// Debugs perforce service commands
+		/// </summary>
+		/// <returns>Perforce ticket</returns>
+		[HttpGet]
+		[Route("/api/v1/debug/perforce/commands")]
+		public async Task<ActionResult<string>> GetPerforceCommands()
+		{
+			if (!await AclService.AuthorizeAsync(AclAction.AdminWrite, User))
+			{
+				return Forbid();
+			}
+
+			StringBuilder Content = new StringBuilder();
+
+			const string DevBuild = "//UE4/Dev-Build";
+			// a test CL I have setup in Dev-Build
+			const int TestCL = 16687685;
+
+			string? PerforceUser = User.GetPerforceUser();
+
+			if (string.IsNullOrEmpty(PerforceUser))
+			{
+				throw new Exception("Perforce user required");
+			}
+
+			Dictionary<string, string> Results = new Dictionary<string, string>();
+
+			// NOTE: This is in Private-Build not DevBuild
+			int NewChangeId = await Perforce.CreateNewChangeAsync(PerforceCluster.DefaultName, "//UE4/Private-Build", "Counter.txt");
+
+			List<ChangeDetails> ChangeDetails = await Perforce.GetChangeDetailsAsync(PerforceCluster.DefaultName, "//UE4/Private-Build", new int[] { NewChangeId }, PerforceUser);
+
+			Results.Add("CreateNewChangeAsync", $"{NewChangeId} : {ChangeDetails[0].Description}");
+
+			await Perforce.UpdateChangelistDescription(PerforceCluster.DefaultName, TestCL, $"Updated Description: New Change CL was {NewChangeId}");
+
+			ChangeDetails = await Perforce.GetChangeDetailsAsync(PerforceCluster.DefaultName, DevBuild, new int[] { TestCL }, PerforceUser);
+
+			Results.Add("UpdateChangelistDescription", $"{TestCL} : {ChangeDetails[0].Description}");
+
+			int LatestChange = await Perforce.GetLatestChangeAsync(PerforceCluster.DefaultName, DevBuild, PerforceUser);
+			Results.Add("GetLatestChangeAsync", LatestChange.ToString(CultureInfo.InvariantCulture));
+
+			int CodeChange = await Perforce.GetCodeChangeAsync(PerforceCluster.DefaultName, DevBuild, LatestChange);
+			Results.Add("GetCodeChangeAsync", CodeChange.ToString(CultureInfo.InvariantCulture));
+
+			PerforceUserInfo? UserInfo = await Perforce.GetUserInfoAsync(PerforceCluster.DefaultName, PerforceUser);
+
+			if (UserInfo == null)
+			{
+				Results.Add("GetUserInfoAsync", "null");
+			}
+			else
+			{
+				Results.Add("GetUserInfoAsync", $"{UserInfo.Login} : {UserInfo.Email}");
+			}
+
+			// changes
+			List<ChangeSummary> Changes = await Perforce.GetChangesAsync(PerforceCluster.DefaultName, DevBuild, null, LatestChange, 10, PerforceUser);
+			int Count = 0;
+			foreach (ChangeSummary Summary in Changes)
+			{
+				Results.Add($"GetChangesAsync Result {Count++}", $"{Summary.Number} : {Summary.Author} - {Summary.Description}");
+			}
+
+			ChangeDetails = await Perforce.GetChangeDetailsAsync(PerforceCluster.DefaultName, DevBuild, Changes.Select(Change => Change.Number).ToArray(), PerforceUser);
+			Count = 0;
+			foreach (ChangeDetails Details in ChangeDetails)
+			{
+				Results.Add($"GetChangeDetailsAsync Result {Count++}", $"{Details.Number} : {Details.Author} - {Details.Files.ToJson()}");
+			}
+
+			String Ticket = "Failed (Expected if not running service account)";
+			try
+			{
+				await Perforce.CreateTicket(PerforceCluster.DefaultName, PerforceUser);
+				Ticket = "Success";
+			}
+			catch
+			{
+
+			}
+
+			Results.Add($"CreateTicket", $"{Ticket}");
+			List<FileSummary> FileSummaries = await Perforce.FindFilesAsync(PerforceCluster.DefaultName, ChangeDetails.SelectMany(Details => Details.Files).Select(File => $"{DevBuild}{File}"));
+			Count = 0;
+			foreach (FileSummary File in FileSummaries)
+			{
+				Results.Add($"FindFilesAsync Result {Count++}", $"{File.Change} : {File.DepotPath} : {File.Exists}");
+			}
+
+			// print async
+
+			byte[] Bytes = await Perforce.PrintAsync(PerforceCluster.DefaultName, $"{DevBuild}/RunUAT.bat");
+			Results.Add($"PrintAsync: {DevBuild}/RunUAT.bat ", $"{System.Text.Encoding.Default.GetString(Bytes)}");
+
+			Bytes = await Perforce.PrintAsync(PerforceCluster.DefaultName, $"{DevBuild}/RunUAT.bat@13916380");
+			Results.Add($"PrintAsync: {DevBuild}/RunUAT.bat@13916380", $"{System.Text.Encoding.Default.GetString(Bytes)}");
+
+			// need to be running as service account for these
+
+			int? DuplicateCL = null;
+
+			int ShelvedChangeOnDevBuild = 16687685;
+
+			try 
+			{
+				DuplicateCL = await Perforce.DuplicateShelvedChangeAsync(PerforceCluster.DefaultName, ShelvedChangeOnDevBuild);
+				Results.Add("DuplicateShelvedChangeAsync", $"Duplicated CL {DuplicateCL} from Shelved CL {ShelvedChangeOnDevBuild}");
+			}
+			catch 
+			{
+				Results.Add($"DuplicateShelvedChangeAsync",  "Failed - expected unless using service account");
+			}
+
+			if (DuplicateCL.HasValue)
+			{
+				// Note change client much exist for the update
+				ChangeDetails = await Perforce.GetChangeDetailsAsync(PerforceCluster.DefaultName, DevBuild, new int[]{ DuplicateCL.Value}, PerforceUser);
+
+				Results.Add("UpdateChangelistDescription - PreUpdate", $"{DuplicateCL} : {ChangeDetails[0].Description}"); 
+
+				try 
+				{
+					await Perforce.UpdateChangelistDescription(PerforceCluster.DefaultName, DuplicateCL.Value, "Updated Description from Horde");
+
+					ChangeDetails = await Perforce.GetChangeDetailsAsync(PerforceCluster.DefaultName, DevBuild, new int[]{ DuplicateCL.Value}, PerforceUser);
+
+					Results.Add("UpdateChangelistDescription - PostUpdate", $"{DuplicateCL} : {ChangeDetails[0].Description}"); 
+				}
+				catch 
+				{
+					Results.Add("UpdateChangelistDescription - Faile", $"Client from change {DuplicateCL} must exist"); 
+				}
+
+				await Perforce.DeleteShelvedChangeAsync(PerforceCluster.DefaultName, DuplicateCL.Value);
+
+				ChangeDetails = await Perforce.GetChangeDetailsAsync(PerforceCluster.DefaultName, DevBuild, new int[] {DuplicateCL.Value} , PerforceUser);
+
+				Results.Add("GetChangeDetailsAsync - After Delete", $"ChangeDetails.Count: {ChangeDetails.Count}"); 
+
+			}
+
+			byte[] ImageBytes = await Perforce.PrintAsync(PerforceCluster.DefaultName, $"{DevBuild}/Samples/Games/ShooterGame/ShooterGame.png");
+
+			Content.AppendLine("<html><body><pre>");
+			foreach (KeyValuePair<string, string> Pair in Results)
+			{
+				Content.AppendLine(HttpUtility.HtmlEncode($"{Pair.Key}={Pair.Value}"));
+			}
+
+			Content.Append($"PrintAsync Binary: {DevBuild}/Samples/Games/ShooterGame/ShooterGame.png\n</pre>");
+			Content.Append($"<img src=\"data:image/png;base64,{Convert.ToBase64String(ImageBytes, Base64FormattingOptions.None)}\"/>");
+
+			Content.Append("</body></html>");
+			return new ContentResult { ContentType = "text/html", StatusCode = (int)HttpStatusCode.OK, Content = Content.ToString() };
+
+		}
+
+		/// <summary>
+		/// Simulates a crash of the native p4 bridge
+		/// </summary>		
+		[HttpGet]		
+		[Route("/api/v1/debug/perforce/debugcrash")]		
+		public async Task<ActionResult> GetPerforceDebugCrash()
+		{
+			if (!await AclService.AuthorizeAsync(AclAction.AdminWrite, User))
+			{
+				return Forbid();
+			}
+
+			P4Debugging.DebugCrash();
+
+			return Ok();
+		}
+
+		/// <summary>
+		/// Debugging fleet managers
+		/// </summary>
+		/// <returns>Nothing</returns>
+		[HttpGet]
+		[Route("/api/v1/debug/fleetmanager")]
+		public async Task<ActionResult<string>> FleetManage([FromQuery] string? PoolId = null, [FromQuery] string? ChangeType = null, [FromQuery] int? Count = null)
+		{
+			if (!await AclService.AuthorizeAsync(AclAction.AdminRead, User))
+			{
+				return Forbid();
+			}
+
+			if (PoolId == null)
+			{
+				return BadRequest("Missing 'PoolId' query parameter");
+			}
+
+			IPool? Pool = await PoolCollection.GetAsync(new PoolId(PoolId));
+			if (Pool == null)
+			{
+				return BadRequest($"Pool with ID '{PoolId}' not found");
+			}
+
+			if (!(ChangeType == "expand" || ChangeType == "shrink"))
+			{
+				return BadRequest("Missing 'ChangeType' query parameter and/or must be set to 'expand' or 'shrink'");
+			}
+
+			if (Count == null || Count.Value <= 0)
+			{
+				return BadRequest("Missing 'Count' query parameter and/or must be greater than 0");
+			}
+
+			List<IAgent> Agents = new List<IAgent>();
+			string Message = "No operation";
+
+			if (ChangeType == "expand")
+			{
+				await FleetManager.ExpandPool(Pool, Agents, Count.Value);
+				Message = $"Expanded pool {Pool.Name} by {Count}";
+			}
+			else if (ChangeType == "shrink")
+			{
+				Agents = await AgentCollection.FindAsync();
+				Agents = Agents.FindAll(a => a.IsInPool(Pool.Id)).ToList();
+				await FleetManager.ShrinkPool(Pool, Agents, Count.Value);
+				Message = $"Shrank pool {Pool.Name} by {Count}";
+			}
+
+			return new ContentResult { ContentType = "text/plain", StatusCode = (int)HttpStatusCode.OK, Content = Message };
+		}
+		
+		/// <summary>
+		/// Starts the profiler session
+		/// </summary>
+		/// <returns>Text message</returns>
+		[HttpGet]
+		[Route("/api/v1/debug/profiler/start")]
+		public async Task<ActionResult> StartProfiler()
+		{
+			await DotTrace.EnsurePrerequisiteAsync();
+
+			string SnapshotDir = Path.Join(Path.GetTempPath(), "horde-profiler-snapshots");
+			if (!Directory.Exists(SnapshotDir))
+			{
+				Directory.CreateDirectory(SnapshotDir);
+			}
+
+			var Config = new DotTrace.Config();
+			Config.SaveToDir(SnapshotDir);
+			DotTrace.Attach(Config);
+			DotTrace.StartCollectingData();
+			
+			return new ContentResult { ContentType = "text/plain", StatusCode = (int)HttpStatusCode.OK, Content = "Profiling session started. Using dir " + SnapshotDir };
+		}
+		
+		/// <summary>
+		/// Stops the profiler session
+		/// </summary>
+		/// <returns>Text message</returns>
+		[HttpGet]
+		[Route("/api/v1/debug/profiler/stop")]
+		public ActionResult StopProfiler()
+		{
+			DotTrace.SaveData();
+			DotTrace.Detach();
+			return new ContentResult { ContentType = "text/plain", StatusCode = (int)HttpStatusCode.OK, Content = "Profiling session stopped" };
+		}
+		
+		/// <summary>
+		/// Downloads the captured profiling snapshots
+		/// </summary>
+		/// <returns>A .zip file containing the profiling snapshots</returns>
+		[HttpGet]
+		[Route("/api/v1/debug/profiler/download")]
+		public ActionResult DownloadProfilingData()
+		{
+			string SnapshotZipFile = DotTrace.GetCollectedSnapshotFilesArchive(false);
+			if (!System.IO.File.Exists(SnapshotZipFile))
+			{
+				return NotFound("The generated snapshot .zip file was not found");
+			}
+			
+			return PhysicalFile(SnapshotZipFile, "application/zip", Path.GetFileName(SnapshotZipFile));
+		}
+	}
+}
+
+#endif
diff --git a/Engine/Source/Programs/Horde/Horde.Build/Server/Globals.cs b/Engine/Source/Programs/Horde/Horde.Build/Server/Globals.cs
new file mode 100644
--- /dev/null
+++ b/Engine/Source/Programs/Horde/Horde.Build/Server/Globals.cs
@@ -0,0 +1,386 @@
+﻿// Copyright Epic Games, Inc. All Rights Reserved.
+
+using System;
+using System.Collections.Generic;
+using System.ComponentModel.DataAnnotations;
+using System.Linq;
+using System.Security.Cryptography;
+using EpicGames.Perforce;
+using Horde.Build.Acls;
+using Horde.Build.Server;
+using Horde.Build.Utilities;
+using MongoDB.Bson;
+using MongoDB.Bson.Serialization.Attributes;
+
+namespace Horde.Build.Models
+{
+	/// <summary>
+	/// Base class for singleton documents
+	/// </summary>
+	public abstract class SingletonBase
+	{
+		/// <summary>
+		/// Unique id for this singleton
+		/// </summary>
+		[BsonId]
+		public ObjectId Id { get; set; }
+
+		/// <summary>
+		/// The revision index of this document
+		/// </summary>
+		public int Revision { get; set; }
+
+		/// <summary>
+		/// Callback to allow the singleton to fix up itself after being read
+		/// </summary>
+		public virtual void PostLoad()
+		{
+		}
+	}
+
+	/// <summary>
+	/// How frequently the maintence window repeats
+	/// </summary>
+	public enum ScheduledDowntimeFrequency
+	{
+		/// <summary>
+		/// Once
+		/// </summary>
+		Once,
+
+		/// <summary>
+		/// Every day
+		/// </summary>
+		Daily,
+
+		/// <summary>
+		/// Every week
+		/// </summary>
+		Weekly,
+	}
+
+	/// <summary>
+	/// Settings for the maintenance window
+	/// </summary>
+	public class ScheduledDowntime
+	{
+		/// <summary>
+		/// Start time
+		/// </summary>
+		public DateTimeOffset StartTime { get; set; }
+
+		/// <summary>
+		/// Finish time
+		/// </summary>
+		public DateTimeOffset FinishTime { get; set; }
+
+		/// <summary>
+		/// Frequency that the window repeats
+		/// </summary>
+		public ScheduledDowntimeFrequency Frequency { get; set; } = ScheduledDowntimeFrequency.Once;
+
+		/// <summary>
+		/// Gets the next scheduled downtime
+		/// </summary>
+		/// <param name="now">The current time</param>
+		/// <returns>Start and finish time</returns>
+		public (DateTimeOffset StartTime, DateTimeOffset FinishTime) GetNext(DateTimeOffset now)
+		{
+			TimeSpan offset = TimeSpan.Zero;
+			if (Frequency == ScheduledDowntimeFrequency.Daily)
+			{
+				double days = (now - StartTime).TotalDays;
+				if (days >= 1.0)
+				{
+					days -= days % 1.0;
+				}
+				offset = TimeSpan.FromDays(days);
+			}
+			else if (Frequency == ScheduledDowntimeFrequency.Weekly)
+			{
+				double days = (now - StartTime).TotalDays;
+				if (days >= 7.0)
+				{
+					days -= days % 7.0;
+				}
+				offset = TimeSpan.FromDays(days);
+			}
+			return (StartTime + offset, FinishTime + offset);
+		}
+
+		/// <summary>
+		/// Determines if this schedule is active
+		/// </summary>
+		/// <param name="now">The current time</param>
+		/// <returns>True if downtime is active</returns>
+		public bool IsActive(DateTimeOffset now)
+		{
+			if (Frequency == ScheduledDowntimeFrequency.Once)
+			{
+				return now >= StartTime && now < FinishTime;
+			}
+			else if (Frequency == ScheduledDowntimeFrequency.Daily)
+			{
+				double days = (now - StartTime).TotalDays;
+				if (days < 0.0)
+				{
+					return false;
+				}
+				else
+				{
+					return (days % 1.0) < (FinishTime - StartTime).TotalDays;
+				}
+			}
+			else if (Frequency == ScheduledDowntimeFrequency.Weekly)
+			{
+				double days = (now - StartTime).TotalDays;
+				if(days < 0.0)
+				{
+					return false;
+				}
+				else
+				{
+					return (days % 7.0) < (FinishTime - StartTime).TotalDays;
+				}
+			}
+			else
+			{
+				return false;
+			}
+		}
+	}
+
+	/// <summary>
+	/// Path to a platform and stream to use for syncing AutoSDK
+	/// </summary>
+	public class AutoSdkWorkspace
+	{
+		/// <summary>
+		/// Name of this workspace
+		/// </summary>
+		public string? Name { get; set; }
+
+		/// <summary>
+		/// The agent properties to check (eg. "OSFamily=Windows")
+		/// </summary>
+		public List<string> Properties { get; set; } = new List<string>();
+
+		/// <summary>
+		/// Username for logging in to the server
+		/// </summary>
+		public string? UserName { get; set; }
+
+		/// <summary>
+		/// Stream to use
+		/// </summary>
+		[Required]
+		public string? Stream { get; set; }
+	}
+
+	/// <summary>
+	/// Information about an individual Perforce server
+	/// </summary>
+	public class PerforceServer
+	{
+		/// <summary>
+		/// The server and port. The server may be a DNS entry with multiple records, in which case it will be actively load balanced.
+		/// </summary>
+		public string ServerAndPort { get; set; } = "perforce:1666";
+
+		/// <summary>
+		/// Whether to query the healthcheck address under each server
+		/// </summary>
+		public bool HealthCheck { get; set; }
+
+		/// <summary>
+		/// Whether to resolve the DNS entries and load balance between different hosts
+		/// </summary>
+		public bool ResolveDns { get; set; }
+
+		/// <summary>
+		/// Maximum number of simultaneous conforms on this server
+		/// </summary>
+		public int MaxConformCount { get; set; }
+
+		/// <summary>
+		/// List of properties for an agent to be eligable to use this server
+		/// </summary>
+		public List<string>? Properties { get; set; }
+	}
+
+	/// <summary>
+	/// Credentials for a Perforce user
+	/// </summary>
+	public class PerforceCredentials
+	{
+		/// <summary>
+		/// The username
+		/// </summary>
+		public string UserName { get; set; } = String.Empty;
+
+		/// <summary>
+		/// Password for the user
+		/// </summary>
+		public string Password { get; set; } = String.Empty;
+	}
+
+	/// <summary>
+	/// Information about a cluster of Perforce servers. 
+	/// </summary>
+	public class PerforceCluster
+	{
+		/// <summary>
+		/// The default cluster name
+		/// </summary>
+		public const string DefaultName = "Default";
+
+		/// <summary>
+		/// Name of the cluster
+		/// </summary>
+		[Required]
+		public string Name { get; set; } = null!;
+
+		/// <summary>
+		/// Username for Horde to log in to this server. Will use the default user if not set.
+		/// </summary>
+		public string? ServiceAccount { get; set; }
+
+		/// <summary>
+		/// Whether the service account can impersonate other users
+		/// </summary>
+		public bool CanImpersonate { get; set; } = true;
+
+		/// <summary>
+		/// List of servers
+		/// </summary>
+		public List<PerforceServer> Servers { get; set; } = new List<PerforceServer>();
+
+		/// <summary>
+		/// List of server credentials
+		/// </summary>
+		public List<PerforceCredentials> Credentials { get; set; } = new List<PerforceCredentials>();
+
+		/// <summary>
+		/// List of autosdk streams
+		/// </summary>
+		public List<AutoSdkWorkspace> AutoSdk { get; set; } = new List<AutoSdkWorkspace>();
+	}
+
+	/// <summary>
+	/// Global server settings
+	/// </summary>
+	[SingletonDocument("5e3981cb28b8ec59cd07184a")]
+	public class Globals : SingletonBase
+	{
+		/// <summary>
+		/// The unique id for all globals objects
+		/// </summary>
+		public static ObjectId StaticId { get; } = new ObjectId("5e3981cb28b8ec59cd07184a");
+
+		/// <summary>
+		/// Unique instance id of this database
+		/// </summary>
+		public ObjectId InstanceId { get; set; }
+
+		/// <summary>
+		/// The config revision
+		/// </summary>
+		public string? ConfigRevision { get; set; }
+
+		/// <summary>
+		/// List of Perforce clusters
+		/// </summary>
+		public List<PerforceCluster> PerforceClusters { get; set; } = new List<PerforceCluster>();
+
+		/// <summary>
+		/// Maximum number of simultaneous conforms
+		/// </summary>
+		public int MaxConformCount { get; set; }
+
+		/// <summary>
+		/// The root ACL
+		/// </summary>
+		public Acl? RootAcl { get; set; }
+
+		/// <summary>
+		/// Randomly generated JWT signing key
+		/// </summary>
+		public byte[]? JwtSigningKey { get; set; }
+
+		/// <summary>
+		/// The current schema version
+		/// </summary>
+		public int? SchemaVersion { get; set; }
+
+		/// <summary>
+		/// The scheduled downtime
+		/// </summary>
+		public List<ScheduledDowntime> ScheduledDowntime { get; set; } = new List<ScheduledDowntime>();
+
+		/// <summary>
+		/// List of compute profiles
+		/// </summary>
+		public List<ComputeClusterConfig> ComputeClusters { get; set; } = new List<ComputeClusterConfig>();
+
+		/// <summary>
+		/// Private constructor for serialization
+		/// </summary>
+		public Globals()
+		{
+			InstanceId = ObjectId.GenerateNewId();
+		}
+
+		/// <summary>
+		/// Rotate the signing key
+		/// </summary>
+		public void RotateSigningKey()
+		{
+			JwtSigningKey = RandomNumberGenerator.GetBytes(128);
+		}
+
+		/// <summary>
+		/// Finds a perforce cluster with the given name or that contains the provided server
+		/// </summary>
+		/// <param name="name">Name of the cluster</param>
+		/// <param name="serverAndPort">Find cluster which contains server</param>
+		/// <returns></returns>
+		public PerforceCluster? FindPerforceCluster(string? name, string? serverAndPort = null)
+		{
+			List<PerforceCluster> clusters = PerforceClusters;
+
+			if (serverAndPort != null)
+			{
+				return clusters.FirstOrDefault(x => x.Servers.FirstOrDefault( server => String.Equals(server.ServerAndPort, serverAndPort, StringComparison.OrdinalIgnoreCase)) != null);
+			}
+
+			if (clusters.Count == 0)
+			{
+				clusters = DefaultClusters;
+			}
+
+			if (name == null)
+			{
+				return clusters.FirstOrDefault();
+			}
+			else
+			{
+				return clusters.FirstOrDefault(x => String.Equals(x.Name, name, StringComparison.OrdinalIgnoreCase));
+			}
+		}
+
+		static List<PerforceCluster> DefaultClusters { get; } = GetDefaultClusters();
+
+		static List<PerforceCluster> GetDefaultClusters()
+		{
+			PerforceServer server = new PerforceServer();
+			server.ServerAndPort = PerforceSettings.Default.ServerAndPort;
+
+			PerforceCluster cluster = new PerforceCluster();
+			cluster.Name = "Default";
+			cluster.CanImpersonate = false;
+			cluster.Servers.Add(server);
+
+			return new List<PerforceCluster> { cluster };
+		}
+	}
+}
diff --git a/Engine/Source/Programs/Horde/Horde.Build/Server/Notices/INoticeCollection.cs b/Engine/Source/Programs/Horde/Horde.Build/Server/Notices/INoticeCollection.cs
new file mode 100644
--- /dev/null
+++ b/Engine/Source/Programs/Horde/Horde.Build/Server/Notices/INoticeCollection.cs
@@ -0,0 +1,59 @@
+// Copyright Epic Games, Inc. All Rights Reserved.
+
+using System;
+using System.Collections.Generic;
+using System.Threading.Tasks;
+using Horde.Build.Models;
+using Horde.Build.Utilities;
+using MongoDB.Bson;
+
+namespace Horde.Build.Collections
+{
+	using UserId = ObjectId<IUser>;
+
+	/// <summary>
+	/// A collection of user specified notices
+	/// </summary>
+	public interface INoticeCollection
+	{
+		/// <summary>
+		/// Add a notice to the collection
+		/// </summary>
+		/// <param name="message"></param>
+		/// <param name="userId"></param>
+		/// <param name="startTime"></param>
+		/// <param name="finishTime"></param>
+		/// <returns></returns>
+		Task<INotice?> AddNoticeAsync(string message, UserId? userId, DateTime? startTime, DateTime? finishTime);
+
+		/// <summary>
+		/// Update an existing notice
+		/// </summary>
+		/// <param name="id"></param>
+		/// <param name="message"></param>
+		/// <param name="startTime"></param>
+		/// <param name="finishTime"></param>
+		/// <returns></returns>
+		Task<bool> UpdateNoticeAsync(ObjectId id, string? message, DateTime? startTime, DateTime? finishTime);
+
+		/// <summary>
+		/// Get a notice by id
+		/// </summary>
+		/// <param name="noticeId"></param>
+		/// <returns></returns>
+		Task<INotice?> GetNoticeAsync(ObjectId noticeId);
+
+		/// <summary>
+		/// Get all notices
+		/// </summary>
+		/// <returns></returns>
+		Task<List<INotice>> GetNoticesAsync();
+
+		/// <summary>
+		/// Remove a notice
+		/// </summary>
+		/// <param name="id"></param>
+		/// <returns></returns>
+		Task<bool> RemoveNoticeAsync(ObjectId id);
+	}
+}
\ No newline at end of file
diff --git a/Engine/Source/Programs/Horde/Horde.Build/Server/Notices/NoticeCollection.cs b/Engine/Source/Programs/Horde/Horde.Build/Server/Notices/NoticeCollection.cs
new file mode 100644
--- /dev/null
+++ b/Engine/Source/Programs/Horde/Horde.Build/Server/Notices/NoticeCollection.cs
@@ -0,0 +1,118 @@
+// Copyright Epic Games, Inc. All Rights Reserved.
+
+using MongoDB.Bson.Serialization.Attributes;
+using System.Threading.Tasks;
+using MongoDB.Driver;
+using System;
+using System.Collections.Generic;
+using System.Linq;
+using Horde.Build.Models;
+using Horde.Build.Server;
+using Horde.Build.Utilities;
+using MongoDB.Bson;
+
+namespace Horde.Build.Collections.Impl
+{
+	using UserId = ObjectId<IUser>;
+
+	/// <summary>
+	/// Collection of notice documents
+	/// </summary>
+	public class NoticeCollection : INoticeCollection
+	{
+
+		/// <summary>
+		/// Document representing a notice 
+		/// </summary>
+		class NoticeDocument : INotice
+		{
+			[BsonRequired, BsonId]
+			public ObjectId Id { get; set; }
+
+			[BsonIgnoreIfNull]
+			public UserId? UserId { get; set; }
+
+			[BsonIgnoreIfNull]
+			public DateTime? StartTime { get; set; }
+
+			[BsonIgnoreIfNull]
+			public DateTime? FinishTime { get; set; }
+
+			public string Message { get; set; } = String.Empty;
+
+			[BsonConstructor]
+			private NoticeDocument()
+			{
+			}
+
+			public NoticeDocument(ObjectId id, string message, UserId? userId, DateTime? startTime, DateTime? finishTime)
+			{
+				Id = id;
+				Message = message;
+				UserId = userId;
+				StartTime = startTime;
+				FinishTime = finishTime;
+			}
+		}
+
+		readonly IMongoCollection<NoticeDocument> _notices;
+
+		/// <summary>
+		/// Constructor
+		/// </summary>
+		public NoticeCollection(MongoService mongoService)
+		{
+			_notices = mongoService.GetCollection<NoticeDocument>("Notices");
+		}
+
+		/// <inheritdoc/>
+		public async Task<INotice?> AddNoticeAsync(string message, UserId? userId, DateTime? startTime, DateTime? finishTime)
+		{
+			NoticeDocument newNotice = new NoticeDocument(ObjectId.GenerateNewId(), message, userId, startTime, finishTime);
+			await _notices.InsertOneAsync(newNotice);
+			return newNotice;
+		}
+
+		/// <inheritdoc/>
+		public async Task<bool> UpdateNoticeAsync(ObjectId id, string? message, DateTime? startTime, DateTime? finishTime)
+		{
+
+			UpdateDefinitionBuilder<NoticeDocument> updateBuilder = Builders<NoticeDocument>.Update;
+			List<UpdateDefinition<NoticeDocument>> updates = new List<UpdateDefinition<NoticeDocument>>();
+
+			if (message != null)
+			{
+				updates.Add(updateBuilder.Set(x => x.Message, message));
+			}
+
+			updates.Add(updateBuilder.Set(x => x.StartTime, startTime));
+
+			updates.Add(updateBuilder.Set(x => x.FinishTime, finishTime));
+
+			NoticeDocument? document = await _notices.FindOneAndUpdateAsync(x => x.Id == id, updateBuilder.Combine(updates));
+
+			return document != null;
+		}
+
+		/// <inheritdoc/>
+		public async Task<INotice?> GetNoticeAsync(ObjectId noticeId)
+		{
+			return await _notices.Find(x => x.Id == noticeId).FirstOrDefaultAsync();
+		}
+
+		/// <inheritdoc/>
+		public async Task<List<INotice>> GetNoticesAsync()
+		{
+			List<NoticeDocument> results = await _notices.Find(x => true).ToListAsync();
+			return results.Select<NoticeDocument, INotice>(x => x).ToList();
+		}
+
+		/// <inheritdoc/>
+		public async Task<bool> RemoveNoticeAsync(ObjectId id)
+		{
+			DeleteResult result = await _notices.DeleteOneAsync(x => x.Id == id);
+			return result.DeletedCount > 0;
+		}
+	}
+}
+
diff --git a/Engine/Source/Programs/Horde/Horde.Build/Server/Notices/NoticeService.cs b/Engine/Source/Programs/Horde/Horde.Build/Server/Notices/NoticeService.cs
new file mode 100644
--- /dev/null
+++ b/Engine/Source/Programs/Horde/Horde.Build/Server/Notices/NoticeService.cs
@@ -0,0 +1,88 @@
+// Copyright Epic Games, Inc. All Rights Reserved.
+
+using System;
+using System.Collections.Generic;
+using System.Threading.Tasks;
+using Horde.Build.Collections;
+using Horde.Build.Models;
+using Horde.Build.Utilities;
+using MongoDB.Bson;
+
+namespace Horde.Build.Services
+{
+	using UserId = ObjectId<IUser>;
+
+	/// <summary>
+	/// Notice service
+	/// </summary>
+	public sealed class NoticeService 
+	{
+		/// <summary>
+		/// Notice collection
+		/// </summary>
+		readonly INoticeCollection _notices;
+
+		/// <summary>
+		/// Device service constructor
+		/// </summary>
+		public NoticeService(INoticeCollection notices)
+		{
+			_notices = notices;
+		}
+
+		/// <summary>
+		/// Add a notice
+		/// </summary>
+		/// <param name="message"></param>
+		/// <param name="userId"></param>
+		/// <param name="startTime"></param>
+		/// <param name="finishTime"></param>
+		/// <returns></returns>
+		public Task<INotice?> AddNoticeAsync(string message, UserId? userId, DateTime? startTime, DateTime? finishTime)
+		{
+			return _notices.AddNoticeAsync(message, userId, startTime, finishTime);
+		}
+
+		/// <summary>
+		/// Update an existing notice
+		/// </summary>
+		/// <param name="id"></param>
+		/// <param name="message"></param>
+		/// <param name="startTime"></param>
+		/// <param name="finishTime"></param>
+		/// <returns></returns>
+		public Task<bool> UpdateNoticeAsync(ObjectId id, string? message, DateTime? startTime, DateTime? finishTime)
+		{
+			return _notices.UpdateNoticeAsync(id, message, startTime, finishTime);
+		}
+
+		/// <summary>
+		/// Get a notice by id
+		/// </summary>
+		/// <param name="noticeId"></param>
+		/// <returns></returns>
+		public Task<INotice?> GetNoticeAsync(ObjectId noticeId)
+		{
+			return _notices.GetNoticeAsync(noticeId);
+		}
+
+		/// <summary>
+		/// Get all notices
+		/// </summary>
+		/// <returns></returns>
+		public Task<List<INotice>> GetNoticesAsync()
+		{
+			return _notices.GetNoticesAsync();
+		}
+
+		/// <summary>
+		/// Remove an existing notice
+		/// </summary>
+		/// <param name="id"></param>
+		/// <returns></returns>
+		public Task<bool> RemoveNoticeAsync(ObjectId id)
+		{
+			return _notices.RemoveNoticeAsync(id);
+		}
+	}
+}
\ No newline at end of file
diff --git a/Engine/Source/Programs/Horde/Horde.Build/Server/Notices/NoticesController.Messages.cs b/Engine/Source/Programs/Horde/Horde.Build/Server/Notices/NoticesController.Messages.cs
new file mode 100644
--- /dev/null
+++ b/Engine/Source/Programs/Horde/Horde.Build/Server/Notices/NoticesController.Messages.cs
@@ -0,0 +1,123 @@
+// Copyright Epic Games, Inc. All Rights Reserved.
+
+using System;
+using System.ComponentModel.DataAnnotations;
+using Horde.Build.Models;
+
+namespace Horde.Build.Api
+{
+
+	/// <summary>
+	/// Parameters required to create a notice
+	/// </summary>
+	public class CreateNoticeRequest
+	{
+		/// <summary>
+		/// Start time to display this message
+		/// </summary>
+		public DateTime? StartTime { get; set; }
+
+		/// <summary>
+		/// Finish time to display this message
+		/// </summary>
+		public DateTime? FinishTime { get; set; }
+
+		/// <summary>
+		/// Message to display
+		/// </summary>
+		[Required]
+		public string Message { get; set; }
+
+		/// <summary>
+		/// Constructor
+		/// </summary>
+		public CreateNoticeRequest(string message)
+		{
+			Message = message;
+		}
+	}
+
+	/// <summary>
+	/// Parameters required to update a notice
+	/// </summary>
+	public class UpdateNoticeRequest
+	{
+		/// <summary>
+		/// The id of the notice to update
+		/// </summary>
+		public string Id { get; set; }
+
+		/// <summary>
+		/// Start time to display this message
+		/// </summary>
+		public DateTime? StartTime { get; set; }
+
+		/// <summary>
+		/// Finish time to display this message
+		/// </summary>
+		public DateTime? FinishTime { get; set; }
+
+		/// <summary>
+		/// Message to display
+		/// </summary>
+		public string? Message { get; set; }
+
+		/// <summary>
+		/// Constructor
+		/// </summary>
+		public UpdateNoticeRequest(string id)
+		{
+			Id = id;
+		}
+	}
+
+	/// <summary>
+	/// Notice information
+	/// </summary>
+	public class GetNoticeResponse
+	{
+		/// <summary>
+		/// The id of the notice to update
+		/// </summary>
+		public string? Id { get; set; }
+
+		/// <summary>
+		/// Start time to display this message
+		/// </summary>
+		public DateTime? StartTime { get; set; }
+
+		/// <summary>
+		/// Finish time to display this message
+		/// </summary>
+		public DateTime? FinishTime { get; set; }
+
+		/// <summary>
+		/// Whether this notice is for scheduled downtime
+		/// </summary>
+		public bool ScheduledDowntime { get; set; } = false;
+
+		/// <summary>
+		/// Whether the notice is currently active
+		/// </summary>
+		public bool Active { get; set; }
+
+		/// <summary>
+		/// Message to display
+		/// </summary>
+		public string? Message { get; set; }
+
+		/// <summary>
+		/// User id who created the notice, otherwise null if a system message
+		/// </summary>
+		public GetThinUserInfoResponse? CreatedByUser { get; set; }
+
+		/// <summary>
+		/// Constructor
+		/// </summary>
+		public GetNoticeResponse()
+		{
+
+		}
+
+	}
+}
diff --git a/Engine/Source/Programs/Horde/Horde.Build/Server/Notices/NoticesController.cs b/Engine/Source/Programs/Horde/Horde.Build/Server/Notices/NoticesController.cs
new file mode 100644
--- /dev/null
+++ b/Engine/Source/Programs/Horde/Horde.Build/Server/Notices/NoticesController.cs
@@ -0,0 +1,168 @@
+// Copyright Epic Games, Inc. All Rights Reserved.
+
+using System;
+using System.Collections.Generic;
+using System.Threading.Tasks;
+using Horde.Build.Acls;
+using Horde.Build.Api;
+using Horde.Build.Collections;
+using Horde.Build.Models;
+using Horde.Build.Services;
+using Horde.Build.Utilities;
+using Microsoft.AspNetCore.Authorization;
+using Microsoft.AspNetCore.Mvc;
+using Microsoft.Extensions.Options;
+using MongoDB.Bson;
+using Horde.Build.Server;
+
+namespace Horde.Build.Controllers
+{
+	/// <summary>
+	/// Controller for the /api/v1/notices endpoint
+	/// </summary>
+	[ApiController]
+	[Authorize]
+	[Route("[controller]")]
+	public class NoticesController : ControllerBase
+	{
+		/// <summary>
+		/// The acl service singleton
+		/// </summary>
+		readonly AclService _aclService;
+
+		/// <summary>
+		/// Reference to the user collection
+		/// </summary>
+		private readonly IUserCollection _userCollection;
+
+		/// <summary>
+		/// Notice service
+		/// </summary>
+		private readonly NoticeService _noticeService;
+
+		/// <summary>
+		/// cached globals
+		/// </summary>
+		readonly LazyCachedValue<Task<Globals>> _cachedGlobals;
+
+		/// <summary>
+		/// cached notices
+		/// </summary>
+		readonly LazyCachedValue<Task<List<INotice>>> _cachedNotices;
+
+		/// <summary>
+		/// Server settings
+		/// </summary>
+		readonly IOptionsMonitor<ServerSettings> _settings;
+
+		/// <summary>
+		/// Constructor
+		/// </summary>
+		/// <param name="settings">The server settings</param>
+		/// <param name="mongoService">The mongo service singleton</param>
+		/// <param name="noticeService">The notice service singleton</param>
+		/// <param name="aclService">The acl service singleton</param>
+		/// <param name="userCollection">The user collection singleton</param>
+		public NoticesController(IOptionsMonitor<ServerSettings> settings, MongoService mongoService, NoticeService noticeService, AclService aclService, IUserCollection userCollection)
+		{			
+			_settings = settings;
+			_aclService = aclService;
+			_userCollection = userCollection;
+			_noticeService = noticeService;
+			_cachedGlobals = new LazyCachedValue<Task<Globals>>(() => mongoService.GetGlobalsAsync(), TimeSpan.FromSeconds(30.0));
+			_cachedNotices = new LazyCachedValue<Task<List<INotice>>>(() => noticeService.GetNoticesAsync(), TimeSpan.FromMinutes(1));
+		}
+
+		/// <summary>
+		/// Add a status message
+		/// </summary>
+		/// <param name="request"></param>
+		/// <returns></returns>
+		[HttpPost("/api/v1/notices")]
+		public async Task<ActionResult> AddNoticeAsync(CreateNoticeRequest request)
+		{
+			if (!await _aclService.AuthorizeAsync(AclAction.AdminWrite, User))
+			{
+				return Forbid();
+			}
+
+			INotice? notice = await _noticeService.AddNoticeAsync(request.Message, User.GetUserId(), request.StartTime, request.FinishTime);
+
+			return notice == null ? NotFound() : Ok();
+
+		}
+
+		/// <summary>
+		/// Update a status message
+		/// </summary>
+		/// <param name="request"></param>
+		/// <returns></returns>
+		[HttpPut("/api/v1/notices")]
+		public async Task<ActionResult> UpdateNoticeAsync(UpdateNoticeRequest request)
+		{
+			if (!await _aclService.AuthorizeAsync(AclAction.AdminWrite, User))
+			{
+				return Forbid();
+			}
+
+			return await _noticeService.UpdateNoticeAsync(new ObjectId(request.Id), request.Message, request.StartTime, request.FinishTime) ? Ok() : NotFound();
+		}
+
+		/// <summary>
+		/// Remove a manually added status message
+		/// </summary>
+		/// <returns></returns>
+		[HttpDelete("/api/v1/notices/{id}")]
+		public async Task<ActionResult> DeleteNoticeAsync(string id)
+		{
+			if (!await _aclService.AuthorizeAsync(AclAction.AdminWrite, User))
+			{
+				return Forbid();
+			}
+
+			await _noticeService.RemoveNoticeAsync(new ObjectId(id));
+			
+			return Ok();
+		}
+
+		/// <summary>
+		/// Gets the current status messages
+		/// </summary>
+		/// <returns>The status messages</returns>
+		[HttpGet("/api/v1/notices")]
+		public async Task<List<GetNoticeResponse>> GetNoticesAsync()
+		{
+			List<GetNoticeResponse> messages = new List<GetNoticeResponse>();
+
+			Globals globals = await _cachedGlobals.GetCached();
+
+			DateTimeOffset now = TimeZoneInfo.ConvertTime(DateTimeOffset.Now, _settings.CurrentValue.TimeZoneInfo);
+						
+			foreach (ScheduledDowntime schedule in globals.ScheduledDowntime)
+			{
+				DateTimeOffset start = TimeZoneInfo.ConvertTimeBySystemTimeZoneId(schedule.GetNext(now).StartTime, "UTC");
+				DateTimeOffset finish = TimeZoneInfo.ConvertTimeBySystemTimeZoneId(schedule.GetNext(now).FinishTime, "UTC");
+				messages.Add(new GetNoticeResponse() { ScheduledDowntime = true, StartTime = new DateTime(start.Ticks, DateTimeKind.Utc), FinishTime = new DateTime(finish.Ticks, DateTimeKind.Utc), Active = schedule.IsActive(now) });
+			}
+
+			// Add user notices
+			List<INotice> notices = await _cachedNotices.GetCached();
+
+			for (int i = 0; i < notices.Count; i++)
+			{
+				INotice notice = notices[i];
+				GetThinUserInfoResponse? userInfo = null;
+
+				if (notice.UserId != null)
+				{
+					userInfo = new GetThinUserInfoResponse(await _userCollection.GetCachedUserAsync(notice.UserId));
+				}
+
+				messages.Add(new GetNoticeResponse() { Id = notice.Id.ToString(), Active = true, Message = notice.Message, CreatedByUser = userInfo});
+			}
+
+			return messages;
+			
+		}
+	}
+}
