Horde: Finish re-org of Horde solution started a while ago.

diff --git a/Engine/Source/Programs/Horde/Horde.Build/Server/RpcService.cs b/Engine/Source/Programs/Horde/Horde.Build/Server/RpcService.cs
new file mode 100644
--- /dev/null
+++ b/Engine/Source/Programs/Horde/Horde.Build/Server/RpcService.cs
@@ -0,0 +1,1009 @@
+// Copyright Epic Games, Inc. All Rights Reserved.
+
+using System;
+using System.Collections.Generic;
+using System.Globalization;
+using System.Linq;
+using System.Net;
+using System.Security.Claims;
+using System.Text;
+using System.Threading;
+using System.Threading.Tasks;
+using EpicGames.Core;
+using Google.Protobuf.WellKnownTypes;
+using Grpc.Core;
+using Horde.Build.Acls;
+using Horde.Build.Collections;
+using Horde.Build.Jobs;
+using Horde.Build.Models;
+using Horde.Build.Server;
+using Horde.Build.Tasks.Impl;
+using Horde.Build.Utilities;
+using HordeCommon;
+using HordeCommon.Rpc;
+using Microsoft.AspNetCore.Authorization;
+using Microsoft.Extensions.Logging;
+using MongoDB.Bson;
+using MongoDB.Bson.Serialization;
+
+namespace Horde.Build.Services
+{
+	using AgentSoftwareChannelName = StringId<AgentSoftwareChannels>;
+	using IStream = Horde.Build.Models.IStream;
+	using JobId = ObjectId<IJob>;
+	using LogId = ObjectId<ILogFile>;
+	using RpcAgentCapabilities = HordeCommon.Rpc.Messages.AgentCapabilities;
+	using RpcDeviceCapabilities = HordeCommon.Rpc.Messages.DeviceCapabilities;
+	using SessionId = ObjectId<ISession>;
+	using StreamId = StringId<IStream>;
+
+	/// <summary>
+	/// Implements the Horde gRPC service for bots updating their status and dequeing work
+	/// </summary>
+	[Authorize]
+	public class RpcService : HordeRpc.HordeRpcBase
+	{
+		/// <summary>
+		/// Timeout before closing a long-polling request (client will retry again) 
+		/// </summary>
+		internal TimeSpan _longPollTimeout = TimeSpan.FromMinutes(9);
+		readonly AclService _aclService;
+		readonly AgentService _agentService;
+		readonly StreamService _streamService;
+		readonly JobService _jobService;
+		readonly AgentSoftwareService _agentSoftwareService;
+		readonly IArtifactCollection _artifactCollection;
+		readonly ILogFileService _logFileService;
+		readonly PoolService _poolService;
+		readonly LifetimeService _lifetimeService;
+		readonly IGraphCollection _graphs;
+		readonly ITestDataCollection _testData;
+		readonly IJobStepRefCollection _jobStepRefCollection;
+		readonly ITemplateCollection _templateCollection;
+		readonly ConformTaskSource _conformTaskSource;
+		readonly IClock _clock;
+		readonly ILogger<RpcService> _logger;
+
+		/// <summary>
+		/// Constructor
+		/// </summary>
+		public RpcService(AclService aclService, AgentService agentService, StreamService streamService, JobService jobService, AgentSoftwareService agentSoftwareService, IArtifactCollection artifactCollection, ILogFileService logFileService, PoolService poolService, LifetimeService lifetimeService, IGraphCollection graphs, ITestDataCollection testData, IJobStepRefCollection jobStepRefCollection, ITemplateCollection templateCollection, ConformTaskSource conformTaskSource, IClock clock, ILogger<RpcService> logger)
+		{
+			_aclService = aclService;
+			_agentService = agentService;
+			_streamService = streamService;
+			_jobService = jobService;
+			_agentSoftwareService = agentSoftwareService;
+			_artifactCollection = artifactCollection;
+			_logFileService = logFileService;
+			_poolService = poolService;
+			_lifetimeService = lifetimeService;
+			_graphs = graphs;
+			_testData = testData;
+			_jobStepRefCollection = jobStepRefCollection;
+			_templateCollection = templateCollection;
+			_conformTaskSource = conformTaskSource;
+			_clock = clock;
+			_logger = logger;
+		}
+
+		/// <summary>
+		/// Waits until the server is terminating
+		/// </summary>
+		/// <param name="reader">Request reader</param>
+		/// <param name="writer">Response writer</param>
+		/// <param name="context">Context for the call</param>
+		/// <returns>Response object</returns>
+		public override async Task QueryServerState(IAsyncStreamReader<QueryServerStateRequest> reader, IServerStreamWriter<QueryServerStateResponse> writer, ServerCallContext context)
+		{
+			if (await reader.MoveNext())
+			{
+				QueryServerStateRequest request = reader.Current;
+				_logger.LogInformation("Start server query for client {Name}", request.Name);
+
+				// Return the current response
+				QueryServerStateResponse response = new QueryServerStateResponse();
+				response.Name = Dns.GetHostName();
+				await writer.WriteAsync(response);
+
+				// Move to the next request from the client. This should always be the end of the stream, but will not occur until the client stops requesting responses.
+				Task<bool> moveNextTask = reader.MoveNext();
+
+				// Wait for the client to close the stream or a shutdown to start
+				Task longPollDelay = Task.Delay(_longPollTimeout);
+				Task waitTask = await Task.WhenAny(moveNextTask, _lifetimeService.StoppingTask, longPollDelay);
+
+				if (waitTask == moveNextTask)
+				{
+					throw new Exception("Unexpected request to QueryServerState posted from client.");
+				}
+				else if (waitTask == _lifetimeService.StoppingTask)
+				{
+					_logger.LogInformation("Notifying client {Name} of server shutdown", request.Name);
+					await writer.WriteAsync(response);
+				}
+				else if (waitTask == longPollDelay)
+				{
+					// Send same response as server shutdown. In the agent perspective, they will be identical.
+					await writer.WriteAsync(response);
+				}
+			}
+		}
+
+		/// <summary>
+		/// Waits until the server is terminating
+		/// </summary>
+		/// <param name="reader">Request reader</param>
+		/// <param name="writer">Response writer</param>
+		/// <param name="context">Context for the call</param>
+		/// <returns>Response object</returns>
+		public override async Task QueryServerStateV2(IAsyncStreamReader<QueryServerStateRequest> reader, IServerStreamWriter<QueryServerStateResponse> writer, ServerCallContext context)
+		{
+			if (await reader.MoveNext())
+			{
+				QueryServerStateRequest request = reader.Current;
+				_logger.LogDebug("Start server query for client {Name}", request.Name);
+
+				try
+				{
+					// Return the current response
+					QueryServerStateResponse response = new QueryServerStateResponse();
+					response.Name = Dns.GetHostName();
+					response.Stopping = _lifetimeService.IsStopping;
+					await writer.WriteAsync(response);
+
+					// Move to the next request from the client. This should always be the end of the stream, but will not occur until the client stops requesting responses.
+					Task<bool> moveNextTask = reader.MoveNext();
+
+					// Wait for the client to close the stream or a shutdown to start
+					if (await Task.WhenAny(moveNextTask, _lifetimeService.StoppingTask) == _lifetimeService.StoppingTask)
+					{
+						response.Stopping = true;
+						await writer.WriteAsync(response);
+					}
+
+					// Wait until the client has finished sending
+					while (await moveNextTask)
+					{
+						moveNextTask = reader.MoveNext();
+					}
+				}
+				catch (Exception ex)
+				{
+					_logger.LogError(ex, "Exception in QueryServerState for {Name}", request.Name);
+					throw;
+				}
+			}
+		}
+
+		/// <summary>
+		/// Updates the workspaces synced for an agent
+		/// </summary>
+		/// <param name="request">The request parameters</param>
+		/// <param name="context">Context for the call</param>
+		/// <returns>Response object</returns>
+		public override async Task<UpdateAgentWorkspacesResponse> UpdateAgentWorkspaces(UpdateAgentWorkspacesRequest request, ServerCallContext context)
+		{
+			for (; ; )
+			{
+				// Get the current agent state
+				IAgent? agent = await _agentService.GetAgentAsync(new AgentId(request.AgentId));
+				if (agent == null)
+				{
+					throw new StructuredRpcException(StatusCode.OutOfRange, "Agent {AgentId} does not exist", request.AgentId);
+				}
+
+				// Get the new workspaces
+				List<AgentWorkspace> newWorkspaces = request.Workspaces.Select(x => new AgentWorkspace(x)).ToList();
+
+				// Get the set of workspaces that are currently required
+				HashSet<AgentWorkspace> conformWorkspaces = await _poolService.GetWorkspacesAsync(agent, DateTime.UtcNow);
+				bool bPendingConform = !conformWorkspaces.SetEquals(newWorkspaces) || (agent.RequestFullConform && !request.RemoveUntrackedFiles);
+
+				// Update the workspaces
+				if (await _agentService.TryUpdateWorkspacesAsync(agent, newWorkspaces, bPendingConform))
+				{
+					UpdateAgentWorkspacesResponse response = new UpdateAgentWorkspacesResponse();
+					if (bPendingConform)
+					{
+						response.Retry = await _conformTaskSource.GetWorkspacesAsync(agent, response.PendingWorkspaces);
+						response.RemoveUntrackedFiles = request.RemoveUntrackedFiles || agent.RequestFullConform;
+					}
+					return response;
+				}
+			}
+		}
+
+		static void CopyPropertyToResource(string name, List<string> properties, Dictionary<string, int> resources)
+		{
+			foreach (string property in properties)
+			{
+				if (property.Length > name.Length && property.StartsWith(name, StringComparison.OrdinalIgnoreCase) && property[name.Length] == '=')
+				{
+					int value;
+					if (Int32.TryParse(property.AsSpan(name.Length + 1), out value))
+					{
+						resources[name] = value;
+					}
+				}
+			}
+		}
+
+		static void GetCapabilities(RpcAgentCapabilities? capabilities, out List<string> properties, out Dictionary<string, int> resources)
+		{
+			properties = new List<string>();
+			resources = new Dictionary<string, int>();
+
+			if (capabilities != null && capabilities.Devices.Count > 0)
+			{
+				RpcDeviceCapabilities device = capabilities.Devices[0];
+				if (device.Properties != null)
+				{
+					properties = new List<string>(device.Properties);
+					CopyPropertyToResource(KnownPropertyNames.LogicalCores, properties, resources);
+					CopyPropertyToResource(KnownPropertyNames.Ram, properties, resources);
+				}
+			}
+		}
+
+		/// <summary>
+		/// Creates a new session
+		/// </summary>
+		/// <param name="request">Request to create a new agent</param>
+		/// <param name="context">Context for the RPC call</param>
+		/// <returns>Information about the new agent</returns>
+		public override async Task<CreateSessionResponse> CreateSession(CreateSessionRequest request, ServerCallContext context)
+		{
+			if (request.Capabilities == null)
+			{
+				throw new StructuredRpcException(StatusCode.InvalidArgument, "Capabilities may not be null");
+			}
+
+			AgentId agentId = new AgentId(request.Name);
+			using IDisposable scope = _logger.BeginScope("CreateSession({AgentId})", agentId.ToString());
+
+			// Find the agent
+			IAgent? agent = await _agentService.GetAgentAsync(agentId);
+			if (agent == null)
+			{
+				if (!await _aclService.AuthorizeAsync(AclAction.CreateAgent, context.GetHttpContext().User))
+				{
+					throw new StructuredRpcException(StatusCode.PermissionDenied, "User is not authenticated to create new agents");
+				}
+
+				const bool BEnabled = true;
+				agent = await _agentService.CreateAgentAsync(request.Name, BEnabled, null, null);
+			}
+
+			// Make sure we're allowed to create sessions on this agent
+			if (!await _agentService.AuthorizeAsync(agent, AclAction.CreateSession, context.GetHttpContext().User, null))
+			{
+				throw new StructuredRpcException(StatusCode.PermissionDenied, "User is not authenticated to create session for {AgentId}", request.Name);
+			}
+
+			// Get the known properties for this agent
+			GetCapabilities(request.Capabilities, out List<string> properties, out Dictionary<string, int> resources);
+
+			// Create a new session
+			agent = await _agentService.CreateSessionAsync(agent, request.Status, properties, resources, request.Version);
+			if (agent == null)
+			{
+				throw new StructuredRpcException(StatusCode.NotFound, "Agent {AgentId} not found", request.Name);
+			}
+
+			// Create the response
+			CreateSessionResponse response = new CreateSessionResponse();
+			response.AgentId = agent.Id.ToString();
+			response.SessionId = agent.SessionId.ToString();
+			response.ExpiryTime = Timestamp.FromDateTime(agent.SessionExpiresAt!.Value);
+			response.Token = _agentService.IssueSessionToken(agent.Id, agent.SessionId!.Value);
+			return response;
+		}
+
+		/// <summary>
+		/// Updates an agent session
+		/// </summary>
+		/// <param name="reader">Request to create a new agent</param>
+		/// <param name="writer">Writer for response objects</param>
+		/// <param name="context">Context for the RPC call</param>
+		/// <returns>Information about the new agent</returns>
+		public override async Task UpdateSession(IAsyncStreamReader<UpdateSessionRequest> reader, IServerStreamWriter<UpdateSessionResponse> writer, ServerCallContext context)
+		{
+			// Read the request object
+			Task<bool> nextRequestTask = reader.MoveNext();
+			if (await nextRequestTask)
+			{
+				UpdateSessionRequest request = reader.Current;
+				using IDisposable scope = _logger.BeginScope("UpdateSession for agent {AgentId}, session {SessionId}", request.AgentId, request.SessionId);
+
+				_logger.LogDebug("Updating session for {AgentId}", request.AgentId);
+				foreach (HordeCommon.Rpc.Messages.Lease lease in request.Leases)
+				{
+					_logger.LogDebug("Session {SessionId}, Lease {LeaseId} - State: {LeaseState}, Outcome: {LeaseOutcome}", request.SessionId, lease.Id, lease.State, lease.Outcome);
+				}
+
+				// Get a task for moving to the next item. This will only complete once the call has closed.
+				using CancellationTokenSource cancellationSource = CancellationTokenSource.CreateLinkedTokenSource(context.CancellationToken);
+				nextRequestTask = reader.MoveNext();
+				nextRequestTask = nextRequestTask.ContinueWith(task => 
+				{ 
+					cancellationSource.Cancel(); 
+					return task.Result; 
+				}, TaskScheduler.Current);
+
+				// Get the current agent state
+				IAgent? agent = await _agentService.GetAgentAsync(new AgentId(request.AgentId));
+				if(agent != null)
+				{
+					// Check we're authorized to update it
+					if (!_agentService.AuthorizeSession(agent, context.GetHttpContext().User))
+					{
+						throw new StructuredRpcException(StatusCode.PermissionDenied, "Not authenticated for {AgentId}", request.AgentId);
+					}
+
+					// Get the new capabilities of this agent
+					List<string>? properties = null;
+					Dictionary<string, int>? resources = null;
+					if (request.Capabilities != null)
+					{
+						GetCapabilities(request.Capabilities, out properties, out resources);
+					}
+
+					// Update the session
+					try
+					{
+						agent = await _agentService.UpdateSessionWithWaitAsync(agent, SessionId.Parse(request.SessionId), request.Status, properties, resources, request.Leases, cancellationSource.Token);
+					}
+					catch (Exception ex)
+					{
+						_logger.LogError(ex, "Swallowed exception while updating session for {AgentId}.", request.AgentId);
+					}
+				}
+
+				// Handle the invalid agent case
+				if (agent == null)
+				{
+					throw new StructuredRpcException(StatusCode.NotFound, "Invalid agent name '{AgentId}'", request.AgentId);
+				}
+
+				// Create the new session info
+				if (!context.CancellationToken.IsCancellationRequested)
+				{
+					UpdateSessionResponse response = new UpdateSessionResponse();
+					response.Leases.Add(agent.Leases.Select(x => x.ToRpcMessage()));
+					response.ExpiryTime = (agent.SessionExpiresAt == null) ? new Timestamp() : Timestamp.FromDateTime(agent.SessionExpiresAt.Value);
+					await writer.WriteAsync(response);
+				}
+
+				// Wait for the client to close the stream
+				while (await nextRequestTask)
+				{
+					nextRequestTask = reader.MoveNext();
+				}
+			}
+		}
+
+		/// <summary>
+		/// Gets information about a stream
+		/// </summary>
+		/// <param name="request">Request arguments</param>
+		/// <param name="context">Context for the RPC call</param>
+		/// <returns>Information about the new agent</returns>
+		public override async Task<GetStreamResponse> GetStream(GetStreamRequest request, ServerCallContext context)
+		{
+			StreamId streamIdValue = new StreamId(request.StreamId);
+
+			IStream? stream = await _streamService.GetStreamAsync(streamIdValue);
+			if (stream == null)
+			{
+				throw new StructuredRpcException(StatusCode.NotFound, "Stream {StreamId} does not exist", request.StreamId);
+			}
+			if (!await _streamService.AuthorizeAsync(stream, AclAction.ViewStream, context.GetHttpContext().User, null))
+			{
+				throw new StructuredRpcException(StatusCode.PermissionDenied, "Not authenticated to access stream {StreamId}", request.StreamId);
+			}
+
+			return stream.ToRpcResponse();
+		}
+
+		/// <summary>
+		/// Gets information about a job
+		/// </summary>
+		/// <param name="request">Request arguments</param>
+		/// <param name="context">Context for the RPC call</param>
+		/// <returns>Information about the new agent</returns>
+		public override async Task<GetJobResponse> GetJob(GetJobRequest request, ServerCallContext context)
+		{
+			JobId jobIdValue = new JobId(request.JobId.ToObjectId());
+
+			IJob? job = await _jobService.GetJobAsync(jobIdValue);
+			if (job == null)
+			{
+				throw new StructuredRpcException(StatusCode.NotFound, "Job {JobId} does not exist", request.JobId);
+			}
+			if (!JobService.AuthorizeSession(job, context.GetHttpContext().User))
+			{
+				throw new StructuredRpcException(StatusCode.PermissionDenied, "Not authenticated to access job {JobId}", request.JobId);
+			}
+
+			return job.ToRpcResponse();
+		}
+
+		/// <summary>
+		/// Updates properties on a job
+		/// </summary>
+		/// <param name="request">Request arguments</param>
+		/// <param name="context">Context for the RPC call</param>
+		/// <returns>Information about the new agent</returns>
+		public override async Task<Empty> UpdateJob(UpdateJobRequest request, ServerCallContext context)
+		{
+			JobId jobIdValue = new JobId(request.JobId);
+
+			IJob? job = await _jobService.GetJobAsync(jobIdValue);
+			if (job == null)
+			{
+				throw new StructuredRpcException(StatusCode.NotFound, "Job {JobId} does not exist", request.JobId);
+			}
+			if (!JobService.AuthorizeSession(job, context.GetHttpContext().User))
+			{
+				throw new StructuredRpcException(StatusCode.PermissionDenied, "Not authenticated to modify job {JobId}", request.JobId);
+			}
+
+			await _jobService.UpdateJobAsync(job, name: request.Name);
+			return new Empty();
+		}
+
+		/// <summary>
+		/// Starts executing a batch
+		/// </summary>
+		/// <param name="request">Request arguments</param>
+		/// <param name="context">Context for the RPC call</param>
+		/// <returns>Information about the new agent</returns>
+		public override async Task<BeginBatchResponse> BeginBatch(BeginBatchRequest request, ServerCallContext context)
+		{
+			SubResourceId batchId = request.BatchId.ToSubResourceId();
+
+			IJob? job = await _jobService.GetJobAsync(new JobId(request.JobId));
+			if (job == null)
+			{
+				throw new StructuredRpcException(StatusCode.NotFound, "Job {JobId} not found", request.JobId);
+			}
+
+			IJobStepBatch batch = AuthorizeBatch(job, request.BatchId.ToSubResourceId(), context);
+			job = await _jobService.UpdateBatchAsync(job, batchId, null, Api.JobStepBatchState.Starting);
+
+			if (job == null)
+			{
+				throw new StructuredRpcException(StatusCode.NotFound, "Batch {JobId}:{BatchId} not found for updating", request.JobId, request.BatchId);
+			}
+
+			IGraph graph = await _jobService.GetGraphAsync(job);
+
+			BeginBatchResponse response = new BeginBatchResponse();
+			response.LogId = batch.LogId.ToString();
+			response.AgentType = graph.Groups[batch.GroupIdx].AgentType;
+			return response;
+		}
+
+		/// <summary>
+		/// Finishes executing a batch
+		/// </summary>
+		/// <param name="request">Request arguments</param>
+		/// <param name="context">Context for the RPC call</param>
+		/// <returns>Information about the new agent</returns>
+		public override async Task<Empty> FinishBatch(FinishBatchRequest request, ServerCallContext context)
+		{
+			IJob? job = await _jobService.GetJobAsync(new JobId(request.JobId));
+			if (job == null)
+			{
+				throw new StructuredRpcException(StatusCode.NotFound, "Job {JobId} not found", request.JobId);
+			}
+
+			IJobStepBatch batch = AuthorizeBatch(job, request.BatchId.ToSubResourceId(), context);
+			await _jobService.UpdateBatchAsync(job, batch.Id, null, Api.JobStepBatchState.Complete);
+			return new Empty();
+		}
+
+		/// <summary>
+		/// Starts executing a step
+		/// </summary>
+		/// <param name="request">Request arguments</param>
+		/// <param name="context">Context for the RPC call</param>
+		/// <returns>Information about the new agent</returns>
+		public override async Task<BeginStepResponse> BeginStep(BeginStepRequest request, ServerCallContext context)
+		{
+			Boxed<ILogFile?> log = new Boxed<ILogFile?>(null);
+			for (; ; )
+			{
+				BeginStepResponse? response = await TryBeginStep(request, log, context);
+				if (response != null)
+				{
+					return response;
+				}
+			}
+		}
+
+		async Task<BeginStepResponse?> TryBeginStep(BeginStepRequest request, Boxed<ILogFile?> log, ServerCallContext context)
+		{
+			// Check the job exists and we can access it
+			IJob? job = await _jobService.GetJobAsync(new JobId(request.JobId));
+			if (job == null)
+			{
+				throw new StructuredRpcException(StatusCode.NotFound, "Job {JobId} not found", request.JobId);
+			}
+
+			// Find the batch being executed
+			IJobStepBatch batch = AuthorizeBatch(job, request.BatchId.ToSubResourceId(), context);
+			if (batch.State != Api.JobStepBatchState.Starting && batch.State != Api.JobStepBatchState.Running)
+			{
+				return new BeginStepResponse { State = BeginStepResponse.Types.Result.Complete };
+			}
+
+			// Figure out which step to execute next
+			IJobStep? step;
+			for (int stepIdx = 0; ; stepIdx++)
+			{
+				// If there aren't any more steps, send a complete message
+				if (stepIdx == batch.Steps.Count)
+				{
+					_logger.LogDebug("Job {JobId} batch {BatchId} is complete", job.Id, batch.Id);
+					if (await _jobService.TryUpdateBatchAsync(job, batch.Id, newState: Api.JobStepBatchState.Stopping) == null)
+					{
+						return null;
+					}
+					return new BeginStepResponse { State = BeginStepResponse.Types.Result.Complete };
+				}
+
+				// Check if this step is ready to be executed
+				step = batch.Steps[stepIdx];
+				if (step.State == JobStepState.Ready)
+				{
+					break;
+				}
+				if (step.State == JobStepState.Waiting)
+				{
+					_logger.LogDebug("Waiting for job {JobId}, batch {BatchId}, step {StepId}", job.Id, batch.Id, step.Id);
+					return new BeginStepResponse { State = BeginStepResponse.Types.Result.Waiting };
+				}
+			}
+
+			// Create a log file if necessary
+			if (log.Value == null)
+			{
+				log.Value = await _logFileService.CreateLogFileAsync(job.Id, batch.SessionId, Api.LogType.Json);
+			}
+
+			// Get the node for this step
+			IGraph graph = await _jobService.GetGraphAsync(job);
+			INode node = graph.Groups[batch.GroupIdx].Nodes[step.NodeIdx];
+
+			// Figure out all the credentials for it (and check we can access them)
+			Dictionary<string, string> credentials = new Dictionary<string, string>();
+			//				if (Node.Credentials != null)
+			//				{
+			//					ClaimsPrincipal Principal = new ClaimsPrincipal(new ClaimsIdentity(Job.Claims.Select(x => new Claim(x.Type, x.Value))));
+			//					if (!await GetCredentialsForStep(Principal, Node, Credentials, Message => FailStep(Job, Batch.Id, Step.Id, Log, Message)))
+			//					{
+			//						Log = null;
+			//						continue;
+			//					}
+			//				}
+
+			// Update the step state
+			IJob? newJob = await _jobService.TryUpdateStepAsync(job, batch.Id, step.Id, JobStepState.Running, JobStepOutcome.Unspecified, newLogId: log.Value.Id);
+			if (newJob != null)
+			{
+				BeginStepResponse response = new BeginStepResponse();
+				response.State = BeginStepResponse.Types.Result.Ready;
+				response.LogId = log.Value.Id.ToString();
+				response.StepId = step.Id.ToString();
+				response.Name = node.Name;
+				response.Credentials.Add(credentials);
+
+				string templateName = "<unknown>";
+				if (job.TemplateHash != null)
+				{
+					ITemplate? template = await _templateCollection.GetAsync(job.TemplateHash);
+					templateName = template != null ? template.Name : templateName;
+				}
+				
+				response.EnvVars.Add("UE_HORDE_TEMPLATEID", job.TemplateId.ToString());
+				response.EnvVars.Add("UE_HORDE_TEMPLATENAME", templateName);
+
+				IJobStepRef? lastStep = await _jobStepRefCollection.GetPrevStepForNodeAsync(job.StreamId, job.TemplateId, node.Name, job.Change);
+				if (lastStep != null)
+				{
+					response.EnvVars.Add("UE_HORDE_LAST_CL", lastStep.Change.ToString(CultureInfo.InvariantCulture));
+
+					if (lastStep.Outcome == JobStepOutcome.Success)
+					{
+						response.EnvVars.Add("UE_HORDE_LAST_SUCCESS_CL", lastStep.Change.ToString(CultureInfo.InvariantCulture));
+					}
+					else if (lastStep.LastSuccess != null)
+					{
+						response.EnvVars.Add("UE_HORDE_LAST_SUCCESS_CL", lastStep.LastSuccess.Value.ToString(CultureInfo.InvariantCulture));
+					}
+
+					if (lastStep.Outcome == JobStepOutcome.Success || lastStep.Outcome == JobStepOutcome.Warnings)
+					{
+						response.EnvVars.Add("UE_HORDE_LAST_WARNING_CL", lastStep.Change.ToString(CultureInfo.InvariantCulture));
+					}
+					else if (lastStep.LastWarning != null)
+					{
+						response.EnvVars.Add("UE_HORDE_LAST_WARNING_CL", lastStep.LastWarning.Value.ToString(CultureInfo.InvariantCulture));
+					}
+				}
+
+				if (node.Properties != null)
+				{
+					response.Properties.Add(node.Properties);
+				}
+				response.Warnings = node.Warnings;
+				return response;
+			}
+
+			return null;
+		}
+
+		async Task<IJob> GetJobAsync(JobId jobId)
+		{
+			IJob? job = await _jobService.GetJobAsync(jobId);
+			if (job == null)
+			{
+				throw new StructuredRpcException(StatusCode.NotFound, "Job {JobId} not found", jobId);
+			}
+			return job;
+		}
+
+		static IJobStepBatch AuthorizeBatch(IJob job, SubResourceId batchId, ServerCallContext context)
+		{
+			IJobStepBatch? batch;
+			if (!job.TryGetBatch(batchId, out batch))
+			{
+				throw new StructuredRpcException(StatusCode.NotFound, "Unable to find batch {JobId}:{BatchId}", job.Id, batchId);
+			}
+			if (batch.SessionId == null)
+			{
+				throw new StructuredRpcException(StatusCode.PermissionDenied, "Batch {JobId}:{BatchId} has no session id", job.Id, batchId);
+			}
+
+			ClaimsPrincipal principal = context.GetHttpContext().User;
+			if (!principal.HasSessionClaim(batch.SessionId.Value))
+			{
+				throw new StructuredRpcException(StatusCode.PermissionDenied, "Session id {SessionId} not valid for batch {JobId}:{BatchId}. Expected {ExpectedSessionId}.", principal.GetSessionClaim() ?? SessionId.Empty, job.Id, batchId, batch.SessionId.Value);
+			}
+
+			return batch;
+		}
+
+		/// <summary>
+		/// Updates the state of a jobstep
+		/// </summary>
+		/// <param name="request">Request arguments</param>
+		/// <param name="context">Context for the RPC call</param>
+		/// <returns>Information about the new agent</returns>
+		public override async Task<Empty> UpdateStep(UpdateStepRequest request, ServerCallContext context)
+		{
+			IJob? job = await _jobService.GetJobAsync(new JobId(request.JobId));
+			if (job == null)
+			{
+				throw new StructuredRpcException(StatusCode.NotFound, "Job {JobId} not found", request.JobId);
+			}
+
+			IJobStepBatch batch = AuthorizeBatch(job, request.BatchId.ToSubResourceId(), context);
+
+			SubResourceId stepId = request.StepId.ToSubResourceId();
+			if (!batch.TryGetStep(stepId, out IJobStep? step))
+			{
+				throw new StructuredRpcException(StatusCode.NotFound, "Unable to find step {JobId}:{BatchId}:{StepId}", job.Id, batch.Id, stepId);
+			}
+
+			JobStepError? error = null;
+			if (step.Error == JobStepError.None && step.HasTimedOut(_clock.UtcNow))
+			{
+				error = JobStepError.TimedOut;
+			}
+
+			await _jobService.UpdateStepAsync(job, batch.Id, request.StepId.ToSubResourceId(), request.State, request.Outcome, error, null, null, null, null, null);
+			return new Empty();
+		}
+
+		/// <summary>
+		/// Get the state of a jobstep
+		/// </summary>
+		/// <param name="request">Request arguments</param>
+		/// <param name="context">Context for the RPC call</param>
+		/// <returns>Information about the step</returns>
+		public override async Task<GetStepResponse> GetStep(GetStepRequest request, ServerCallContext context)
+		{
+			IJob job = await GetJobAsync(new JobId(request.JobId));
+			IJobStepBatch batch = AuthorizeBatch(job, request.BatchId.ToSubResourceId(), context);
+
+			SubResourceId stepId = request.StepId.ToSubResourceId();
+			if (!batch.TryGetStep(stepId, out IJobStep? step))
+			{
+				throw new StructuredRpcException(StatusCode.NotFound, "Unable to find step {JobId}:{BatchId}:{StepId}", job.Id, batch.Id, stepId);
+			}
+
+			return new GetStepResponse { Outcome = step.Outcome, State = step.State, AbortRequested = step.AbortRequested || step.HasTimedOut(_clock.UtcNow) };
+		}
+
+		/// <summary>
+		/// Updates the state of a jobstep
+		/// </summary>
+		/// <param name="request">Request arguments</param>
+		/// <param name="context">Context for the RPC call</param>
+		/// <returns>Information about the new agent</returns>
+		public override async Task<UpdateGraphResponse> UpdateGraph(UpdateGraphRequest request, ServerCallContext context)
+		{
+			List<NewGroup> newGroups = new List<NewGroup>();
+			foreach (CreateGroupRequest group in request.Groups)
+			{
+				List<NewNode> newNodes = new List<NewNode>();
+				foreach (CreateNodeRequest node in group.Nodes)
+				{
+					NewNode newNode = new NewNode(node.Name, node.InputDependencies.ToList(), node.OrderDependencies.ToList(), node.Priority, node.AllowRetry, node.RunEarly, node.Warnings, new Dictionary<string, string>(node.Credentials), new Dictionary<string, string>(node.Properties), new NodeAnnotations(node.Annotations));
+					newNodes.Add(newNode);
+				}
+				newGroups.Add(new NewGroup(group.AgentType, newNodes));
+			}
+
+			List<NewAggregate> newAggregates = new List<NewAggregate>();
+			foreach (CreateAggregateRequest aggregate in request.Aggregates)
+			{
+				NewAggregate newAggregate = new NewAggregate(aggregate.Name, aggregate.Nodes.ToList());
+				newAggregates.Add(newAggregate);
+			}
+
+			List<NewLabel> newLabels = new List<NewLabel>();
+			foreach (CreateLabelRequest label in request.Labels)
+			{
+				NewLabel newLabel = new NewLabel();
+				newLabel.DashboardName = String.IsNullOrEmpty(label.DashboardName) ? null : label.DashboardName;
+				newLabel.DashboardCategory = String.IsNullOrEmpty(label.DashboardCategory) ? null : label.DashboardCategory;
+				newLabel.UgsName = String.IsNullOrEmpty(label.UgsName) ? null : label.UgsName;
+				newLabel.UgsProject = String.IsNullOrEmpty(label.UgsProject) ? null : label.UgsProject;
+				newLabel.Change = label.Change;
+				newLabel.RequiredNodes = label.RequiredNodes.ToList();
+				newLabel.IncludedNodes = label.IncludedNodes.ToList();
+				newLabels.Add(newLabel);
+			}
+
+			JobId jobIdValue = new JobId(request.JobId);
+			for (; ; )
+			{
+				IJob? job = await _jobService.GetJobAsync(jobIdValue);
+				if (job == null)
+				{
+					throw new StructuredRpcException(StatusCode.NotFound, "Resource not found");
+				}
+				if (!JobService.AuthorizeSession(job, context.GetHttpContext().User))
+				{
+					throw new StructuredRpcException(StatusCode.PermissionDenied, "Access denied");
+				}
+
+				IGraph graph = await _jobService.GetGraphAsync(job);
+				graph = await _graphs.AppendAsync(graph, newGroups, newAggregates, newLabels);
+
+				IJob? newJob = await _jobService.TryUpdateGraphAsync(job, graph);
+				if (newJob != null)
+				{
+					return new UpdateGraphResponse();
+				}
+			}
+		}
+
+		/// <summary>
+		/// Creates a set of events
+		/// </summary>
+		/// <param name="request">Request arguments</param>
+		/// <param name="context">Context for the RPC call</param>
+		/// <returns>Information about the new agent</returns>
+		public override async Task<Empty> CreateEvents(CreateEventsRequest request, ServerCallContext context)
+		{
+			if (!await _aclService.AuthorizeAsync(AclAction.CreateEvent, context.GetHttpContext().User))
+			{
+				throw new StructuredRpcException(StatusCode.NotFound, "Access denied");
+			}
+
+			List<NewLogEventData> newEvents = new List<NewLogEventData>();
+			foreach (CreateEventRequest createEvent in request.Events)
+			{
+				NewLogEventData newEvent = new NewLogEventData();
+				newEvent.LogId = new LogId(createEvent.LogId);
+				newEvent.Severity = createEvent.Severity;
+				newEvent.LineIndex = createEvent.LineIndex;
+				newEvent.LineCount = createEvent.LineCount;
+				newEvents.Add(newEvent);
+			}
+			await _logFileService.CreateEventsAsync(newEvents);
+			return new Empty();
+		}
+
+		/// <summary>
+		/// Writes output to a log file
+		/// </summary>
+		/// <param name="request">Request arguments</param>
+		/// <param name="context">Context for the RPC call</param>
+		/// <returns>Information about the new agent</returns>
+		public override async Task<Empty> WriteOutput(WriteOutputRequest request, ServerCallContext context)
+		{
+			ILogFile? logFile = await _logFileService.GetCachedLogFileAsync(new LogId(request.LogId));
+			if (logFile == null)
+			{
+				throw new StructuredRpcException(StatusCode.NotFound, "Resource not found");
+			}
+			if (!Horde.Build.Services.LogFileService.AuthorizeForSession(logFile, context.GetHttpContext().User))
+			{
+				throw new StructuredRpcException(StatusCode.PermissionDenied, "Access denied");
+			}
+
+			await _logFileService.WriteLogDataAsync(logFile, request.Offset, request.LineIndex, request.Data.ToArray(), request.Flush);
+			return new Empty();
+		}
+
+		/// <summary>
+		/// Uploads a new agent archive
+		/// </summary>
+		/// <param name="request">Request arguments</param>
+		/// <param name="context">Context for the RPC call</param>
+		/// <returns>Information about the new agent</returns>
+		public override async Task<UploadSoftwareResponse> UploadSoftware(UploadSoftwareRequest request, ServerCallContext context)
+		{
+			if (!await _aclService.AuthorizeAsync(AclAction.UploadSoftware, context.GetHttpContext().User))
+			{
+				throw new StructuredRpcException(StatusCode.PermissionDenied, "Access to software is forbidden");
+			}
+
+			string version = await _agentSoftwareService.SetArchiveAsync(new AgentSoftwareChannelName(request.Channel), null, request.Data.ToArray());
+
+			UploadSoftwareResponse response = new UploadSoftwareResponse();
+			response.Version = version;
+			return response;
+		}
+
+		/// <summary>
+		/// Downloads a new agent archive
+		/// </summary>
+		/// <param name="request">Request arguments</param>
+		/// <param name="responseStream">Writer for the output data</param>
+		/// <param name="context">Context for the RPC call</param>
+		/// <returns>Information about the new agent</returns>
+		public override async Task DownloadSoftware(DownloadSoftwareRequest request, IServerStreamWriter<DownloadSoftwareResponse> responseStream, ServerCallContext context)
+		{
+			if (!await _aclService.AuthorizeAsync(AclAction.DownloadSoftware, context.GetHttpContext().User))
+			{
+				throw new StructuredRpcException(StatusCode.NotFound, "Access to software is forbidden");
+			}
+
+			byte[]? data = await _agentSoftwareService.GetArchiveAsync(request.Version);
+			if (data == null)
+			{
+				throw new StructuredRpcException(StatusCode.NotFound, "Missing version {Version}");
+			}
+
+			for (int offset = 0; offset < data.Length;)
+			{
+				int nextOffset = Math.Min(offset + 128 * 1024, data.Length);
+
+				DownloadSoftwareResponse response = new DownloadSoftwareResponse();
+				response.Data = Google.Protobuf.ByteString.CopyFrom(data.AsSpan(offset, nextOffset - offset));
+
+				await responseStream.WriteAsync(response);
+
+				offset = nextOffset;
+			}
+		}
+
+		/// <summary>
+		/// Uploads a new artifact
+		/// </summary>
+		/// <param name="reader">Request arguments</param>
+		/// <param name="context">Context for the RPC call</param>
+		/// <returns>Information about the new agent</returns>
+		public override async Task<UploadArtifactResponse> UploadArtifact(IAsyncStreamReader<UploadArtifactRequest> reader, ServerCallContext context)
+		{
+			// Advance to the metadata object
+			if (!await reader.MoveNext())
+			{
+				throw new StructuredRpcException(StatusCode.DataLoss, "Missing request for artifact upload");
+			}
+
+			// Read the request object
+			UploadArtifactMetadata? metadata = reader.Current.Metadata;
+			if (metadata == null)
+			{
+				throw new StructuredRpcException(StatusCode.DataLoss, "Expected metadata in first artifact request");
+			}
+
+			// Get the job and step
+			IJob job = await GetJobAsync(new JobId(metadata.JobId));
+			AuthorizeBatch(job, metadata.BatchId.ToSubResourceId(), context);
+
+			IJobStep? step;
+			if (!job.TryGetStep(metadata.BatchId.ToSubResourceId(), metadata.StepId.ToSubResourceId(), out step))
+			{
+				throw new StructuredRpcException(StatusCode.NotFound, "Unable to find step {JobId}:{BatchId}:{StepId}", job.Id, metadata.BatchId, metadata.StepId);
+			}
+
+			// Upload the stream
+			using (ArtifactChunkStream inputStream = new ArtifactChunkStream(reader, metadata.Length))
+			{
+				IArtifact artifact = await _artifactCollection.CreateArtifactAsync(job.Id, step.Id, metadata.Name, metadata.MimeType, inputStream);
+
+				UploadArtifactResponse response = new UploadArtifactResponse();
+				response.Id = artifact.Id.ToString();
+				return response;
+			}
+		}
+
+		/// <summary>
+		/// Uploads new test data
+		/// </summary>
+		/// <param name="reader">Request arguments</param>
+		/// <param name="context">Context for the RPC call</param>
+		/// <returns>Information about the new agent</returns>
+		public override async Task<UploadTestDataResponse> UploadTestData(IAsyncStreamReader<UploadTestDataRequest> reader, ServerCallContext context)
+		{
+			IJob? job = null;
+			IJobStep? jobStep = null;
+
+			while (await reader.MoveNext())
+			{
+				UploadTestDataRequest request = reader.Current;
+
+				JobId jobId = new JobId(request.JobId);
+				if (job == null || jobId != job.Id)
+				{
+					job = await _jobService.GetJobAsync(jobId);
+					if (job == null)
+					{
+						throw new StructuredRpcException(StatusCode.NotFound, "Unable to find job {JobId}", jobId);
+					}
+					jobStep = null;
+				}
+
+				SubResourceId jobStepId = request.JobStepId.ToSubResourceId();
+				if (jobStep == null || jobStepId != jobStep.Id)
+				{
+					if (!job.TryGetStep(jobStepId, out jobStep))
+					{
+						throw new StructuredRpcException(StatusCode.NotFound, "Unable to find step {JobStepId} on job {JobId}", jobStepId, jobId);
+					}
+				}
+
+				string text = Encoding.UTF8.GetString(request.Value.ToArray());
+				BsonDocument document = BsonSerializer.Deserialize<BsonDocument>(text);
+				await _testData.AddAsync(job, jobStep, request.Key, document);
+			}
+
+			return new UploadTestDataResponse();
+		}
+
+		/// <summary>
+		/// Create a new report on a job or job step
+		/// </summary>
+		/// <param name="request"></param>
+		/// <param name="context"></param>
+		/// <returns></returns>
+		public override async Task<CreateReportResponse> CreateReport(CreateReportRequest request, ServerCallContext context)
+		{
+			IJob job = await GetJobAsync(new JobId(request.JobId));
+			IJobStepBatch batch = AuthorizeBatch(job, request.BatchId.ToSubResourceId(), context);
+
+			Report newReport = new Report { Name = request.Name, Placement = request.Placement, ArtifactId = request.ArtifactId.ToObjectId() };
+			if (request.Scope == ReportScope.Job)
+			{
+				_logger.LogDebug("Adding report to job {JobId}: {Name} -> {ArtifactId}", job.Id, request.Name, request.ArtifactId);
+				await _jobService.UpdateJobAsync(job, reports: new List<Report> { newReport });
+			}
+			else
+			{
+				_logger.LogDebug("Adding report to step {JobId}:{BatchId}:{StepId}: {Name} -> {ArtifactId}", job.Id, batch.Id, request.StepId, request.Name, request.ArtifactId);
+				await _jobService.UpdateStepAsync(job, batch.Id, request.StepId.ToSubResourceId(), newReports: new List<Report> { newReport });
+			}
+
+			return new CreateReportResponse();
+		}
+	}
+}
diff --git a/Engine/Source/Programs/Horde/Horde.Build/Server/ServerController.cs b/Engine/Source/Programs/Horde/Horde.Build/Server/ServerController.cs
new file mode 100644
--- /dev/null
+++ b/Engine/Source/Programs/Horde/Horde.Build/Server/ServerController.cs
@@ -0,0 +1,56 @@
+// Copyright Epic Games, Inc. All Rights Reserved.
+
+using System.Diagnostics;
+using System.Reflection;
+using Horde.Build.Api;
+using Microsoft.AspNetCore.Authorization;
+using Microsoft.AspNetCore.Mvc;
+using Microsoft.Extensions.Options;
+
+namespace Horde.Build.Controllers
+{
+	/// <summary>
+	/// Controller managing account status
+	/// </summary>
+	[ApiController]
+	[Authorize]
+	[Route("[controller]")]
+	public class ServerController : ControllerBase
+	{
+
+		/// <summary>
+		/// Settings for the server
+		/// </summary>
+		readonly IOptionsMonitor<ServerSettings> _settings;
+
+		/// <summary>
+		/// Constructor
+		/// </summary>
+		public ServerController(IOptionsMonitor<ServerSettings> settings)
+		{
+			_settings = settings;
+		}
+		
+		/// <summary>
+		/// Get server version
+		/// </summary>
+		[HttpGet]
+		[Route("/api/v1/server/version")]
+		public ActionResult GetVersionAsync()
+		{
+			FileVersionInfo fileVersionInfo = FileVersionInfo.GetVersionInfo(Assembly.GetExecutingAssembly().Location);
+			return Ok(fileVersionInfo.ProductVersion);
+		}		
+
+		/// <summary>
+		/// Get server information
+		/// </summary>
+		[HttpGet]
+		[Route("/api/v1/server/info")]
+		[ProducesResponseType(typeof(GetServerInfoResponse), 200)]
+		public ActionResult<GetServerInfoResponse> GetServerInfo()
+		{
+			return new GetServerInfoResponse(_settings.CurrentValue.SingleInstance);
+		}
+	}
+}
diff --git a/Engine/Source/Programs/Horde/Horde.Build/Server/ServerMessages.cs b/Engine/Source/Programs/Horde/Horde.Build/Server/ServerMessages.cs
new file mode 100644
--- /dev/null
+++ b/Engine/Source/Programs/Horde/Horde.Build/Server/ServerMessages.cs
@@ -0,0 +1,42 @@
+// Copyright Epic Games, Inc. All Rights Reserved.
+
+using System;
+using System.Diagnostics;
+using System.Reflection;
+using System.Runtime.InteropServices;
+
+namespace Horde.Build.Api
+{
+	/// <summary>
+	/// Server Info
+	/// </summary>
+	public class GetServerInfoResponse
+	{
+        /// <summary>
+		/// Server version info
+		/// </summary>
+        public string ServerVersion { get; set; }
+
+        /// <summary>
+        /// The operating system server is hosted on
+        /// </summary>
+        public string OsDescription { get; set; }
+
+		/// <summary>
+		/// Whether this is an installed Horde build
+		/// </summary>
+		public bool SingleInstance { get; set; }
+
+		/// <summary>
+		/// Constructor
+		/// </summary>
+		public GetServerInfoResponse( bool singleInstance)
+        {
+            FileVersionInfo versionInfo = FileVersionInfo.GetVersionInfo(Assembly.GetExecutingAssembly().Location);		
+			ServerVersion = versionInfo.ProductVersion ?? String.Empty;			
+			OsDescription = RuntimeInformation.OSDescription;
+			SingleInstance = singleInstance;
+		}
+	}
+}
+
diff --git a/Engine/Source/Programs/Horde/Horde.Build/Streams/IStream.cs b/Engine/Source/Programs/Horde/Horde.Build/Streams/IStream.cs
new file mode 100644
--- /dev/null
+++ b/Engine/Source/Programs/Horde/Horde.Build/Streams/IStream.cs
@@ -0,0 +1,785 @@
+// Copyright Epic Games, Inc. All Rights Reserved.
+
+using System;
+using System.Collections.Generic;
+using System.Diagnostics.CodeAnalysis;
+using System.Linq;
+using System.Text.RegularExpressions;
+using EpicGames.Core;
+using Horde.Build.Acls;
+using Horde.Build.Api;
+using Horde.Build.Server;
+using Horde.Build.Utilities;
+using HordeCommon;
+using MongoDB.Bson.Serialization.Attributes;
+
+namespace Horde.Build.Models
+{
+	using PoolId = StringId<IPool>;
+	using ProjectId = StringId<IProject>;
+	using StreamId = StringId<IStream>;
+	using TemplateRefId = StringId<TemplateRef>;
+	using UserId = ObjectId<IUser>;
+
+	/// <summary>
+	/// Exception thrown when stream validation fails
+	/// </summary>
+	public class InvalidStreamException : Exception
+	{
+		/// <inheritdoc/>
+		public InvalidStreamException()
+		{
+		}
+
+		/// <inheritdoc/>
+		public InvalidStreamException(string message) : base(message)
+		{
+		}
+
+		/// <inheritdoc/>
+		public InvalidStreamException(string message, Exception innerEx) : base(message, innerEx)
+		{
+		}
+	}
+
+	/// <summary>
+	/// Mapping from a BuildGraph agent type to a set of machines on the farm
+	/// </summary>
+	public class AgentType
+	{
+		/// <summary>
+		/// Name of the pool of agents to use
+		/// </summary>
+		[BsonRequired]
+		public PoolId Pool { get; set; }
+
+		/// <summary>
+		/// Name of the workspace to execute on
+		/// </summary>
+		public string? Workspace { get; set; }
+
+		/// <summary>
+		/// Path to the temporary storage dir
+		/// </summary>
+		public string? TempStorageDir { get; set; }
+
+		/// <summary>
+		/// Environment variables to be set when executing the job
+		/// </summary>
+		public Dictionary<string, string>? Environment { get; set; }
+
+		/// <summary>
+		/// Default constructor
+		/// </summary>
+		[BsonConstructor]
+		private AgentType()
+		{
+		}
+
+		/// <summary>
+		/// Constructor
+		/// </summary>
+		/// <param name="pool">The pool for this agent</param>
+		/// <param name="workspace">Name of the workspace to use</param>
+		/// <param name="tempStorageDir">Path to the temp storage directory</param>
+		public AgentType(PoolId pool, string workspace, string? tempStorageDir)
+		{
+			Pool = pool;
+			Workspace = workspace;
+			TempStorageDir = tempStorageDir;
+		}
+
+		/// <summary>
+		/// Constructor
+		/// </summary>
+		/// <param name="request">The object to construct from</param>
+		public AgentType(CreateAgentTypeRequest request)
+		{
+			Pool = new PoolId(request.Pool);
+			Workspace = request.Workspace;
+			TempStorageDir = request.TempStorageDir;
+			Environment = request.Environment;
+		}
+
+		/// <summary>
+		/// Constructs an AgentType object from an optional request
+		/// </summary>
+		/// <param name="request">The request object</param>
+		/// <returns>New agent type object</returns>
+		[return: NotNullIfNotNull("request")]
+		public static AgentType? FromRequest(CreateAgentTypeRequest? request)
+		{
+			return (request != null) ? new AgentType(request) : null;
+		}
+
+		/// <summary>
+		/// Creates an API response object from this stream
+		/// </summary>
+		/// <returns>The response object</returns>
+		public Api.GetAgentTypeResponse ToApiResponse()
+		{
+			return new Api.GetAgentTypeResponse(Pool.ToString(), Workspace, TempStorageDir, Environment);
+		}
+
+		/// <summary>
+		/// Creates an API response object from this stream
+		/// </summary>
+		/// <returns>The response object</returns>
+		public HordeCommon.Rpc.GetAgentTypeResponse ToRpcResponse()
+		{
+			HordeCommon.Rpc.GetAgentTypeResponse response = new HordeCommon.Rpc.GetAgentTypeResponse();
+			if (TempStorageDir != null)
+			{
+				response.TempStorageDir = TempStorageDir;
+			}
+			if (Environment != null)
+			{
+				response.Environment.Add(Environment);
+			}
+			return response;
+		}
+	}
+
+	/// <summary>
+	/// Information about a workspace type
+	/// </summary>
+	public class WorkspaceType
+	{
+		/// <summary>
+		/// Name of the Perforce cluster to use
+		/// </summary>
+		[BsonIgnoreIfNull]
+		public string? Cluster { get; set; }
+
+		/// <summary>
+		/// The Perforce server and port
+		/// </summary>
+		[BsonIgnoreIfNull]
+		public string? ServerAndPort { get; set; }
+
+		/// <summary>
+		/// The Perforce username for syncing this workspace
+		/// </summary>
+		[BsonIgnoreIfNull]
+		public string? UserName { get; set; }
+
+		/// <summary>
+		/// The Perforce password for syncing this workspace
+		/// </summary>
+		[BsonIgnoreIfNull]
+		public string? Password { get; set; }
+
+		/// <summary>
+		/// Identifier to distinguish this workspace from other workspaces. Defaults to the workspace type name.
+		/// </summary>
+		[BsonIgnoreIfNull]
+		public string? Identifier { get; set; }
+
+		/// <summary>
+		/// Override for the stream to sync
+		/// </summary>
+		[BsonIgnoreIfNull]
+		public string? Stream { get; set; }
+
+		/// <summary>
+		/// Custom view for the workspace
+		/// </summary>
+		[BsonIgnoreIfNull]
+		public List<string>? View { get; set; }
+
+		/// <summary>
+		/// Whether to use an incrementally synced workspace
+		/// </summary>
+		public bool Incremental { get; set; }
+
+		/// <summary>
+		/// Whether to use the AutoSDK
+		/// </summary>
+		public bool UseAutoSdk { get; set; } = true;
+
+		/// <summary>
+		/// Default constructor
+		/// </summary>
+		public WorkspaceType()
+		{
+		}
+
+		/// <summary>
+		/// Constructor
+		/// </summary>
+		/// <param name="request">The object to construct from</param>
+		public WorkspaceType(CreateWorkspaceTypeRequest request)
+		{
+			Cluster = request.Cluster;
+			ServerAndPort = request.ServerAndPort;
+			UserName = request.UserName;
+			Password = request.Password;
+			Identifier = request.Identifier;
+			Stream = request.Stream;
+			View = request.View;
+			Incremental = request.Incremental;
+			UseAutoSdk = request.UseAutoSdk;
+		}
+
+		/// <summary>
+		/// Constructs an AgentType object from an optional request
+		/// </summary>
+		/// <param name="request">The request object</param>
+		/// <returns>New agent type object</returns>
+		[return: NotNullIfNotNull("request")]
+		public static WorkspaceType? FromRequest(CreateWorkspaceTypeRequest? request)
+		{
+			return (request != null) ? new WorkspaceType(request) : null;
+		}
+
+		/// <summary>
+		/// Creates an API response object from this stream
+		/// </summary>
+		/// <returns>The response object</returns>
+		public Api.GetWorkspaceTypeResponse ToApiResponse()
+		{
+			return new Api.GetWorkspaceTypeResponse(Cluster, ServerAndPort, UserName, Identifier, Stream, View, Incremental, UseAutoSdk);
+		}
+	}
+
+	/// <summary>
+	/// Allows triggering another downstream job on succesful completion of a step or aggregate
+	/// </summary>
+	public class ChainedJobTemplate
+	{
+		/// <summary>
+		/// Name of the target that needs to complete successfully
+		/// </summary>
+		public string Trigger { get; set; } = String.Empty;
+
+		/// <summary>
+		/// The new template to trigger
+		/// </summary>
+		public TemplateRefId TemplateRefId { get; set; }
+
+		/// <summary>
+		/// Default constructor for serialization
+		/// </summary>
+		private ChainedJobTemplate()
+		{
+		}
+
+		/// <summary>
+		/// Constructor
+		/// </summary>
+		/// <param name="trigger">Name of the target that needs to complete</param>
+		/// <param name="templateRefId">The new template to trigger</param>
+		public ChainedJobTemplate(string trigger, TemplateRefId templateRefId)
+		{
+			Trigger = trigger;
+			TemplateRefId = templateRefId;
+		}
+
+		/// <summary>
+		/// Constructor
+		/// </summary>
+		/// <param name="request">Request to construct from</param>
+		public ChainedJobTemplate(CreateChainedJobTemplateRequest request)
+			: this(request.Trigger, new TemplateRefId(request.TemplateId))
+		{
+		}
+	}
+
+	/// <summary>
+	/// A
+	/// </summary>
+	public class TemplateStepState
+	{
+		/// <summary>
+		/// Name of the step
+		/// </summary>
+		public string Name { get; set; } = String.Empty ;
+
+		/// <summary>
+		/// User who paused the step
+		/// </summary>
+		public UserId? PausedByUserId { get; set; }
+
+		/// <summary>
+		/// The UTC time when the step was paused
+		/// </summary>
+		public DateTime? PauseTimeUtc { get; set; }
+
+
+		/// <summary>
+		/// Default constructor for serialization
+		/// </summary>
+		private TemplateStepState()
+		{
+		}
+
+		/// <summary>
+		/// Constructor
+		/// </summary>
+		public TemplateStepState( string name, UserId? pausedByUserId = null, DateTime? pauseTimeUtc = null)
+		{
+			Name = name;
+			PausedByUserId = pausedByUserId;
+			PauseTimeUtc = pauseTimeUtc;
+		}
+
+	}
+
+
+	/// <summary>
+	/// Reference to a template
+	/// </summary>
+	public class TemplateRef
+	{
+		/// <summary>
+		/// The template name (duplicated from the template object)
+		/// </summary>
+		public string Name { get; set; }
+
+		/// <summary>
+		/// Hash of the template definition
+		/// </summary>
+		public ContentHash Hash { get; set; }
+
+		/// <summary>
+		/// Whether to show badges in UGS for this schedule
+		/// </summary>
+		public bool ShowUgsBadges { get; set; }
+
+		/// <summary>
+		/// Whether to show desktop alerts for build health issues created from jobs this type
+		/// </summary>
+		public bool ShowUgsAlerts { get; set; }
+
+		/// <summary>
+		/// Notification channel for this template. Overrides the stream channel if set.
+		/// </summary>
+		public string? NotificationChannel { get; set; }
+
+		/// <summary>
+		/// Notification channel filter for this template. Errors|Warnings|Success
+		/// </summary>
+		public string? NotificationChannelFilter { get; set; }
+
+		/// <summary>
+		/// Channel for triage notification messages
+		/// </summary>
+		public string? TriageChannel { get; set; }
+
+		/// <summary>
+		/// List of schedules for this template
+		/// </summary>
+		[BsonIgnoreIfNull]
+		public Schedule? Schedule { get; set; }
+
+		/// <summary>
+		/// List of downstream templates to trigger at the same change
+		/// </summary>
+		[BsonIgnoreIfNull]
+		public List<ChainedJobTemplate>? ChainedJobs { get; set; }
+
+		/// <summary>
+		/// List of template step states
+		/// </summary>
+		[BsonIgnoreIfNull]
+		public List<TemplateStepState>? StepStates { get;set; }
+
+		/// <summary>
+		/// Custom permissions for this template
+		/// </summary>
+		public Acl? Acl { get; set; }
+
+		/// <summary>
+		/// Private constructor for serialization
+		/// </summary>
+		private TemplateRef()
+		{
+			Name = null!;
+			Hash = null!;
+		}
+
+		/// <summary>
+		/// Constructor
+		/// </summary>
+		/// <param name="template">The template being referenced</param>
+		/// <param name="showUgsBadges">Whether to show badges in UGS for this job</param>
+		/// <param name="showUgsAlerts">Whether to show alerts in UGS for this job</param>
+		/// <param name="notificationChannel">Notification channel for this template</param>
+		/// <param name="notificationChannelFilter">Notification channel filter for this template</param>
+		/// <param name="triageChannel"></param>
+		/// <param name="schedule">Schedule for this template</param>
+		/// <param name="triggers">List of downstream templates to trigger</param>
+		/// <param name="stepStates">List of template step states</param>
+		/// <param name="acl">ACL for this template</param>
+		public TemplateRef(ITemplate template, bool showUgsBadges = false, bool showUgsAlerts = false, string? notificationChannel = null, string? notificationChannelFilter = null, string? triageChannel = null, Schedule? schedule = null, List<ChainedJobTemplate>? triggers = null, List<TemplateStepState>? stepStates = null, Acl? acl = null)
+		{
+			Name = template.Name;
+			Hash = template.Id;
+			ShowUgsBadges = showUgsBadges;
+			ShowUgsAlerts = showUgsAlerts;
+			NotificationChannel = notificationChannel;
+			NotificationChannelFilter = notificationChannelFilter;
+			TriageChannel = triageChannel;
+			Schedule = schedule;
+			ChainedJobs = triggers;
+			StepStates = stepStates;
+			Acl = acl;
+		}
+	}
+
+	/// <summary>
+	/// Query used to identify a base changelist for a preflight
+	/// </summary>
+	public class ChangeQuery
+	{
+		/// <summary>
+		/// Template to search for
+		/// </summary>
+		public TemplateRefId? TemplateRefId { get; set; }
+
+		/// <summary>
+		/// The target to look at the status for
+		/// </summary>
+		public string? Target { get; set; }
+
+		/// <summary>
+		/// Whether to match a job that contains warnings
+		/// </summary>
+		public List<JobStepOutcome>? Outcomes { get; set; }
+
+		/// <summary>
+		/// Convert to a request object
+		/// </summary>
+		/// <returns></returns>
+		public ChangeQueryRequest ToRequest()
+		{
+			return new ChangeQueryRequest { TemplateId = TemplateRefId?.ToString(), Target = Target, Outcomes = Outcomes };
+		}
+	}
+
+	/// <summary>
+	/// Definition of a query to execute to find the changelist to run a build at
+	/// </summary>
+	public class DefaultPreflight
+	{
+		/// <summary>
+		/// The template id to execute
+		/// </summary>
+		public TemplateRefId? TemplateRefId { get; set; }
+
+		/// <summary>
+		/// Query specifying a changelist to use
+		/// </summary>
+		public ChangeQuery? Change { get; set; }
+
+		/// <summary>
+		/// The job type to query for the change to use
+		/// </summary>
+		[Obsolete("Use Change.TemplateRefId instead")]
+		public TemplateRefId? ChangeTemplateRefId { get; set; }
+
+		/// <summary>
+		/// Constructor
+		/// </summary>
+		/// <param name="templateRefId"></param>
+		/// <param name="change">The job type to query for the change to use</param>
+		public DefaultPreflight(TemplateRefId? templateRefId, ChangeQuery? change)
+		{
+			TemplateRefId = templateRefId;
+			Change = change;
+		}
+
+		/// <summary>
+		/// Convert to a request object
+		/// </summary>
+		/// <returns></returns>
+		public DefaultPreflightRequest ToRequest()
+		{
+#pragma warning disable CS0618 // Type or member is obsolete
+			ChangeQueryRequest? changeRequest = null;
+			if (Change != null)
+			{
+				changeRequest = Change.ToRequest();
+			}
+			else if (ChangeTemplateRefId != null)
+			{
+				changeRequest = new ChangeQueryRequest { TemplateId = ChangeTemplateRefId.ToString() };
+			}
+
+			return new DefaultPreflightRequest { TemplateId = TemplateRefId?.ToString(), Change = changeRequest, ChangeTemplateId = changeRequest?.TemplateId };
+#pragma warning restore CS0618 // Type or member is obsolete
+		}
+	}
+
+	/// <summary>
+	/// Extension methods for template refs
+	/// </summary>
+	static class TemplateRefExtensions
+	{
+		/// <summary>
+		/// Adds a new template ref to a list
+		/// </summary>
+		/// <param name="templateRefs">List of template refs</param>
+		/// <param name="templateRef">The template ref to add</param>
+		public static void AddRef(this Dictionary<TemplateRefId, TemplateRef> templateRefs, TemplateRef templateRef)
+		{
+			templateRefs.Add(new TemplateRefId(templateRef.Name), templateRef);
+		}
+	}
+
+	/// <summary>
+	/// How to replicate data for this stream
+	/// </summary>
+	public enum ContentReplicationMode
+	{
+		/// <summary>
+		/// No content will be replicated for this stream
+		/// </summary>
+		None,
+
+		/// <summary>
+		/// Only replicate depot path and revision data for each file
+		/// </summary>
+		RevisionsOnly,
+
+		/// <summary>
+		/// Replicate full stream contents to storage
+		/// </summary>
+		Full,
+	}
+
+	/// <summary>
+	/// Information about a stream
+	/// </summary>
+	public interface IStream
+	{
+		/// <summary>
+		/// Name of the stream.
+		/// </summary>
+		public StreamId Id { get; }
+
+		/// <summary>
+		/// The project that this stream belongs to
+		/// </summary>
+		public ProjectId ProjectId { get; }
+
+		/// <summary>
+		/// The stream name
+		/// </summary>
+		public string Name { get; }
+
+		/// <summary>
+		/// Name of the perforce cluster
+		/// </summary>
+		public string ClusterName { get; }
+
+		/// <summary>
+		/// The revision of config file used for this stream
+		/// </summary>
+		public string ConfigRevision { get; }
+
+		/// <summary>
+		/// Configuration settings for the stream
+		/// </summary>
+		public StreamConfig Config { get; }
+
+		/// <summary>
+		/// Order to display on the dashboard's drop-down list
+		/// </summary>
+		public int Order { get; }
+
+		/// <summary>
+		/// Notification channel for all jobs in this stream
+		/// </summary>
+		public string? NotificationChannel { get; }
+
+		/// <summary>
+		/// Notification channel filter for all jobs in this stream. Errors|Warnings|Success
+		/// </summary>
+		public string? NotificationChannelFilter { get; }
+
+		/// <summary>
+		/// Channel to post issue triage notifications
+		/// </summary>
+		public string? TriageChannel { get; }
+
+		/// <summary>
+		/// Default template to use for preflights
+		/// </summary>
+		public DefaultPreflight? DefaultPreflight { get; }
+
+		/// <summary>
+		/// List of pages to display in the dashboard
+		/// </summary>
+		public IReadOnlyList<StreamTab> Tabs { get; }
+
+		/// <summary>
+		/// Dictionary of agent types
+		/// </summary>
+		public IReadOnlyDictionary<string, AgentType> AgentTypes { get; }
+
+		/// <summary>
+		/// Dictionary of workspace types
+		/// </summary>
+		public IReadOnlyDictionary<string, WorkspaceType> WorkspaceTypes { get; }
+
+		/// <summary>
+		/// List of templates available for this stream
+		/// </summary>
+		public IReadOnlyDictionary<TemplateRefId, TemplateRef> Templates { get; }
+
+		/// <summary>
+		/// Stream is paused for builds until specified time
+		/// </summary>
+		public DateTime? PausedUntil { get; }
+		
+		/// <summary>
+		/// Comment/reason for why the stream was paused
+		/// </summary>
+		public string? PauseComment { get; }
+
+		/// <summary>
+		/// Whether to replicate data from Perforce to Horde Storage.
+		/// </summary>
+		public ContentReplicationMode ReplicationMode { get; }
+
+		/// <summary>
+		/// Filter for paths to be replicated. May be removed in future.
+		/// </summary>
+		public string? ReplicationFilter { get; }
+
+		/// <summary>
+		/// The ACL for this object
+		/// </summary>
+		public Acl? Acl { get; }
+	}
+
+	/// <summary>
+	/// Extension methods for streams
+	/// </summary>
+	static class StreamExtensions
+	{
+		/// <summary>
+		/// Tries to get an agent workspace definition from the given type name
+		/// </summary>
+		/// <param name="stream">The stream object</param>
+		/// <param name="agentType">The agent type</param>
+		/// <param name="workspace">Receives the agent workspace definition</param>
+		/// <returns>True if the agent type was valid, and an agent workspace could be created</returns>
+		public static bool TryGetAgentWorkspace(this IStream stream, AgentType agentType, [NotNullWhen(true)] out (AgentWorkspace, bool)? workspace)
+		{
+			// Get the workspace settings
+			if (agentType.Workspace == null)
+			{
+				// Use the default settings (fast switching workspace, clean 
+				workspace = (new AgentWorkspace(null, null, GetDefaultWorkspaceIdentifier(stream), stream.Name, null, false), true);
+				return true;
+			}
+			else
+			{
+				// Try to get the matching workspace type
+				WorkspaceType? workspaceType;
+				if (!stream.WorkspaceTypes.TryGetValue(agentType.Workspace, out workspaceType))
+				{
+					workspace = null;
+					return false;
+				}
+
+				// Get the workspace identifier
+				string identifier;
+				if (workspaceType.Identifier != null)
+				{
+					identifier = workspaceType.Identifier;
+				}
+				else if (workspaceType.Incremental)
+				{
+					identifier = $"{stream.GetEscapedName()}+{agentType.Workspace}";
+				}
+				else
+				{
+					identifier = GetDefaultWorkspaceIdentifier(stream);
+				}
+
+				// Create the new workspace
+				workspace = (new AgentWorkspace(workspaceType.Cluster, workspaceType.UserName, identifier, workspaceType.Stream ?? stream.Name, workspaceType.View, workspaceType.Incremental), workspaceType.UseAutoSdk);
+				return true;
+			}
+		}
+
+		/// <summary>
+		/// The escaped name of this stream. Removes all non-identifier characters.
+		/// </summary>
+		/// <param name="stream">The stream object</param>
+		/// <returns>Escaped name for the stream</returns>
+		public static string GetEscapedName(this IStream stream)
+		{
+			return Regex.Replace(stream.Name, @"[^a-zA-Z0-9_]", "+");
+		}
+
+		/// <summary>
+		/// Gets the default identifier for workspaces created for this stream. Just includes an escaped depot name.
+		/// </summary>
+		/// <param name="stream">The stream object</param>
+		/// <returns>The default workspace identifier</returns>
+		private static string GetDefaultWorkspaceIdentifier(IStream stream)
+		{
+			return Regex.Replace(stream.GetEscapedName(), @"^(\+\+[^+]*).*$", "$1");
+		}
+
+		/// <summary>
+		/// Converts to a public response object
+		/// </summary>
+		/// <param name="stream">The stream object</param>
+		/// <param name="bIncludeAcl">Whether to include the ACL in the response object</param>
+		/// <param name="apiTemplateRefs">The template refs for this stream. Passed separately because they have their own ACL.</param>
+		/// <returns>New response instance</returns>
+		public static Api.GetStreamResponse ToApiResponse(this IStream stream, bool bIncludeAcl, List<GetTemplateRefResponse> apiTemplateRefs)
+		{
+			List<GetStreamTabResponse> apiTabs = stream.Tabs.ConvertAll(x => x.ToResponse());
+			Dictionary<string, GetAgentTypeResponse> apiAgentTypes = stream.AgentTypes.ToDictionary(x => x.Key, x => x.Value.ToApiResponse());
+			Dictionary<string, GetWorkspaceTypeResponse> apiWorkspaceTypes = stream.WorkspaceTypes.ToDictionary(x => x.Key, x => x.Value.ToApiResponse());
+			GetAclResponse? apiAcl = (bIncludeAcl && stream.Acl != null)? new GetAclResponse(stream.Acl) : null;
+			return new Api.GetStreamResponse(stream.Id.ToString(), stream.ProjectId.ToString(), stream.Name, stream.ConfigRevision, stream.Order, stream.NotificationChannel, stream.NotificationChannelFilter, stream.TriageChannel, stream.DefaultPreflight?.ToRequest(), apiTabs, apiAgentTypes, apiWorkspaceTypes, apiTemplateRefs, apiAcl, stream.PausedUntil, stream.PauseComment, stream.Config.Workflows);
+		}
+
+		/// <summary>
+		/// Converts to an RPC response object
+		/// </summary>
+		/// <param name="stream">The stream object</param>
+		/// <returns>New response instance</returns>
+		public static HordeCommon.Rpc.GetStreamResponse ToRpcResponse(this IStream stream)
+		{
+			HordeCommon.Rpc.GetStreamResponse response = new HordeCommon.Rpc.GetStreamResponse();
+			response.Name = stream.Name;
+			response.AgentTypes.Add(stream.AgentTypes.ToDictionary(x => x.Key, x => x.Value.ToRpcResponse()));
+			return response;
+		}
+
+		/// <summary>
+		/// Check if stream is paused for new builds
+		/// </summary>
+		/// <param name="stream">The stream object</param>
+		/// <param name="currentTime">Current time (allow tests to pass in a fake clock)</param>
+		/// <returns>If stream is paused</returns>
+		public static bool IsPaused(this IStream stream, DateTime currentTime)
+		{
+			return stream.PausedUntil != null && stream.PausedUntil > currentTime;
+		}
+	}
+
+	/// <summary>
+	/// Projection of a stream definition to just include permissions info
+	/// </summary>
+	public interface IStreamPermissions
+	{
+		/// <summary>
+		/// ACL for the stream
+		/// </summary>
+		public Acl? Acl { get; }
+
+		/// <summary>
+		/// The project containing this stream
+		/// </summary>
+		public ProjectId ProjectId { get; }
+	}
+}
diff --git a/Engine/Source/Programs/Horde/Horde.Build/Streams/IStreamCollection.cs b/Engine/Source/Programs/Horde/Horde.Build/Streams/IStreamCollection.cs
new file mode 100644
--- /dev/null
+++ b/Engine/Source/Programs/Horde/Horde.Build/Streams/IStreamCollection.cs
@@ -0,0 +1,121 @@
+// Copyright Epic Games, Inc. All Rights Reserved.
+
+using System;
+using System.Collections.Generic;
+using System.Threading.Tasks;
+using Horde.Build.Api;
+using Horde.Build.Models;
+using Horde.Build.Server;
+using Horde.Build.Utilities;
+
+namespace Horde.Build.Collections
+{
+	using JobId = ObjectId<IJob>;
+	using ProjectId = StringId<IProject>;
+	using StreamId = StringId<IStream>;
+	using TemplateRefId = StringId<TemplateRef>;
+
+	/// <summary>
+	/// Collection of stream documents
+	/// </summary>
+	public interface IStreamCollection
+	{
+		/// <summary>
+		/// Creates or replaces a stream configuration
+		/// </summary>
+		/// <param name="id">Unique id for the new stream</param>
+		/// <param name="stream">The current stream value. If not-null, this will attempt to replace the existing instance.</param>
+		/// <param name="revision">The config file revision</param>
+		/// <param name="projectId">The project id</param>
+		/// <returns></returns>
+		Task<IStream?> TryCreateOrReplaceAsync(StreamId id, IStream? stream, string revision, ProjectId projectId);
+
+		/// <summary>
+		/// Gets a stream by ID
+		/// </summary>
+		/// <param name="streamId">Unique id of the stream</param>
+		/// <returns>The stream document</returns>
+		Task<IStream?> GetAsync(StreamId streamId);
+
+		/// <summary>
+		/// Gets a stream's permissions by ID
+		/// </summary>
+		/// <param name="streamId">Unique id of the stream</param>
+		/// <returns>The stream document</returns>
+		Task<IStreamPermissions?> GetPermissionsAsync(StreamId streamId);
+
+		/// <summary>
+		/// Enumerates all streams
+		/// </summary>
+		/// <returns></returns>
+		Task<List<IStream>> FindAllAsync();
+
+		/// <summary>
+		/// Gets all the available streams for a project
+		/// </summary>
+		/// <param name="projectIds">Unique id of the projects to query</param>
+		/// <returns>List of stream documents</returns>
+		Task<List<IStream>> FindForProjectsAsync(ProjectId[] projectIds);
+
+		/// <summary>
+		/// Updates user-facing properties for an existing stream
+		/// </summary>
+		/// <param name="stream">The stream to update</param>
+		/// <param name="newPausedUntil">The new datetime for pausing builds</param>
+		/// <param name="newPauseComment">The reason for pausing</param>
+		/// <returns>The updated stream if successful, null otherwise</returns>
+		Task<IStream?> TryUpdatePauseStateAsync(IStream stream, DateTime? newPausedUntil, string? newPauseComment);
+
+		/// <summary>
+		/// Attempts to update the last trigger time for a schedule
+		/// </summary>
+		/// <param name="stream">The stream to update</param>
+		/// <param name="templateRefId">The template ref id</param>
+		/// <param name="lastTriggerTimeUtc">New last trigger time for the schedule</param>
+		/// <param name="lastTriggerChange">New last trigger changelist for the schedule</param>
+		/// <param name="newActiveJobs">New list of active jobs</param>
+		/// <returns>The updated stream if successful, null otherwise</returns>
+		Task<IStream?> TryUpdateScheduleTriggerAsync(IStream stream, TemplateRefId templateRefId, DateTime? lastTriggerTimeUtc, int? lastTriggerChange, List<JobId> newActiveJobs);
+
+		/// <summary>
+		/// Attempts to update a stream template ref
+		/// </summary>
+		/// <param name="streamInterface">The stream containing the template ref</param>
+		/// <param name="templateRefId">The template ref to update</param>
+		/// <param name="stepStates">The stream states to update, pass an empty list to clear all step states, otherwise will be a partial update based on included step updates</param>
+		/// <returns></returns>
+		Task<IStream?> TryUpdateTemplateRefAsync(IStream streamInterface, TemplateRefId templateRefId, List<UpdateStepStateRequest>? stepStates = null);
+
+		/// <summary>
+		/// Delete a stream
+		/// </summary>
+		/// <param name="streamId">Unique id of the stream</param>
+		/// <returns>Async task</returns>
+		Task DeleteAsync(StreamId streamId);
+	}
+
+	static class StreamCollectionExtensions
+	{
+		/// <summary>
+		/// Creates or replaces a stream configuration
+		/// </summary>
+		/// <param name="streamCollection">The stream collection</param>
+		/// <param name="id">Unique id for the new stream</param>
+		/// <param name="stream">The current stream value. If not-null, this will attempt to replace the existing instance.</param>
+		/// <param name="revision">The config file revision</param>
+		/// <param name="projectId">The project id</param>
+		/// <returns></returns>
+		public static async Task<IStream> CreateOrReplaceAsync(this IStreamCollection streamCollection, StreamId id, IStream? stream, string revision, ProjectId projectId)
+		{
+			for (; ; )
+			{
+				stream = await streamCollection.TryCreateOrReplaceAsync(id, stream, revision, projectId);
+				if (stream != null)
+				{
+					return stream;
+				}
+				stream = await streamCollection.GetAsync(id);
+			}
+		}
+	}
+}
diff --git a/Engine/Source/Programs/Horde/Horde.Build/Streams/StreamCollection.cs b/Engine/Source/Programs/Horde/Horde.Build/Streams/StreamCollection.cs
new file mode 100644
--- /dev/null
+++ b/Engine/Source/Programs/Horde/Horde.Build/Streams/StreamCollection.cs
@@ -0,0 +1,576 @@
+// Copyright Epic Games, Inc. All Rights Reserved.
+
+using System;
+using System.Collections.Generic;
+using System.Diagnostics.CodeAnalysis;
+using System.IO;
+using System.Linq;
+using System.Threading.Tasks;
+using Horde.Build.Acls;
+using Horde.Build.Api;
+using Horde.Build.Config;
+using Horde.Build.Models;
+using Horde.Build.Server;
+using Horde.Build.Utilities;
+using HordeCommon;
+using Microsoft.Extensions.Logging;
+using MongoDB.Bson.Serialization.Attributes;
+using MongoDB.Driver;
+
+namespace Horde.Build.Collections.Impl
+{
+	using JobId = ObjectId<IJob>;
+	using ProjectId = StringId<IProject>;
+	using StreamId = StringId<IStream>;
+	using TemplateRefId = StringId<TemplateRef>;
+	using UserId = ObjectId<IUser>;
+
+	/// <summary>
+	/// Collection of stream documents
+	/// </summary>
+	public class StreamCollection : IStreamCollection
+	{
+		/// <summary>
+		/// Information about a stream
+		/// </summary>
+		class StreamDocument : IStream
+		{
+			public const int DefaultOrder = 128;
+
+			[BsonRequired, BsonId]
+			public StreamId Id { get; set; }
+
+			[BsonRequired]
+			public ProjectId ProjectId { get; set; }
+
+			[BsonRequired]
+			public string Name { get; set; }
+
+			public string? ClusterName { get; set; }
+			public string ConfigRevision { get; set; } = String.Empty;
+
+			[BsonIgnore]
+			public StreamConfig? Config { get; set; }
+
+			public int Order { get; set; } = DefaultOrder;
+			public string? NotificationChannel { get; set; }
+			public string? NotificationChannelFilter { get; set; }
+			public string? TriageChannel { get; set; }
+			public DefaultPreflight? DefaultPreflight { get; set; }
+			public List<StreamTab> Tabs { get; set; } = new List<StreamTab>();
+			public Dictionary<string, AgentType> AgentTypes { get; set; } = new Dictionary<string, AgentType>(StringComparer.Ordinal);
+			public Dictionary<string, WorkspaceType> WorkspaceTypes { get; set; } = new Dictionary<string, WorkspaceType>(StringComparer.Ordinal);
+			public Dictionary<TemplateRefId, TemplateRef> Templates { get; set; } = new Dictionary<TemplateRefId, TemplateRef>();
+			public DateTime? PausedUntil { get; set; }
+			public string? PauseComment { get; set; }
+
+			[BsonIgnoreIfDefault]
+			[BsonDefaultValue(ContentReplicationMode.None)]
+			public ContentReplicationMode ReplicationMode { get; set; }
+
+			[BsonIgnoreIfNull]
+			public string? ReplicationFilter { get; set; }
+
+			public Acl? Acl { get; set; }
+			public int UpdateIndex { get; set; }
+			public bool Deleted { get; set; }
+
+			StreamConfig IStream.Config => Config!;
+			string IStream.ClusterName => ClusterName ?? PerforceCluster.DefaultName;
+			IReadOnlyList<StreamTab> IStream.Tabs => Tabs;
+			IReadOnlyDictionary<string, AgentType> IStream.AgentTypes => AgentTypes;
+			IReadOnlyDictionary<string, WorkspaceType> IStream.WorkspaceTypes => WorkspaceTypes;
+			IReadOnlyDictionary<TemplateRefId, TemplateRef> IStream.Templates => Templates;
+
+			[BsonConstructor]
+			private StreamDocument()
+			{
+				Name = null!;
+			}
+
+			public StreamDocument(StreamId id, string name, ProjectId projectId)
+			{
+				Id = id;
+				Name = name;
+				ProjectId = projectId;
+			}
+		}
+
+		/// <summary>
+		/// Projection of a stream definition to just include permissions info
+		/// </summary>
+		[SuppressMessage("Design", "CA1812: Class is never instantiated")]
+		private class StreamPermissions : IStreamPermissions
+		{
+			public Acl? Acl { get; set; }
+			public ProjectId ProjectId { get; set; }
+
+			public static readonly ProjectionDefinition<StreamDocument> Projection = Builders<StreamDocument>.Projection.Include(x => x.Acl).Include(x => x.ProjectId);
+		}
+
+		readonly IMongoCollection<StreamDocument> _streams;
+		readonly ConfigCollection _configCollection;
+		readonly IClock _clock;
+		readonly ITemplateCollection _templateCollection;
+		readonly ILogger<StreamCollection> _logger;
+
+		/// <summary>
+		/// Constructor
+		/// </summary>
+		/// <param name="mongoService">The database service instance</param>
+		/// <param name="configCollection"></param>
+		/// <param name="clock"></param>
+		/// <param name="templateCollection"></param>
+		/// <param name="logger"></param>
+		public StreamCollection(MongoService mongoService, ConfigCollection configCollection, IClock clock, ITemplateCollection templateCollection, ILogger<StreamCollection> logger)
+		{
+			_streams = mongoService.GetCollection<StreamDocument>("Streams");
+			_configCollection = configCollection;
+			_clock = clock;
+			_templateCollection = templateCollection;
+			_logger = logger;
+		}
+
+		async Task PostLoadAsync(StreamDocument stream)
+		{
+			try
+			{
+				if (stream.Deleted)
+				{
+					stream.Config = new StreamConfig();
+				}
+				else
+				{
+					stream.Config = await _configCollection.GetConfigAsync<StreamConfig>(stream.ConfigRevision);
+				}
+			}
+			catch (Exception)
+			{
+				_logger.LogError("Unable to get stream config for {StreamId} at {Revision}", stream.Id, stream.ConfigRevision);
+				throw;
+			}
+		}
+
+		/// <inheritdoc/>
+		public async Task<IStream?> TryCreateOrReplaceAsync(StreamId id, IStream? stream, string revision, ProjectId projectId)
+		{
+			StreamConfig config = await _configCollection.GetConfigAsync<StreamConfig>(revision);
+
+			List<StreamTab> tabs = config.Tabs.ConvertAll(x => StreamTab.FromRequest(x));
+			Dictionary<TemplateRefId, TemplateRef> templateRefs = await CreateTemplateRefsAsync(config.Templates, stream, _templateCollection);
+
+			Dictionary<string, AgentType> agentTypes = new Dictionary<string, AgentType>();
+			if (config.AgentTypes != null)
+			{
+				agentTypes = config.AgentTypes.Where(x => x.Value != null).ToDictionary(x => x.Key, x => new AgentType(x.Value!));
+			}
+
+			Dictionary<string, WorkspaceType> workspaceTypes = new Dictionary<string, WorkspaceType>();
+			if (config.WorkspaceTypes != null)
+			{
+				workspaceTypes = config.WorkspaceTypes.Where(x => x.Value != null).ToDictionary(x => x.Key, x => new WorkspaceType(x.Value!));
+			}
+
+			DefaultPreflight? defaultPreflight = config.DefaultPreflight?.ToModel();
+			if (defaultPreflight == null && config.DefaultPreflightTemplate != null)
+			{
+				defaultPreflight = new DefaultPreflight(new TemplateRefId(config.DefaultPreflightTemplate), null);
+			}
+
+			Validate(id, defaultPreflight, templateRefs, tabs, agentTypes, workspaceTypes);
+
+			Acl? acl = Acl.Merge(new Acl(), config.Acl);
+			if (stream == null)
+			{
+				return await TryCreateAsync(id, projectId, revision, config, defaultPreflight, tabs, agentTypes, workspaceTypes, templateRefs, acl);
+			}
+			else
+			{
+				return await TryReplaceAsync(stream, projectId, revision, config, defaultPreflight, tabs, agentTypes, workspaceTypes, templateRefs, acl);
+			}
+		}
+
+		/// <summary>
+		/// Creates a list of template refs from a set of request objects
+		/// </summary>
+		/// <param name="requests">Request objects</param>
+		/// <param name="stream">The current stream state</param>
+		/// <param name="templateCollection">The template service</param>
+		/// <returns>List of new template references</returns>
+		async Task<Dictionary<TemplateRefId, TemplateRef>> CreateTemplateRefsAsync(List<TemplateRefConfig> requests, IStream? stream, ITemplateCollection templateCollection)
+		{
+			Dictionary<TemplateRefId, TemplateRef> newTemplateRefs = new Dictionary<TemplateRefId, TemplateRef>();
+			foreach (TemplateRefConfig request in requests)
+			{
+				// Create the template
+				ITemplate newTemplate = await templateCollection.AddAsync(request.Name, request.Priority, request.AllowPreflights, request.UpdateIssues, request.PromoteIssuesByDefault, request.InitialAgentType, request.SubmitNewChange, request.SubmitDescription, request.Arguments, request.Parameters.ConvertAll(x => x.ToModel()));
+
+				// Get an identifier for the new template ref
+				TemplateRefId newTemplateRefId = request.Id;
+
+				// Create the schedule object
+				Schedule? schedule = null;
+				if (request.Schedule != null)
+				{
+					schedule = request.Schedule.ToModel(_clock.UtcNow);
+				}
+
+				// Add it to the list
+				TemplateRef newTemplateRef = new TemplateRef(newTemplate, request.ShowUgsBadges, request.ShowUgsAlerts, request.NotificationChannel, request.NotificationChannelFilter, request.TriageChannel, schedule, request.ChainedJobs?.ConvertAll(x => new ChainedJobTemplate(x)), null, Acl.Merge(null, request.Acl));
+				if (stream != null && stream.Templates.TryGetValue(newTemplateRefId, out TemplateRef? oldTemplateRef))
+				{
+					if (oldTemplateRef.Schedule != null && newTemplateRef.Schedule != null)
+					{
+						newTemplateRef.Schedule.CopyState(oldTemplateRef.Schedule);
+					}
+
+					if (oldTemplateRef.StepStates != null)
+					{
+						newTemplateRef.StepStates = new List<TemplateStepState>(oldTemplateRef.StepStates);
+					}
+
+				}
+				newTemplateRefs.Add(newTemplateRefId, newTemplateRef);
+			}
+			foreach (TemplateRef templateRef in newTemplateRefs.Values)
+			{
+				if (templateRef.ChainedJobs != null)
+				{
+					foreach (ChainedJobTemplate chainedJob in templateRef.ChainedJobs)
+					{
+						if (!newTemplateRefs.ContainsKey(chainedJob.TemplateRefId))
+						{
+							throw new InvalidDataException($"Invalid template ref id '{chainedJob.TemplateRefId}");
+						}
+					}
+				}
+			}
+			return newTemplateRefs;
+		}
+
+		/// <inheritdoc/>
+		public async Task<IStream?> TryUpdateTemplateRefAsync(IStream streamInterface, TemplateRefId templateRefId, List<UpdateStepStateRequest>? stepStates = null)
+		{
+			StreamDocument stream = (StreamDocument)streamInterface;
+
+			TemplateRef? oldTemplateRef;
+			TemplateRef? newTemplateRef = null;
+
+			if (!stream.Templates.TryGetValue(templateRefId, out oldTemplateRef))
+			{
+				return null;
+			}
+
+			ITemplate? template = await _templateCollection.GetAsync(oldTemplateRef.Hash);
+			if (template == null)
+			{
+				return null;
+			}
+
+			UpdateDefinitionBuilder<StreamDocument> updateBuilder = Builders<StreamDocument>.Update;
+			List<UpdateDefinition<StreamDocument>> updates = new List<UpdateDefinition<StreamDocument>>();
+
+			Dictionary<TemplateRefId, TemplateRef> newTemplates = new Dictionary<TemplateRefId, TemplateRef>(stream.Templates);
+
+			// clear
+			if (stepStates != null && stepStates.Count == 0)
+			{
+				bool hasUpdates = false;
+				foreach (KeyValuePair<TemplateRefId, TemplateRef> entry in newTemplates)
+				{
+					if (entry.Value?.StepStates != null)
+					{
+						hasUpdates = true;
+						entry.Value.StepStates = null;
+					}
+				}
+
+				if (hasUpdates)
+				{
+					updates.Add(updateBuilder.Set(x => x.Templates, newTemplates));
+				}
+				
+			}
+			else if (stepStates != null)
+			{
+				newTemplateRef = new TemplateRef(template, oldTemplateRef.ShowUgsAlerts, oldTemplateRef.ShowUgsAlerts, oldTemplateRef.NotificationChannel, oldTemplateRef.NotificationChannelFilter, oldTemplateRef.TriageChannel, oldTemplateRef.Schedule, oldTemplateRef.ChainedJobs, null, oldTemplateRef.Acl);
+
+				// get currently valid step states
+				List<TemplateStepState> newStepStates = oldTemplateRef.StepStates?.Where(x => x.PausedByUserId != null).ToList() ?? new List<TemplateStepState>();
+
+				// generate update list
+				foreach (UpdateStepStateRequest updateState in stepStates)
+				{
+					TemplateStepState? newState = newStepStates.Where(x => x.Name == updateState.Name).FirstOrDefault();
+
+					UserId? PausedByUserId = updateState.PausedByUserId != null ? new UserId(updateState.PausedByUserId) : null;
+
+					if (newState == null)
+					{
+						// if this is a new state without anything set, ignore it
+						if (PausedByUserId == null)
+						{
+							continue;
+						}
+
+						newStepStates.Add(new TemplateStepState(updateState.Name, PausedByUserId, PausedByUserId != null ? DateTime.UtcNow : null));
+					}
+					else
+					{
+						if (PausedByUserId == null)
+						{
+							newState.PauseTimeUtc = null;
+						}
+						else if (newState.PauseTimeUtc == null)
+						{
+							newState.PauseTimeUtc = DateTime.UtcNow;
+						}
+
+						newState.PausedByUserId = PausedByUserId;						
+					}
+				}
+
+				newTemplateRef.StepStates = newStepStates.Where(x => x.PausedByUserId != null).ToList();
+				if (newTemplateRef.StepStates.Count == 0)
+				{
+					newTemplateRef.StepStates = null;
+				}
+
+				newTemplates[templateRefId] = newTemplateRef;
+				updates.Add(updateBuilder.Set(x => x.Templates, newTemplates));
+
+			}
+
+			if (updates.Count == 0)
+			{
+				return streamInterface;
+			}
+
+			return await TryUpdateStreamAsync(stream, updateBuilder.Combine(updates));
+
+		}
+
+		/// <inheritdoc/>
+		async Task<IStream?> TryCreateAsync(StreamId id, ProjectId projectId, string configRevision, StreamConfig config, DefaultPreflight? defaultPreflight, List<StreamTab> tabs, Dictionary<string, AgentType> agentTypes, Dictionary<string, WorkspaceType> workspaceTypes, Dictionary<TemplateRefId, TemplateRef> templateRefs, Acl? acl)
+		{
+			StreamDocument newStream = new StreamDocument(id, config.Name, projectId);
+			newStream.ClusterName = config.ClusterName;
+			newStream.ConfigRevision = configRevision;
+			newStream.Config = config;
+			newStream.Order = config.Order ?? StreamDocument.DefaultOrder;
+			newStream.NotificationChannel = config.NotificationChannel;
+			newStream.NotificationChannelFilter = config.NotificationChannelFilter;
+			newStream.TriageChannel = config.TriageChannel;
+			newStream.DefaultPreflight = defaultPreflight;
+			newStream.Tabs = tabs;
+			newStream.AgentTypes = agentTypes;
+			newStream.WorkspaceTypes = workspaceTypes;
+			newStream.Templates = templateRefs;
+			newStream.ReplicationMode = config.ReplicationMode;
+			newStream.ReplicationFilter = config.ReplicationFilter;
+			newStream.Acl = acl;
+
+			try
+			{
+				await _streams.InsertOneAsync(newStream);
+				return newStream;
+			}
+			catch (MongoWriteException ex)
+			{
+				if (ex.WriteError.Category == ServerErrorCategory.DuplicateKey)
+				{
+					return null;
+				}
+				else
+				{
+					throw;
+				}
+			}
+		}
+
+		/// <inheritdoc/>
+		async Task<IStream?> TryReplaceAsync(IStream streamInterface, ProjectId projectId, string configRevision, StreamConfig config, DefaultPreflight? defaultPreflight, List<StreamTab> tabs, Dictionary<string, AgentType>? agentTypes, Dictionary<string, WorkspaceType>? workspaceTypes, Dictionary<TemplateRefId, TemplateRef>? templateRefs, Acl? acl)
+		{
+			int order = config.Order ?? StreamDocument.DefaultOrder;
+
+			StreamDocument stream = (StreamDocument)streamInterface;
+
+			UpdateDefinitionBuilder<StreamDocument> updateBuilder = Builders<StreamDocument>.Update;
+
+			List<UpdateDefinition<StreamDocument>> updates = new List<UpdateDefinition<StreamDocument>>();
+			updates.Add(updateBuilder.Set(x => x.Name, config.Name));
+			updates.Add(updateBuilder.Set(x => x.ProjectId, projectId));
+			updates.Add(updateBuilder.Set(x => x.ClusterName, config.ClusterName));
+			updates.Add(updateBuilder.Set(x => x.ConfigRevision, configRevision));
+			updates.Add(updateBuilder.Set(x => x.Order, order));
+			updates.Add(updateBuilder.Set(x => x.NotificationChannel, config.NotificationChannel));
+			updates.Add(updateBuilder.Set(x => x.NotificationChannelFilter, config.NotificationChannelFilter));
+			updates.Add(updateBuilder.Set(x => x.TriageChannel, config.TriageChannel));
+			updates.Add(updateBuilder.Set(x => x.DefaultPreflight, defaultPreflight));
+			updates.Add(updateBuilder.Set(x => x.Tabs, tabs ?? new List<StreamTab>()));
+			updates.Add(updateBuilder.Set(x => x.AgentTypes, agentTypes ?? new Dictionary<string, AgentType>()));
+			updates.Add(updateBuilder.Set(x => x.WorkspaceTypes, workspaceTypes ?? new Dictionary<string, WorkspaceType>()));
+			updates.Add(updateBuilder.Set(x => x.Templates, templateRefs ?? new Dictionary<TemplateRefId, TemplateRef>()));
+			updates.Add(updateBuilder.Set(x => x.ReplicationMode, config.ReplicationMode));
+			updates.Add(updateBuilder.SetOrUnsetNullRef(x => x.ReplicationFilter, config.ReplicationFilter));
+			updates.Add(updateBuilder.SetOrUnsetNullRef(x => x.Acl, acl));
+			updates.Add(updateBuilder.Unset(x => x.Deleted));
+
+			return await TryUpdateStreamAsync(stream, updateBuilder.Combine(updates));
+		}
+
+		/// <inheritdoc/>
+		public async Task<IStream?> GetAsync(StreamId streamId)
+		{
+			StreamDocument? stream = await _streams.Find<StreamDocument>(x => x.Id == streamId).FirstOrDefaultAsync();
+			if (stream != null)
+			{
+				await PostLoadAsync(stream);
+			}
+			return stream;
+		}
+
+		/// <inheritdoc/>
+		public async Task<IStreamPermissions?> GetPermissionsAsync(StreamId streamId)
+		{
+			return await _streams.Find<StreamDocument>(x => x.Id == streamId).Project<StreamPermissions>(StreamPermissions.Projection).FirstOrDefaultAsync();
+		}
+
+		/// <inheritdoc/>
+		public async Task<List<IStream>> FindAllAsync()
+		{
+			List<StreamDocument> results = await _streams.Find(Builders<StreamDocument>.Filter.Ne(x => x.Deleted, true)).ToListAsync();
+			foreach (StreamDocument result in results)
+			{
+				await PostLoadAsync(result);
+			}
+			return results.ConvertAll<IStream>(x => x);
+		}
+
+		/// <inheritdoc/>
+		public async Task<List<IStream>> FindForProjectsAsync(ProjectId[] projectIds)
+		{
+			FilterDefinition<StreamDocument> filter = Builders<StreamDocument>.Filter.In(x => x.ProjectId, projectIds) & Builders<StreamDocument>.Filter.Ne(x => x.Deleted, true);
+
+			List<StreamDocument> results = await _streams.Find(filter).ToListAsync();
+			foreach (StreamDocument result in results)
+			{
+				await PostLoadAsync(result);
+			}
+			return results.ConvertAll<IStream>(x => x);
+		}
+
+		/// <inheritdoc/>
+		public async Task<IStream?> TryUpdatePauseStateAsync(IStream streamInterface, DateTime? newPausedUntil, string? newPauseComment)
+		{
+			StreamDocument stream = (StreamDocument)streamInterface;
+
+			UpdateDefinitionBuilder<StreamDocument> updateBuilder = Builders<StreamDocument>.Update;
+
+			List<UpdateDefinition<StreamDocument>> updates = new List<UpdateDefinition<StreamDocument>>();
+			stream.PausedUntil = newPausedUntil;
+			stream.PauseComment = newPauseComment;
+			updates.Add(updateBuilder.Set(x => x.PausedUntil, newPausedUntil));
+			updates.Add(updateBuilder.Set(x => x.PauseComment, newPauseComment));
+
+			return await TryUpdateStreamAsync(stream, updateBuilder.Combine(updates));
+		}
+
+		/// <inheritdoc/>
+		public async Task<IStream?> TryUpdateScheduleTriggerAsync(IStream streamInterface, TemplateRefId templateId, DateTime? lastTriggerTimeUtc, int? lastTriggerChange, List<JobId> newActiveJobs)
+		{
+			StreamDocument stream = (StreamDocument)streamInterface;
+			TemplateRef template = stream.Templates[templateId];
+			Schedule schedule = template.Schedule!;
+
+			// Build the updates. MongoDB driver cannot parse TemplateRefId in expression tree; need to specify field name explicitly
+			List<UpdateDefinition<StreamDocument>> updates = new List<UpdateDefinition<StreamDocument>>();
+			if (lastTriggerTimeUtc.HasValue && lastTriggerTimeUtc.Value != schedule.LastTriggerTime)
+			{
+				FieldDefinition<StreamDocument, DateTimeOffset> lastTriggerTimeField = $"{nameof(stream.Templates)}.{templateId}.{nameof(template.Schedule)}.{nameof(schedule.LastTriggerTime)}";
+				updates.Add(Builders<StreamDocument>.Update.Set(lastTriggerTimeField, lastTriggerTimeUtc.Value));
+				schedule.LastTriggerTimeUtc = lastTriggerTimeUtc.Value;
+			}
+			if (lastTriggerChange.HasValue && lastTriggerChange.Value > schedule.LastTriggerChange)
+			{
+				FieldDefinition<StreamDocument, int> lastTriggerChangeField = $"{nameof(stream.Templates)}.{templateId}.{nameof(template.Schedule)}.{nameof(schedule.LastTriggerChange)}";
+				updates.Add(Builders<StreamDocument>.Update.Set(lastTriggerChangeField, lastTriggerChange.Value));
+				schedule.LastTriggerChange = lastTriggerChange.Value;
+			}
+			if (newActiveJobs != null)
+			{
+				FieldDefinition<StreamDocument, List<JobId>> field = $"{nameof(stream.Templates)}.{templateId}.{nameof(template.Schedule)}.{nameof(schedule.ActiveJobs)}";
+				updates.Add(Builders<StreamDocument>.Update.Set(field, newActiveJobs));
+				schedule.ActiveJobs = newActiveJobs;
+			}
+
+			return (updates.Count == 0)? streamInterface : await TryUpdateStreamAsync(stream, Builders<StreamDocument>.Update.Combine(updates));
+		}
+
+		/// <summary>
+		/// Update a stream
+		/// </summary>
+		/// <param name="stream">The stream to update</param>
+		/// <param name="update">The update definition</param>
+		/// <returns>The updated document, or null the update failed</returns>
+		private async Task<StreamDocument?> TryUpdateStreamAsync(StreamDocument stream, UpdateDefinition<StreamDocument> update)
+		{
+			FilterDefinition<StreamDocument> filter = Builders<StreamDocument>.Filter.Expr(x => x.Id == stream.Id && x.UpdateIndex == stream.UpdateIndex);
+			update = update.Set(x => x.UpdateIndex, stream.UpdateIndex + 1);
+
+			FindOneAndUpdateOptions<StreamDocument> options = new FindOneAndUpdateOptions<StreamDocument> { ReturnDocument = ReturnDocument.After };
+
+			StreamDocument? result = await _streams.FindOneAndUpdateAsync(filter, update, options);
+			if(result != null)
+			{
+				await PostLoadAsync(result);
+			}
+			return result;
+		}
+
+		/// <inheritdoc/>
+		public async Task DeleteAsync(StreamId streamId)
+		{
+			await _streams.UpdateOneAsync<StreamDocument>(x => x.Id == streamId, Builders<StreamDocument>.Update.Set(x => x.Deleted, true).Inc(x => x.UpdateIndex, 1));
+		}
+
+		/// <summary>
+		/// Checks the stream definition for consistency
+		/// </summary>
+		public static void Validate(StreamId streamId, DefaultPreflight? defaultPreflight, IReadOnlyDictionary<TemplateRefId, TemplateRef> templates, IReadOnlyList<StreamTab> tabs, IReadOnlyDictionary<string, AgentType> agentTypes, IReadOnlyDictionary<string, WorkspaceType> workspaceTypes)
+		{
+			// Check the default preflight template is valid
+			if (defaultPreflight != null)
+			{
+				if (defaultPreflight.TemplateRefId != null && !templates.ContainsKey(defaultPreflight.TemplateRefId.Value))
+				{
+					throw new InvalidStreamException($"Default preflight template was listed as '{defaultPreflight.TemplateRefId.Value}', but no template was found by that name");
+				}
+			}
+
+			// Check that all the templates are referenced by a tab
+			HashSet<TemplateRefId> remainingTemplates = new HashSet<TemplateRefId>(templates.Keys);
+			foreach (JobsTab jobsTab in tabs.OfType<JobsTab>())
+			{
+				if (jobsTab.Templates != null)
+				{
+					remainingTemplates.ExceptWith(jobsTab.Templates);
+				}
+			}
+			if (remainingTemplates.Count > 0)
+			{
+				throw new InvalidStreamException(String.Join("\n", remainingTemplates.Select(x => $"Template '{x}' is not listed on any tab for {streamId}")));
+			}
+
+			// Check that all the agent types reference valid workspace names
+			foreach (KeyValuePair<string, AgentType> pair in agentTypes)
+			{
+				string? workspaceTypeName = pair.Value.Workspace;
+				if (workspaceTypeName != null && !workspaceTypes.ContainsKey(workspaceTypeName))
+				{
+					throw new InvalidStreamException($"Agent type '{pair.Key}' references undefined workspace type '{pair.Value.Workspace}' in {streamId}");
+				}
+			}
+		}
+	}
+}
