Horde: Finish re-org of Horde solution started a while ago.

diff --git a/Engine/Source/Programs/Horde/Horde.Build/Streams/StreamService.cs b/Engine/Source/Programs/Horde/Horde.Build/Streams/StreamService.cs
new file mode 100644
--- /dev/null
+++ b/Engine/Source/Programs/Horde/Horde.Build/Streams/StreamService.cs
@@ -0,0 +1,333 @@
+// Copyright Epic Games, Inc. All Rights Reserved.
+
+using System;
+using System.Collections.Generic;
+using System.Linq;
+using System.Security.Claims;
+using System.Threading.Tasks;
+using Horde.Build.Acls;
+using Horde.Build.Api;
+using Horde.Build.Collections;
+using Horde.Build.Models;
+using Horde.Build.Utilities;
+using Microsoft.Extensions.Caching.Memory;
+using MongoDB.Driver;
+
+namespace Horde.Build.Services
+{
+	using JobId = ObjectId<IJob>;
+	using ProjectId = StringId<IProject>;
+	using StreamId = StringId<IStream>;
+	using TemplateRefId = StringId<TemplateRef>;
+
+	/// <summary>
+	/// Cache of information about stream ACLs
+	/// </summary>
+	public class StreamPermissionsCache : ProjectPermissionsCache
+	{
+		/// <summary>
+		/// Map of stream id to permissions for that stream
+		/// </summary>
+		public Dictionary<StreamId, IStreamPermissions?> Streams { get; set; } = new Dictionary<StreamId, IStreamPermissions?>();
+	}
+
+	/// <summary>
+	/// Wraps functionality for manipulating streams
+	/// </summary>
+	public sealed class StreamService : IDisposable
+	{
+		/// <summary>
+		/// The project service instance
+		/// </summary>
+		readonly ProjectService _projectService;
+
+		/// <summary>
+		/// Collection of stream documents
+		/// </summary>
+		readonly IStreamCollection _streams;
+
+		/// <summary>
+		/// Cache of stream documents
+		/// </summary>
+		readonly MemoryCache _streamCache = new MemoryCache(new MemoryCacheOptions());
+
+		/// <summary>
+		/// Accessor for the stream collection
+		/// </summary>
+		public IStreamCollection StreamCollection => _streams;
+
+		/// <summary>
+		/// Constructor
+		/// </summary>
+		/// <param name="projectService">The project service instance</param>
+		/// <param name="streams">Collection of stream documents</param>
+		public StreamService(ProjectService projectService, IStreamCollection streams)
+		{
+			_projectService = projectService;
+			_streams = streams;
+		}
+
+		/// <summary>
+		/// Dispose of any managed resources
+		/// </summary>
+		public void Dispose()
+		{
+			_streamCache.Dispose();
+		}
+
+		/// <summary>
+		/// Deletes an existing stream
+		/// </summary>
+		/// <param name="streamId">Unique id of the stream</param>
+		/// <returns>Async task object</returns>
+		public async Task DeleteStreamAsync(StreamId streamId)
+		{
+			await _streams.DeleteAsync(streamId);
+		}
+
+		/// <summary>
+		/// Updates an existing stream
+		/// </summary>
+		/// <param name="stream">The stream to update</param>
+		/// <param name="newPausedUntil">The new datetime for pausing builds</param>
+		/// <param name="newPauseComment">The reason for pausing</param>
+		/// <returns>Async task object</returns>
+		public async Task<IStream?> UpdatePauseStateAsync(IStream? stream, DateTime? newPausedUntil = null, string? newPauseComment = null)
+		{
+			for (; stream != null; stream = await GetStreamAsync(stream.Id))
+			{
+				IStream? newStream = await _streams.TryUpdatePauseStateAsync(stream, newPausedUntil, newPauseComment);
+				if (newStream != null)
+				{
+					return newStream;
+				}
+			}
+			return null;
+		}
+
+		/// <summary>
+		/// Attempts to update the last trigger time for a schedule
+		/// </summary>
+		/// <param name="stream">The stream to update</param>
+		/// <param name="templateRefId">The template ref id</param>
+		/// <param name="lastTriggerTimeUtc"></param>
+		/// <param name="lastTriggerChange"></param>
+		/// <param name="addJobs">Jobs to add</param>
+		/// <param name="removeJobs">Jobs to remove</param>
+		/// <returns>True if the stream was updated</returns>
+		public async Task<IStream?> UpdateScheduleTriggerAsync(IStream stream, TemplateRefId templateRefId, DateTime? lastTriggerTimeUtc = null, int? lastTriggerChange = null, List<JobId>? addJobs = null, List<JobId>? removeJobs = null)
+		{
+			IStream? newStream = stream;
+			while (newStream != null)
+			{
+				TemplateRef? templateRef;
+				if (!newStream.Templates.TryGetValue(templateRefId, out templateRef))
+				{
+					break;
+				}
+				if (templateRef.Schedule == null)
+				{
+					break;
+				}
+
+				IEnumerable<JobId> newActiveJobs = templateRef.Schedule.ActiveJobs;
+				if (removeJobs != null)
+				{
+					newActiveJobs = newActiveJobs.Except(removeJobs);
+				}
+				if (addJobs != null)
+				{
+					newActiveJobs = newActiveJobs.Union(addJobs);
+				}
+
+				newStream = await _streams.TryUpdateScheduleTriggerAsync(newStream, templateRefId, lastTriggerTimeUtc, lastTriggerChange, newActiveJobs.ToList());
+
+				if (newStream != null)
+				{
+					return newStream;
+				}
+
+				newStream = await _streams.GetAsync(stream.Id);
+			}
+			return null;
+		}
+
+		/// <summary>
+		/// Attempts to update a stream template ref
+		/// </summary>
+		/// <param name="streamInterface"></param>
+		/// <param name="templateRefId"></param>
+		/// <param name="stepStates"></param>
+		/// <returns></returns>
+		public Task<IStream?> TryUpdateTemplateRefAsync(IStream streamInterface, TemplateRefId templateRefId, List<UpdateStepStateRequest>? stepStates = null)
+		{
+			return _streams.TryUpdateTemplateRefAsync(streamInterface, templateRefId, stepStates );
+		}
+
+		/// <summary>
+		/// Gets all the available streams for a project
+		/// </summary>
+		/// <returns>List of stream documents</returns>
+		public Task<List<IStream>> GetStreamsAsync()
+		{
+			return _streams.FindAllAsync();
+		}
+
+		/// <summary>
+		/// Gets all the available streams for a project
+		/// </summary>
+		/// <param name="projectId">Unique id of the project to query</param>
+		/// <returns>List of stream documents</returns>
+		public Task<List<IStream>> GetStreamsAsync(ProjectId projectId)
+		{
+			return _streams.FindForProjectsAsync(new[] { projectId });
+		}
+
+		/// <summary>
+		/// Gets all the available streams for a project
+		/// </summary>
+		/// <param name="projectIds">Unique id of the project to query</param>
+		/// <returns>List of stream documents</returns>
+		public Task<List<IStream>> GetStreamsAsync(ProjectId[] projectIds)
+		{
+			if (projectIds.Length == 0)
+			{
+				return _streams.FindAllAsync();
+			}
+			else
+			{
+				return _streams.FindForProjectsAsync(projectIds);
+			}
+		}
+
+		/// <summary>
+		/// Gets a stream by ID
+		/// </summary>
+		/// <param name="streamId">Unique id of the stream</param>
+		/// <returns>The stream document</returns>
+		public Task<IStream?> GetStreamAsync(StreamId streamId)
+		{
+			return _streams.GetAsync(streamId);
+		}
+
+		/// <summary>
+		/// Adds a cached stream interface
+		/// </summary>
+		/// <param name="stream">The stream to add</param>
+		/// <returns>The new stream</returns>
+		private void AddCachedStream(IStream stream)
+		{
+			MemoryCacheEntryOptions options = new MemoryCacheEntryOptions().SetAbsoluteExpiration(TimeSpan.FromMinutes(1.0));
+			_streamCache.Set(stream.Id, stream, options);
+		}
+
+		/// <summary>
+		/// Gets a cached stream interface
+		/// </summary>
+		/// <param name="streamId">Unique id for the stream</param>
+		/// <returns>The new stream</returns>
+		public async Task<IStream?> GetCachedStream(StreamId streamId)
+		{
+			object? stream;
+			if (!_streamCache.TryGetValue(streamId, out stream))
+			{
+				stream = await GetStreamAsync(streamId);
+				if (stream != null)
+				{
+					AddCachedStream((IStream)stream);
+				}
+			}
+			return (IStream?)stream;
+		}
+
+		/// <summary>
+		/// Gets a stream's permissions by ID
+		/// </summary>
+		/// <param name="streamId">Unique id of the stream</param>
+		/// <returns>The stream document</returns>
+		public Task<IStreamPermissions?> GetStreamPermissionsAsync(StreamId streamId)
+		{
+			return _streams.GetPermissionsAsync(streamId);
+		}
+
+		/// <summary>
+		/// Determines if the user is authorized to perform an action on a particular stream
+		/// </summary>
+		/// <param name="acl">ACL for the stream to check</param>
+		/// <param name="projectId">The parent project id</param>
+		/// <param name="action">The action being performed</param>
+		/// <param name="user">The principal to authorize</param>
+		/// <param name="cache">Cache of project permissions</param>
+		/// <returns>True if the action is authorized</returns>
+		private Task<bool> AuthorizeAsync(Acl? acl, ProjectId projectId, AclAction action, ClaimsPrincipal user, ProjectPermissionsCache? cache)
+		{
+			bool? result = acl?.Authorize(action, user);
+			if (result == null)
+			{
+				return _projectService.AuthorizeAsync(projectId, action, user, cache);
+			}
+			else
+			{
+				return Task.FromResult(result.Value);
+			}
+		}
+
+		/// <summary>
+		/// Determines if the user is authorized to perform an action on a particular stream
+		/// </summary>
+		/// <param name="stream">The stream to check</param>
+		/// <param name="action">The action being performed</param>
+		/// <param name="user">The principal to authorize</param>
+		/// <param name="cache">Cache of project permissions</param>
+		/// <returns>True if the action is authorized</returns>
+		public Task<bool> AuthorizeAsync(IStream stream, AclAction action, ClaimsPrincipal user, ProjectPermissionsCache? cache)
+		{
+			return AuthorizeAsync(stream.Acl, stream.ProjectId, action, user, cache);
+		}
+
+		/// <summary>
+		/// Determines if the user is authorized to perform an action on a particular stream
+		/// </summary>
+		/// <param name="stream">The stream to check</param>
+		/// <param name="template">Template within the stream to check</param>
+		/// <param name="action">The action being performed</param>
+		/// <param name="user">The principal to authorize</param>
+		/// <param name="cache">Cache of project permissions</param>
+		/// <returns>True if the action is authorized</returns>
+		public Task<bool> AuthorizeAsync(IStream stream, TemplateRef template, AclAction action, ClaimsPrincipal user, ProjectPermissionsCache? cache)
+		{
+			bool? result = template.Acl?.Authorize(action, user);
+			if (result == null)
+			{
+				return AuthorizeAsync(stream, action, user, cache);
+			}
+			else
+			{
+				return Task.FromResult(result.Value);
+			}
+		}
+
+		/// <summary>
+		/// Determines if the user is authorized to perform an action on a particular project
+		/// </summary>
+		/// <param name="streamId">The stream id to check</param>
+		/// <param name="action">The action being performed</param>
+		/// <param name="user">The principal to authorize</param>
+		/// <param name="cache">Cache of stream permissions</param>
+		/// <returns>True if the action is authorized</returns>
+		public async Task<bool> AuthorizeAsync(StreamId streamId, AclAction action, ClaimsPrincipal user, StreamPermissionsCache? cache)
+		{
+			IStreamPermissions? permissions;
+			if (cache == null)
+			{
+				permissions = await GetStreamPermissionsAsync(streamId);
+			}
+			else if (!cache.Streams.TryGetValue(streamId, out permissions))
+			{
+				permissions = await GetStreamPermissionsAsync(streamId);
+				cache.Streams.Add(streamId, permissions);
+			}
+			return permissions != null && await AuthorizeAsync(permissions.Acl, permissions.ProjectId, action, user, cache);
+		}
+	}
+}
diff --git a/Engine/Source/Programs/Horde/Horde.Build/Streams/StreamTab.cs b/Engine/Source/Programs/Horde/Horde.Build/Streams/StreamTab.cs
new file mode 100644
--- /dev/null
+++ b/Engine/Source/Programs/Horde/Horde.Build/Streams/StreamTab.cs
@@ -0,0 +1,267 @@
+// Copyright Epic Games, Inc. All Rights Reserved.
+
+using System;
+using System.Collections.Generic;
+using Horde.Build.Api;
+using Horde.Build.Utilities;
+using MongoDB.Bson.Serialization.Attributes;
+
+namespace Horde.Build.Models
+{
+	using TemplateRefId = StringId<TemplateRef>;
+
+	/// <summary>
+	/// Information about a page to display in the dashboard for a stream
+	/// </summary>
+	[BsonKnownTypes(typeof(JobsTab))]
+	public abstract class StreamTab
+	{
+		/// <summary>
+		/// Title of this page
+		/// </summary>
+		[BsonRequired]
+		public string Title { get; set; }
+
+		/// <summary>
+		/// Private constructor for serialization
+		/// </summary>
+		private StreamTab()
+		{
+			Title = null!;
+		}
+
+		/// <summary>
+		/// Constructor
+		/// </summary>
+		/// <param name="title">Title of this page</param>
+		protected StreamTab(string title)
+		{
+			Title = title;
+		}
+
+		/// <summary>
+		/// Constructor
+		/// </summary>
+		/// <param name="request">Public request object</param>
+		protected StreamTab(CreateStreamTabRequest request)
+		{
+			Title = request.Title;
+		}
+
+		/// <summary>
+		/// Creates an instance from a request
+		/// </summary>
+		/// <param name="request">The request object</param>
+		/// <returns>New tab instance</returns>
+		public static StreamTab FromRequest(CreateStreamTabRequest request)
+		{
+			CreateJobsTabRequest? jobsRequest = request as CreateJobsTabRequest;
+			if (jobsRequest != null)
+			{
+				return new JobsTab(jobsRequest);
+			}
+
+			throw new Exception($"Unknown tab request type '{request.GetType()}'");
+		}
+
+		/// <summary>
+		/// Creates a response object
+		/// </summary>
+		/// <returns>Response object</returns>
+		public abstract GetStreamTabResponse ToResponse();
+	}
+
+	/// <summary>
+	/// Describes a column to display on the jobs page
+	/// </summary>
+	[BsonKnownTypes(typeof(JobsTabLabelColumn), typeof(JobsTabParameterColumn))]
+	public abstract class JobsTabColumn
+	{
+		/// <summary>
+		/// Heading for this column
+		/// </summary>
+		public string Heading { get; set; }
+
+		/// <summary>
+		/// Relative width of this column.
+		/// </summary>
+		public int? RelativeWidth { get; set; }
+
+		/// <summary>
+		/// Private constructor for serialization
+		/// </summary>
+		protected JobsTabColumn()
+		{
+			Heading = null!;
+			RelativeWidth = 1;
+		}
+
+		/// <summary>
+		/// Constructor
+		/// </summary>
+		/// <param name="heading">Heading for this column</param>
+		/// <param name="relativeWidth">Relative width of this column.</param>
+		protected JobsTabColumn(string heading, int? relativeWidth)
+		{
+			Heading = heading;
+			RelativeWidth = relativeWidth;
+		}
+
+		/// <summary>
+		/// Constructor
+		/// </summary>
+		/// <param name="request">Public request object</param>
+		protected JobsTabColumn(CreateJobsTabColumnRequest request)
+		{
+			Heading = request.Heading;
+			RelativeWidth = request.RelativeWidth;
+		}
+
+		/// <summary>
+		/// Converts this object to a JSON response object
+		/// </summary>
+		/// <returns>Response object</returns>
+		public abstract GetJobsTabColumnResponse ToResponse();
+	}
+
+	/// <summary>
+	/// Column that contains a set of labels
+	/// </summary>
+	public class JobsTabLabelColumn : JobsTabColumn
+	{
+		/// <summary>
+		/// Category of labels to display in this column. If null, includes any label not matched by another column.
+		/// </summary>
+		public string? Category { get; set; }
+
+		/// <summary>
+		/// Private constructor for serialization
+		/// </summary>
+		protected JobsTabLabelColumn()
+		{
+		}
+
+		/// <summary>
+		/// Constructor
+		/// </summary>
+		/// <param name="heading">Heading for this column</param>
+		/// <param name="category">Category of labels to display in this column</param>
+		/// <param name="relativeWidth">Relative width of this column.</param>
+		public JobsTabLabelColumn(string heading, string? category, int? relativeWidth)
+			: base(heading, relativeWidth)
+		{
+			Category = category;
+		}
+
+		/// <summary>
+		/// Converts this object to a JSON response object
+		/// </summary>
+		/// <returns>Response object</returns>
+		public override GetJobsTabColumnResponse ToResponse()
+		{
+			return new GetJobsTabLabelColumnResponse(Heading, Category, RelativeWidth);
+		}
+	}
+
+	/// <summary>
+	/// Include the value for a parameter on the jobs tab
+	/// </summary>
+	public class JobsTabParameterColumn : JobsTabColumn
+	{
+		/// <summary>
+		/// Name of a parameter to show in this column. Should be in the form of a prefix, eg. "-set:Foo="
+		/// </summary>
+		public string? Parameter { get; set; }
+
+		/// <summary>
+		/// Private constructor for serialization
+		/// </summary>
+		protected JobsTabParameterColumn()
+		{
+		}
+
+		/// <summary>
+		/// Constructor
+		/// </summary>
+		/// <param name="heading">Heading for this column</param>
+		/// <param name="parameter">Category of labels to display in this column</param>
+		/// <param name="relativeWidth">Relative width of this column.</param>
+		public JobsTabParameterColumn(string heading, string? parameter, int? relativeWidth)
+			: base(heading, relativeWidth)
+		{
+			Parameter = parameter;
+		}
+
+		/// <summary>
+		/// Converts this object to a JSON response object
+		/// </summary>
+		/// <returns>Response object</returns>
+		public override GetJobsTabColumnResponse ToResponse()
+		{
+			return new GetJobsTabParameterColumnResponse(Heading, Parameter, RelativeWidth);
+		}
+	}
+
+	/// <summary>
+	/// Describes a job page
+	/// </summary>
+	[BsonDiscriminator("Jobs")]
+	public class JobsTab : StreamTab
+	{
+		/// <summary>
+		/// Whether to show job names
+		/// </summary>
+		public bool ShowNames { get; set; }
+
+		/// <summary>
+		/// Names of jobs to include on this page. If there is only one name specified, the name column does not need to be displayed.
+		/// </summary>
+		public List<string>? JobNames { get; set; }
+
+		/// <summary>
+		/// List of templates to display on this tab.
+		/// </summary>
+		public List<TemplateRefId>? Templates { get; set; }
+
+		/// <summary>
+		/// Columns to display in the table
+		/// </summary>
+		public List<JobsTabColumn>? Columns { get; set; }
+
+		/// <summary>
+		/// Constructor
+		/// </summary>
+		/// <param name="title">Title for this page</param>
+		/// <param name="showNames">Show names of jobs on the page</param>
+		/// <param name="templates">Template names to include on this page</param>
+		/// <param name="jobNames">Names of jobs to include on this page</param>
+		/// <param name="columns">Columns ot display in the table</param>
+		public JobsTab(string title, bool showNames, List<TemplateRefId>? templates, List<string>? jobNames, List<JobsTabColumn>? columns)
+			: base(title)
+		{
+			ShowNames = showNames;
+			Templates = templates;
+			JobNames = jobNames;
+			Columns = columns;
+		}
+
+		/// <summary>
+		/// Constructor
+		/// </summary>
+		/// <param name="request">Public request object</param>
+		public JobsTab(CreateJobsTabRequest request)
+			: base(request)
+		{
+			ShowNames = request.ShowNames;
+			Templates = request.Templates?.ConvertAll(x => new TemplateRefId(x));
+			JobNames = request.JobNames;
+			Columns = request.Columns?.ConvertAll(x => x.ToModel());
+		}
+
+		/// <inheritdoc/>
+		public override GetStreamTabResponse ToResponse()
+		{
+			return new GetJobsTabResponse(Title, ShowNames, Templates?.ConvertAll(x => x.ToString()), JobNames, Columns?.ConvertAll(x => x.ToResponse()));
+		}
+	}
+}
diff --git a/Engine/Source/Programs/Horde/Horde.Build/Streams/StreamsController.Messages.cs b/Engine/Source/Programs/Horde/Horde.Build/Streams/StreamsController.Messages.cs
new file mode 100644
--- /dev/null
+++ b/Engine/Source/Programs/Horde/Horde.Build/Streams/StreamsController.Messages.cs
@@ -0,0 +1,946 @@
+// Copyright Epic Games, Inc. All Rights Reserved.
+
+using System;
+using System.Collections.Generic;
+using System.ComponentModel.DataAnnotations;
+using EpicGames.Core;
+using Horde.Build.Models;
+using Horde.Build.Server;
+using Horde.Build.Utilities;
+using HordeCommon;
+
+namespace Horde.Build.Api
+{
+	using TemplateRefId = StringId<TemplateRef>;
+	using WorkflowId = StringId<WorkflowConfig>;
+
+	/// <summary>
+	/// Query selecting the base changelist to use
+	/// </summary>
+	public class ChangeQueryRequest
+	{
+		/// <summary>
+		/// The template id to query
+		/// </summary>
+		public string? TemplateId { get; set; }
+
+		/// <summary>
+		/// The target to query
+		/// </summary>
+		public string? Target { get; set; }
+
+		/// <summary>
+		/// Whether to match a job that produced warnings
+		/// </summary>
+		public List<JobStepOutcome>? Outcomes { get; set; }
+
+		/// <summary>
+		/// Convert to a model object
+		/// </summary>
+		/// <returns></returns>
+		public ChangeQuery ToModel()
+		{
+			ChangeQuery query = new ChangeQuery();
+			if (TemplateId != null)
+			{
+				query.TemplateRefId = new StringId<TemplateRef>(TemplateId);
+			}
+			query.Target = Target;
+			query.Outcomes = Outcomes;
+			return query;
+		}
+	}
+
+	/// <summary>
+	/// Specifies defaults for running a preflight
+	/// </summary>
+	public class DefaultPreflightRequest
+	{
+		/// <summary>
+		/// The template id to query
+		/// </summary>
+		public string? TemplateId { get; set; }
+
+		/// <summary>
+		/// The last successful job type to use for the base changelist
+		/// </summary>
+		[Obsolete("Use Change.TemplateId instead")]
+		public string? ChangeTemplateId { get; set; }
+
+		/// <summary>
+		/// Query for the change to use
+		/// </summary>
+		public ChangeQueryRequest? Change { get; set; }
+
+		/// <summary>
+		/// Convert to a model object
+		/// </summary>
+		/// <returns></returns>
+		public DefaultPreflight ToModel()
+		{
+#pragma warning disable CS0618 // Type or member is obsolete
+			if (ChangeTemplateId != null)
+			{
+				Change ??= new ChangeQueryRequest();
+				Change.TemplateId = ChangeTemplateId;
+			}
+			return new DefaultPreflight((TemplateId != null) ? (TemplateRefId?)new TemplateRefId(TemplateId) : null, Change?.ToModel());
+#pragma warning restore CS0618 // Type or member is obsolete
+		}
+	}
+
+	/// <summary>
+	/// Mapping from a BuildGraph agent type to a set of machines on the farm
+	/// </summary>
+	public class CreateAgentTypeRequest
+	{
+		/// <summary>
+		/// Pool of agents to use for this agent type
+		/// </summary>
+		[Required]
+		public string Pool { get; set; } = null!;
+
+		/// <summary>
+		/// Name of the workspace to sync
+		/// </summary>
+		public string? Workspace { get; set; }
+
+		/// <summary>
+		/// Path to the temporary storage dir
+		/// </summary>
+		public string? TempStorageDir { get; set; }
+
+		/// <summary>
+		/// Environment variables to be set when executing the job
+		/// </summary>
+		public Dictionary<string, string>? Environment { get; set; }
+	}
+
+	/// <summary>
+	/// Information about a workspace type
+	/// </summary>
+	public class CreateWorkspaceTypeRequest
+	{
+		/// <summary>
+		/// Name of the Perforce server cluster to use
+		/// </summary>
+		public string? Cluster { get; set; }
+
+		/// <summary>
+		/// The Perforce server and port (eg. perforce:1666)
+		/// </summary>
+		public string? ServerAndPort { get; set; }
+
+		/// <summary>
+		/// User to log into Perforce with (defaults to buildmachine)
+		/// </summary>
+		public string? UserName { get; set; }
+
+		/// <summary>
+		/// Password to use to log into the workspace
+		/// </summary>
+		public string? Password { get; set; }
+
+		/// <summary>
+		/// Identifier to distinguish this workspace from other workspaces. Defaults to the workspace type name.
+		/// </summary>
+		public string? Identifier { get; set; }
+
+		/// <summary>
+		/// Override for the stream to sync
+		/// </summary>
+		public string? Stream { get; set; }
+
+		/// <summary>
+		/// Custom view for the workspace
+		/// </summary>
+		public List<string>? View { get; set; }
+
+		/// <summary>
+		/// Whether to use an incrementally synced workspace
+		/// </summary>
+		public bool Incremental { get; set; }
+
+		/// <summary>
+		/// Whether to use the AutoSDK
+		/// </summary>
+		public bool UseAutoSdk { get; set; } = true;
+	}
+
+	/// <summary>
+	/// Trigger for another template
+	/// </summary>
+	public class CreateChainedJobTemplateRequest
+	{
+		/// <summary>
+		/// Name of the target that needs to complete before starting the other template
+		/// </summary>
+		[Required]
+		public string Trigger { get; set; } = String.Empty;
+
+		/// <summary>
+		/// Id of the template to trigger
+		/// </summary>
+		[Required]
+		public string TemplateId { get; set; } = String.Empty;
+	}
+
+	/// <summary>
+	/// Parameters to create a template within a stream
+	/// </summary>
+	public class TemplateRefConfig : TemplateConfig
+	{
+		TemplateRefId _id;
+
+		/// <summary>
+		/// Optional identifier for this ref. If not specified, an id will be generated from the name.
+		/// </summary>
+		public TemplateRefId Id
+		{
+			get => _id.IsEmpty ? TemplateRefId.Sanitize(Name) : _id;
+			set => _id = value;
+		}
+
+		/// <summary>
+		/// Whether to show badges in UGS for these jobs
+		/// </summary>
+		public bool ShowUgsBadges { get; set; }
+
+		/// <summary>
+		/// Whether to show alerts in UGS for these jobs
+		/// </summary>
+		public bool ShowUgsAlerts { get; set; }
+
+		/// <summary>
+		/// Notification channel for this template. Overrides the stream channel if set.
+		/// </summary>
+		public string? NotificationChannel { get; set; }
+
+		/// <summary>
+		/// Notification channel filter for this template. Can be Success|Failure|Warnings
+		/// </summary>
+		public string? NotificationChannelFilter { get; set; }
+
+		/// <summary>
+		/// Triage channel for this template. Overrides the stream channel if set.
+		/// </summary>
+		public string? TriageChannel { get; set; }
+
+		/// <summary>
+		/// Workflow to user for this stream
+		/// </summary>
+		public WorkflowId? WorkflowId
+		{
+			get => Annotations.WorkflowId;
+			set => Annotations.WorkflowId = value;
+		}
+
+		/// <summary>
+		/// Default annotations to apply to nodes in this template
+		/// </summary>
+		public NodeAnnotations Annotations { get; set; } = new NodeAnnotations();
+
+		/// <summary>
+		/// Schedule to execute this template
+		/// </summary>
+		public CreateScheduleRequest? Schedule { get; set; }
+
+		/// <summary>
+		/// List of chained job triggers
+		/// </summary>
+		public List<CreateChainedJobTemplateRequest>? ChainedJobs { get; set; }
+
+		/// <summary>
+		/// The ACL for this template
+		/// </summary>
+		public UpdateAclRequest? Acl { get; set; }
+	}
+
+	/// <summary>
+	/// Step state update request
+	/// </summary>
+	public class UpdateStepStateRequest
+	{
+		/// <summary>
+		/// Name of the step
+		/// </summary>
+		[Required]
+		public string Name { get; set; } = String.Empty;
+
+		/// <summary>
+		/// User who paused the step
+		/// </summary>
+		public string? PausedByUserId { get; set; }
+	}
+
+
+	/// <summary>
+	/// Updates an existing stream template ref
+	/// </summary>
+	public class UpdateTemplateRefRequest
+	{
+		
+		/// <summary>
+		/// Step states to update
+		/// </summary>
+		public List<UpdateStepStateRequest>? StepStates { get; set; }
+
+	}
+
+
+	/// <summary>
+	/// Mapping from a BuildGraph agent type to a set of machines on the farm
+	/// </summary>
+	public class GetAgentTypeResponse
+	{
+		/// <summary>
+		/// Pool of agents to use for this agent type
+		/// </summary>
+		public string Pool { get; set; }
+
+		/// <summary>
+		/// Name of the workspace to sync
+		/// </summary>
+		public string? Workspace { get; set; }
+
+		/// <summary>
+		/// Path to the temporary storage dir
+		/// </summary>
+		public string? TempStorageDir { get; set; }
+
+		/// <summary>
+		/// Environment variables to be set when executing the job
+		/// </summary>
+		public Dictionary<string, string>? Environment { get; set; }
+
+		/// <summary>
+		/// Constructor
+		/// </summary>
+		/// <param name="pool">Pool of agents to use for this agent type</param>
+		/// <param name="workspace">Name of the workspace to sync</param>
+		/// <param name="tempStorageDir">Path to the temp storage directory</param>
+		/// <param name="environment">Environment variables to be set when executing this job</param>
+		public GetAgentTypeResponse(string pool, string? workspace, string? tempStorageDir, Dictionary<string, string>? environment)
+		{
+			Pool = pool;
+			Workspace = workspace;
+			TempStorageDir = tempStorageDir;
+			Environment = environment;
+		}
+	}
+
+	/// <summary>
+	/// Information about a workspace type
+	/// </summary>
+	public class GetWorkspaceTypeResponse
+	{
+		/// <summary>
+		/// The Perforce server cluster
+		/// </summary>
+		public string? Cluster { get; set; }
+
+		/// <summary>
+		/// The Perforce server and port (eg. perforce:1666)
+		/// </summary>
+		public string? ServerAndPort { get; set; }
+
+		/// <summary>
+		/// User to log into Perforce with (defaults to buildmachine)
+		/// </summary>
+		public string? UserName { get; set; }
+
+		/// <summary>
+		/// Identifier to distinguish this workspace from other workspaces. Defaults to the workspace type name.
+		/// </summary>
+		public string? Identifier { get; set; }
+
+		/// <summary>
+		/// Override for the stream to sync
+		/// </summary>
+		public string? Stream { get; set; }
+
+		/// <summary>
+		/// Custom view for the workspace
+		/// </summary>
+		public List<string>? View { get; set; }
+
+		/// <summary>
+		/// Whether to use an incrementally synced workspace
+		/// </summary>
+		public bool Incremental { get; set; }
+
+		/// <summary>
+		/// Whether to use the AutoSDK
+		/// </summary>
+		public bool UseAutoSdk { get; set; }
+
+		/// <summary>
+		/// Default constructor
+		/// </summary>
+		public GetWorkspaceTypeResponse()
+		{
+		}
+
+		/// <summary>
+		/// Constructor
+		/// </summary>
+		/// <param name="cluster">The server cluster</param>
+		/// <param name="serverAndPort">The perforce server</param>
+		/// <param name="userName">The perforce user name</param>
+		/// <param name="identifier">Identifier to distinguish this workspace from other workspaces. Defaults to the workspace type name.</param>
+		/// <param name="stream">Override for the stream to sync</param>
+		/// <param name="view">Custom view for the workspace</param>
+		/// <param name="bIncremental">Whether to use an incrementally synced workspace</param>
+		/// <param name="bUseAutoSdk">Whether to use the AutoSDK</param>
+		public GetWorkspaceTypeResponse(string? cluster, string? serverAndPort, string? userName, string? identifier, string? stream, List<string>? view, bool bIncremental, bool bUseAutoSdk)
+		{
+			Cluster = cluster;
+			ServerAndPort = serverAndPort;
+			UserName = userName;
+			Identifier = identifier;
+			Stream = stream;
+			View = view;
+			Incremental = bIncremental;
+			UseAutoSdk = bUseAutoSdk;
+		}
+	}
+
+	/// <summary>
+	/// Trigger for another template
+	/// </summary>
+	public class GetChainedJobTemplateResponse
+	{
+		/// <summary>
+		/// Name of the target that needs to complete before starting the other template
+		/// </summary>
+		public string Trigger { get; set; }
+
+		/// <summary>
+		/// Id of the template to trigger
+		/// </summary>
+		public string TemplateId { get; set; }
+
+		/// <summary>
+		/// Constructor
+		/// </summary>
+		/// <param name="trigger">The trigger definition</param>
+		public GetChainedJobTemplateResponse(ChainedJobTemplate trigger)
+		{
+			Trigger = trigger.Trigger;
+			TemplateId = trigger.TemplateRefId.ToString();
+		}
+	}
+
+	/// <summary>
+	/// State information for a step in the stream
+	/// </summary>
+	public class GetTemplateStepStateResponse
+	{
+		/// <summary>
+		/// Name of the step
+		/// </summary>
+		public string Name { get; set; } = String.Empty;
+
+		/// <summary>
+		/// User who paused the step
+		/// </summary>
+		public GetThinUserInfoResponse? PausedByUserInfo { get; set; }
+
+		/// <summary>
+		/// The UTC time when the step was paused
+		/// </summary>
+		public DateTime? PauseTimeUtc { get; set; }
+
+		/// <summary>
+		/// Default constructor for serialization
+		/// </summary>
+		private GetTemplateStepStateResponse()
+		{
+		}
+
+		/// <summary>
+		/// Constructor
+		/// </summary>
+		public GetTemplateStepStateResponse(TemplateStepState state, GetThinUserInfoResponse? pausedByUserInfo)
+		{
+			Name = state.Name;			
+			PauseTimeUtc = state.PauseTimeUtc;
+			PausedByUserInfo = pausedByUserInfo;
+		}
+
+	}
+
+
+	/// <summary>
+	/// Information about a template in this stream
+	/// </summary>
+	public class GetTemplateRefResponse : GetTemplateResponseBase
+	{
+		/// <summary>
+		/// Id of the template ref
+		/// </summary>
+		public string Id { get; set; }
+
+		/// <summary>
+		/// Hash of the template definition
+		/// </summary>
+		public string Hash { get; set; }
+
+		/// <summary>
+		/// Whether to show badges in UGS for these jobs
+		/// </summary>
+		public bool ShowUgsBadges { get; set; }
+
+		/// <summary>
+		/// Whether to show alerts in UGS for these jobs
+		/// </summary>
+		public bool ShowUgsAlerts { get; set; }
+
+		/// <summary>
+		/// Notification channel for this template. Overrides the stream channel if set.
+		/// </summary>
+		public string? NotificationChannel { get; set; }
+
+		/// <summary>
+		/// Notification channel filter for this template. Can be Success|Failure|Warnings
+		/// </summary>
+		public string? NotificationChannelFilter { get; set; }
+
+		/// <summary>
+		/// Triage channel for this template. Overrides the stream channel if set.
+		/// </summary>
+		public string? TriageChannel { get; set; }
+
+		/// <summary>
+		/// The schedule for this ref
+		/// </summary>
+		public GetScheduleResponse? Schedule { get; set; }
+
+		/// <summary>
+		/// List of templates to trigger
+		/// </summary>
+		public List<GetChainedJobTemplateResponse>? ChainedJobs { get; set; }
+
+		/// <summary>
+		/// List of step states
+		/// </summary>
+		public List<GetTemplateStepStateResponse>? StepStates { get; set; }
+
+		/// <summary>
+		/// ACL for this template
+		/// </summary>
+		public GetAclResponse? Acl { get; set; }
+
+		/// <summary>
+		/// Constructor
+		/// </summary>
+		/// <param name="id">The template ref id</param>
+		/// <param name="templateRef">The template ref</param>
+		/// <param name="template">The template instance</param>
+		/// <param name="stepStates">The template step states</param>
+		/// <param name="bIncludeAcl">Whether to include the ACL in the response</param>
+		public GetTemplateRefResponse(TemplateRefId id, TemplateRef templateRef, ITemplate template, List<GetTemplateStepStateResponse>? stepStates, bool bIncludeAcl)
+			: base(template)
+		{
+			Id = id.ToString();
+			Hash = templateRef.Hash.ToString();
+			ShowUgsBadges = templateRef.ShowUgsBadges;
+			ShowUgsAlerts = templateRef.ShowUgsAlerts;
+			NotificationChannel = templateRef.NotificationChannel;
+			NotificationChannelFilter = templateRef.NotificationChannelFilter;
+			Schedule = (templateRef.Schedule != null) ? new GetScheduleResponse(templateRef.Schedule) : null;
+			ChainedJobs = (templateRef.ChainedJobs != null && templateRef.ChainedJobs.Count > 0) ? templateRef.ChainedJobs.ConvertAll(x => new GetChainedJobTemplateResponse(x)) : null;
+			StepStates = stepStates;
+			Acl = (bIncludeAcl && templateRef.Acl != null)? new GetAclResponse(templateRef.Acl) : null;
+		}
+	}
+
+	/// <summary>
+	/// Response describing a stream
+	/// </summary>
+	public class GetStreamResponse
+	{
+		/// <summary>
+		/// Unique id of the stream
+		/// </summary>
+		public string Id { get; set; }
+
+		/// <summary>
+		/// Unique id of the project containing this stream
+		/// </summary>
+		public string ProjectId { get; set; }
+
+		/// <summary>
+		/// Name of the stream
+		/// </summary>
+		public string Name { get; set; }
+
+		/// <summary>
+		/// The config file path on the server
+		/// </summary>
+		public string ConfigPath { get; set; } = String.Empty;
+
+		/// <summary>
+		/// Revision of the config file 
+		/// </summary>
+		public string ConfigRevision { get; set; } = String.Empty;
+
+		/// <summary>
+		/// Order to display in the list
+		/// </summary>
+		public int Order { get; set; }
+
+		/// <summary>
+		/// Notification channel for all jobs in this stream
+		/// </summary>
+		public string? NotificationChannel { get; set; }
+
+		/// <summary>
+		/// Notification channel filter for this template. Can be Success|Failure|Warnings
+		/// </summary>
+		public string? NotificationChannelFilter { get; set; }
+
+		/// <summary>
+		/// Channel to post issue triage notifications
+		/// </summary>
+		public string? TriageChannel { get; set; }
+
+		/// <summary>
+		/// Default template for running preflights
+		/// </summary>
+		public string? DefaultPreflightTemplate { get; set; }
+
+		/// <summary>
+		/// Default template to use for preflights
+		/// </summary>
+		public DefaultPreflightRequest? DefaultPreflight { get; set; }
+
+		/// <summary>
+		/// List of tabs to display for this stream
+		/// </summary>
+		public List<GetStreamTabResponse> Tabs { get; set; }
+
+		/// <summary>
+		/// Map of agent name to type
+		/// </summary>
+		public Dictionary<string, GetAgentTypeResponse> AgentTypes { get; set; }
+
+		/// <summary>
+		/// Map of workspace name to type
+		/// </summary>
+		public Dictionary<string, GetWorkspaceTypeResponse>? WorkspaceTypes { get; set; }
+
+		/// <summary>
+		/// Templates for jobs in this stream
+		/// </summary>
+		public List<GetTemplateRefResponse> Templates { get; set; }
+
+		/// <summary>
+		/// Custom permissions for this object
+		/// </summary>
+		public GetAclResponse? Acl { get; set; }
+		
+		/// <summary>
+		/// Stream paused for new builds until this date
+		/// </summary>
+		public DateTime? PausedUntil { get; set; }
+		
+		/// <summary>
+		/// Reason for stream being paused
+		/// </summary>
+		public string? PauseComment { get; set; }
+
+		/// <summary>
+		/// Workflows for this stream
+		/// </summary>
+		public List<WorkflowConfig> Workflows { get; set; }
+
+		/// <summary>
+		/// Constructor
+		/// </summary>
+		/// <param name="id">Unique id of the stream</param>
+		/// <param name="projectId">Unique id of the project containing the stream</param>
+		/// <param name="name">Name of the stream</param>
+		/// <param name="configRevision">The config file path on the server</param>
+		/// <param name="order">Order to display this stream</param>
+		/// <param name="notificationChannel"></param>
+		/// <param name="notificationChannelFilter"></param>
+		/// <param name="triageChannel"></param>
+		/// <param name="defaultPreflight">The default template to use for preflights</param>
+		/// <param name="tabs">List of tabs to display for this stream</param>
+		/// <param name="agentTypes">Map of agent type name to description</param>
+		/// <param name="workspaceTypes">Map of workspace name to description</param>
+		/// <param name="templates">Templates for this stream</param>
+		/// <param name="acl">Permissions for this object</param>
+		/// <param name="pausedUntil">Stream paused for new builds until this date</param>
+		/// <param name="pauseComment">Reason for stream being paused</param>
+		/// <param name="workflows">Workflows for this stream</param>
+		public GetStreamResponse(string id, string projectId, string name, string configRevision, int order, string? notificationChannel, string? notificationChannelFilter, string? triageChannel, DefaultPreflightRequest? defaultPreflight, List<GetStreamTabResponse> tabs, Dictionary<string, GetAgentTypeResponse> agentTypes, Dictionary<string, GetWorkspaceTypeResponse>? workspaceTypes, List<GetTemplateRefResponse> templates, GetAclResponse? acl, DateTime? pausedUntil, string? pauseComment, List<WorkflowConfig> workflows)
+		{
+			Id = id;
+			ProjectId = projectId;
+			Name = name;
+			ConfigRevision = configRevision;
+			Order = order;
+			NotificationChannel = notificationChannel;
+			NotificationChannelFilter = notificationChannelFilter;
+			TriageChannel = triageChannel;
+			DefaultPreflightTemplate = defaultPreflight?.TemplateId;
+			DefaultPreflight = defaultPreflight;
+			Tabs = tabs;
+			AgentTypes = agentTypes;
+			WorkspaceTypes = workspaceTypes;
+			Templates = templates;
+			Acl = acl;
+			PausedUntil = pausedUntil;
+			PauseComment = pauseComment;
+			Workflows = workflows;
+		}
+	}
+
+	/// <summary>
+	/// Information about a page to display in the dashboard for a stream
+	/// </summary>
+	[JsonKnownTypes(typeof(CreateJobsTabRequest))]
+	public abstract class CreateStreamTabRequest
+	{
+		/// <summary>
+		/// Title of this page
+		/// </summary>
+		[Required]
+		public string Title { get; set; } = null!;
+	}
+
+	/// <summary>
+	/// Type of a column in a jobs tab
+	/// </summary>
+	public enum JobsTabColumnType
+	{
+		/// <summary>
+		/// Contains labels
+		/// </summary>
+		Labels,
+
+		/// <summary>
+		/// Contains parameters
+		/// </summary>
+		Parameter
+	}
+
+	/// <summary>
+	/// Describes a column to display on the jobs page
+	/// </summary>
+	public class CreateJobsTabColumnRequest
+	{
+		/// <summary>
+		/// The type of column
+		/// </summary>
+		public JobsTabColumnType Type { get; set; } = JobsTabColumnType.Labels;
+
+		/// <summary>
+		/// Heading for this column
+		/// </summary>
+		[Required]
+		public string Heading { get; set; } = null!;
+
+		/// <summary>
+		/// Category of aggregates to display in this column. If null, includes any aggregate not matched by another column.
+		/// </summary>
+		public string? Category { get; set; }
+
+		/// <summary>
+		/// Parameter to show in this column
+		/// </summary>
+		public string? Parameter { get; set; }
+
+		/// <summary>
+		/// Relative width of this column.
+		/// </summary>
+		public int? RelativeWidth { get; set; }
+
+		/// <summary>
+		/// Construct a JobsTabColumn object from this request
+		/// </summary>
+		/// <returns>Column object</returns>
+		public JobsTabColumn ToModel()
+		{
+			switch (Type)
+			{
+				case JobsTabColumnType.Labels:
+					return new JobsTabLabelColumn(Heading, Category, RelativeWidth);
+				case JobsTabColumnType.Parameter:
+					return new JobsTabParameterColumn(Heading, Parameter ?? "Undefined", RelativeWidth);
+				default:
+					return new JobsTabLabelColumn(Heading, "Undefined", RelativeWidth);
+			}
+		}
+	}
+
+	/// <summary>
+	/// Describes a job page
+	/// </summary>
+	[JsonDiscriminator("Jobs")]
+	public class CreateJobsTabRequest : CreateStreamTabRequest
+	{
+		/// <summary>
+		/// Whether to show job names on this page
+		/// </summary>
+		public bool ShowNames { get; set; }
+
+		/// <summary>
+		/// Names of jobs to include on this page. If there is only one name specified, the name column does not need to be displayed.
+		/// </summary>
+		public List<string>? JobNames { get; set; }
+
+		/// <summary>
+		/// List of job template names to show on this page.
+		/// </summary>
+		public List<string>? Templates { get; set; }
+
+		/// <summary>
+		/// Columns to display for different types of aggregates
+		/// </summary>
+		public List<CreateJobsTabColumnRequest>? Columns { get; set; }
+	}
+
+	/// <summary>
+	/// Information about a page to display in the dashboard for a stream
+	/// </summary>
+	[JsonKnownTypes(typeof(GetJobsTabResponse))]
+	public abstract class GetStreamTabResponse
+	{
+		/// <summary>
+		/// Title of this page
+		/// </summary>
+		public string Title { get; set; }
+
+		/// <summary>
+		/// Constructor
+		/// </summary>
+		/// <param name="title">Title of this page</param>
+		protected GetStreamTabResponse(string title)//, StreamPageType Type)
+		{
+			Title = title;
+		}
+	}
+
+	/// <summary>
+	/// Describes a column to display on the jobs page
+	/// </summary>
+	[JsonKnownTypes(typeof(GetJobsTabLabelColumnResponse), typeof(GetJobsTabParameterColumnResponse))]
+	public abstract class GetJobsTabColumnResponse
+	{
+		/// <summary>
+		/// Heading for this column
+		/// </summary>
+		public string Heading { get; set; } = null!;
+
+		/// <summary>
+		/// Relative width of this column.
+		/// </summary>
+		public int? RelativeWidth { get; set; }
+
+		/// <summary>
+		/// Constructor
+		/// </summary>
+		/// <param name="heading">Heading for this column</param>
+		/// <param name="relativeWidth">Relative width of this column</param>
+		protected GetJobsTabColumnResponse(string heading, int? relativeWidth)
+		{
+			Heading = heading;
+			RelativeWidth = relativeWidth;
+		}
+	}
+
+	/// <summary>
+	/// Describes a label column to display on the jobs page
+	/// </summary>
+	[JsonDiscriminator(nameof(JobsTabColumnType.Labels))]
+	public class GetJobsTabLabelColumnResponse : GetJobsTabColumnResponse
+	{
+		/// <summary>
+		/// Category of aggregates to display in this column. If null, includes any aggregate not matched by another column.
+		/// </summary>
+		public string? Category { get; set; }
+
+		/// <summary>
+		/// Constructor
+		/// </summary>
+		/// <param name="heading">Heading for this column</param>
+		/// <param name="category">Category of aggregates to display in this column. If null, includes any aggregate not matched by another column.</param>
+		/// <param name="relativeWidth">Relative width of this column</param>
+		public GetJobsTabLabelColumnResponse(string heading, string? category, int? relativeWidth)
+			: base(heading, relativeWidth)
+		{
+			Category = category;
+		}
+	}
+
+	/// <summary>
+	/// Describes a parameter column to display on the jobs page
+	/// </summary>
+	[JsonDiscriminator(nameof(JobsTabColumnType.Parameter))]
+	public class GetJobsTabParameterColumnResponse : GetJobsTabColumnResponse
+	{
+		/// <summary>
+		/// Parameter to show in this column
+		/// </summary>
+		public string? Parameter { get; set; }
+
+		/// <summary>
+		/// Constructor
+		/// </summary>
+		/// <param name="heading">Heading for this column</param>
+		/// <param name="parameter">Name of the parameter to display</param>
+		/// <param name="relativeWidth">Relative width of this column</param>
+		public GetJobsTabParameterColumnResponse(string heading, string? parameter, int? relativeWidth)
+			: base(heading, relativeWidth)
+		{
+			Parameter = parameter;
+		}
+	}
+
+	/// <summary>
+	/// Describes a job page
+	/// </summary>
+	[JsonDiscriminator("Jobs")]
+	public class GetJobsTabResponse : GetStreamTabResponse
+	{
+		/// <summary>
+		/// Whether to show names on the page
+		/// </summary>
+		public bool ShowNames { get; set; }
+
+		/// <summary>
+		/// List of templates to show on the page
+		/// </summary>
+		public List<string>? Templates { get; set; }
+
+		/// <summary>
+		/// Names of jobs to include on this page. If there is only one name specified, the name column does not need to be displayed.
+		/// </summary>
+		public List<string>? JobNames { get; set; }
+
+		/// <summary>
+		/// Columns to display for different types of aggregates
+		/// </summary>
+		public List<GetJobsTabColumnResponse>? Columns { get; set; }
+
+		/// <summary>
+		/// Constructor
+		/// </summary>
+		/// <param name="title">Title for this page</param>
+		/// <param name="bShowNames">Whether to show names on the page</param>
+		/// <param name="templates">Templates to include on this page</param>
+		/// <param name="jobNames">List of job names to include on this page</param>
+		/// <param name="columns">List of columns to display</param>
+		public GetJobsTabResponse(string title, bool bShowNames, List<string>? templates, List<string>? jobNames, List<GetJobsTabColumnResponse>? columns)
+			: base(title)
+		{
+			ShowNames = bShowNames;
+			Templates = templates;
+			JobNames = jobNames;
+			Columns = columns;
+		}
+	}
+}
diff --git a/Engine/Source/Programs/Horde/Horde.Build/Streams/StreamsController.cs b/Engine/Source/Programs/Horde/Horde.Build/Streams/StreamsController.cs
new file mode 100644
--- /dev/null
+++ b/Engine/Source/Programs/Horde/Horde.Build/Streams/StreamsController.cs
@@ -0,0 +1,357 @@
+// Copyright Epic Games, Inc. All Rights Reserved.
+
+using System;
+using System.Collections.Generic;
+using System.Linq;
+using System.Threading.Tasks;
+using Horde.Build.Acls;
+using Horde.Build.Api;
+using Horde.Build.Collections;
+using Horde.Build.Models;
+using Horde.Build.Services;
+using Horde.Build.Utilities;
+using Microsoft.AspNetCore.Authorization;
+using Microsoft.AspNetCore.Mvc;
+
+namespace Horde.Build.Controllers
+{
+	using ProjectId = StringId<IProject>;
+	using StreamId = StringId<IStream>;
+	using TemplateRefId = StringId<TemplateRef>;
+
+	/// <summary>
+	/// Controller for the /api/v1/streams endpoint
+	/// </summary>
+	[ApiController]
+	[Authorize]
+	[Route("[controller]")]
+	public class StreamsController : HordeControllerBase
+	{
+		private readonly StreamService _streamService;
+		private readonly ITemplateCollection _templateCollection;
+		private readonly IJobStepRefCollection _jobStepRefCollection;
+		private readonly IPerforceService _perforceService;
+		private readonly IUserCollection _userCollection;
+		private readonly AclService _aclService;
+		/// <summary>
+		/// Constructor
+		/// </summary>
+		public StreamsController(AclService aclService, StreamService streamService, ITemplateCollection templateCollection, IJobStepRefCollection jobStepRefCollection, IUserCollection userCollection, IPerforceService perforceService)
+		{
+			_streamService = streamService;
+			_templateCollection = templateCollection;
+			_jobStepRefCollection = jobStepRefCollection;
+			_userCollection = userCollection;
+			_perforceService = perforceService;
+			_aclService = aclService;
+		}
+
+		/// <summary>
+		/// Query all the streams for a particular project.
+		/// </summary>
+		/// <param name="projectIds">Unique id of the project to query</param>
+		/// <param name="filter">Filter for the properties to return</param>
+		/// <returns>Information about all the projects</returns>
+		[HttpGet]
+		[Route("/api/v1/streams")]
+		[ProducesResponseType(typeof(List<GetStreamResponse>), 200)]
+		public async Task<ActionResult<List<object>>> GetStreamsAsync([FromQuery(Name = "ProjectId")] string[] projectIds, [FromQuery] PropertyFilter? filter = null)
+		{
+			ProjectId[] projectIdValues = Array.ConvertAll(projectIds, x => new ProjectId(x));
+
+			List<IStream> streams = await _streamService.GetStreamsAsync(projectIdValues);
+			ProjectPermissionsCache permissionsCache = new ProjectPermissionsCache();
+
+			List<GetStreamResponse> responses = new List<GetStreamResponse>();
+			foreach (IStream stream in streams)
+			{
+				if (await _streamService.AuthorizeAsync(stream, AclAction.ViewStream, User, permissionsCache))
+				{
+					GetStreamResponse response = await CreateGetStreamResponse(stream, permissionsCache);
+					responses.Add(response);
+				}
+			}
+			return responses.OrderBy(x => x.Order).ThenBy(x => x.Name).Select(x => PropertyFilter.Apply(x, filter)).ToList();
+		}
+
+		/// <summary>
+		/// Retrieve information about a specific stream.
+		/// </summary>
+		/// <param name="streamId">Id of the stream to get information about</param>
+		/// <param name="filter">Filter for the properties to return</param>
+		/// <returns>Information about the requested project</returns>
+		[HttpGet]
+		[Route("/api/v1/streams/{streamId}")]
+		[ProducesResponseType(typeof(GetStreamResponse), 200)]
+		public async Task<ActionResult<object>> GetStreamAsync(StreamId streamId, [FromQuery] PropertyFilter? filter = null)
+		{
+			IStream? stream = await _streamService.GetStreamAsync(streamId);
+			if (stream == null)
+			{
+				return NotFound(streamId);
+			}
+
+			ProjectPermissionsCache permissionsCache = new ProjectPermissionsCache();
+			if (!await _streamService.AuthorizeAsync(stream, AclAction.ViewStream, User, permissionsCache))
+			{
+				return Forbid(AclAction.ViewStream, streamId);
+			}
+
+			return PropertyFilter.Apply(await CreateGetStreamResponse(stream, permissionsCache), filter);
+		}
+
+		/// <summary>
+		/// Create a stream response object, including all the templates
+		/// </summary>
+		/// <param name="stream">Stream to create response for</param>
+		/// <param name="cache">Permissions cache</param>
+		/// <returns>Response object</returns>
+		async Task<GetStreamResponse> CreateGetStreamResponse(IStream stream, ProjectPermissionsCache cache)
+		{
+			bool bIncludeAcl = stream.Acl != null && await _streamService.AuthorizeAsync(stream, AclAction.ViewPermissions, User, cache);
+
+			List<GetTemplateRefResponse> apiTemplateRefs = new List<GetTemplateRefResponse>();
+			foreach (KeyValuePair<TemplateRefId, TemplateRef> pair in stream.Templates)
+			{
+				if (await _streamService.AuthorizeAsync(stream, pair.Value, AclAction.ViewTemplate, User, cache))
+				{
+					ITemplate? template = await _templateCollection.GetAsync(pair.Value.Hash);
+					
+					if (template != null)						
+					{
+						TemplateRef tref = pair.Value;
+
+						List<GetTemplateStepStateResponse>? stepStates = null;
+						if (tref.StepStates != null)
+						{
+							for (int i = 0; i < tref.StepStates.Count; i++)
+							{
+								TemplateStepState state = tref.StepStates[i];
+
+								if (state.PausedByUserId == null)
+								{
+									continue;
+								}
+
+								stepStates ??= new List<GetTemplateStepStateResponse>();
+
+								GetThinUserInfoResponse? pausedByUserInfo = null;
+								if (state.PausedByUserId != null)
+								{
+									pausedByUserInfo = new GetThinUserInfoResponse(await _userCollection.GetCachedUserAsync(state.PausedByUserId));
+								}
+
+								stepStates.Add(new GetTemplateStepStateResponse(state, pausedByUserInfo));
+
+							}
+						}
+
+						bool bIncludeTemplateAcl = pair.Value.Acl != null && await _streamService.AuthorizeAsync(stream, pair.Value, AclAction.ViewPermissions, User, cache);
+						apiTemplateRefs.Add(new GetTemplateRefResponse(pair.Key, pair.Value, template, stepStates, bIncludeTemplateAcl));
+					}
+				}
+			}
+
+			return stream.ToApiResponse(bIncludeAcl, apiTemplateRefs);
+		}
+
+		/// <summary>
+		/// Gets a list of changes for a stream
+		/// </summary>
+		/// <param name="streamId">The stream id</param>
+		/// <param name="min">The starting changelist number</param>
+		/// <param name="max">The ending changelist number</param>
+		/// <param name="results">Number of results to return</param>
+		/// <param name="filter">The filter to apply to the results</param>
+		/// <returns>Http result code</returns>
+		[HttpGet]
+		[Route("/api/v1/streams/{streamId}/changes")]
+		[ProducesResponseType(typeof(List<GetChangeSummaryResponse>), 200)]
+		public async Task<ActionResult<List<object>>> GetChangesAsync(StreamId streamId, [FromQuery] int? min = null, [FromQuery] int? max = null, [FromQuery] int results = 50, PropertyFilter? filter = null)
+		{
+			IStream? stream = await _streamService.GetStreamAsync(streamId);
+			if (stream == null)
+			{
+				return NotFound(streamId);
+			}
+			if (!await _streamService.AuthorizeAsync(stream, AclAction.ViewChanges, User, null))
+			{
+				return Forbid(AclAction.ViewChanges, streamId);
+			}
+
+			string? perforceUser = User.GetPerforceUser();
+			if(perforceUser == null)
+			{
+				return BadRequest("Current user does not have an associated Perforce user");
+			}
+
+			List<ChangeSummary> commits = await _perforceService.GetChangesAsync(stream.ClusterName, stream.Name, min, max, results, perforceUser);
+
+			List<GetChangeSummaryResponse> responses = new List<GetChangeSummaryResponse>();
+			foreach (ChangeSummary commit in commits)
+			{
+				responses.Add(new GetChangeSummaryResponse(commit));
+			}
+			return responses.ConvertAll(x => PropertyFilter.Apply(x, filter));
+		}
+
+		/// <summary>
+		/// Gets a list of changes for a stream
+		/// </summary>
+		/// <param name="streamId">The stream id</param>
+		/// <param name="changeNumber">The changelist number</param>
+		/// <param name="filter">The filter to apply to the results</param>
+		/// <returns>Http result code</returns>
+		[HttpGet]
+		[Route("/api/v1/streams/{streamId}/changes/{changeNumber}")]
+		[ProducesResponseType(typeof(GetChangeDetailsResponse), 200)]
+		public async Task<ActionResult<object>> GetChangeDetailsAsync(StreamId streamId, int changeNumber, PropertyFilter? filter = null)
+		{
+			IStream? stream = await _streamService.GetStreamAsync(streamId);
+			if (stream == null)
+			{
+				return NotFound(streamId);
+			}
+			if (!await _streamService.AuthorizeAsync(stream, AclAction.ViewChanges, User, null))
+			{
+				return Forbid(AclAction.ViewChanges, streamId);
+			}
+
+			string? perforceUser = User.GetPerforceUser();
+			if(perforceUser == null)
+			{
+				return BadRequest("Current user does not have an associated Perforce user");
+			}
+
+			ChangeDetails? changeDetails = await _perforceService.GetChangeDetailsAsync(stream.ClusterName, stream.Name, changeNumber, perforceUser);
+			if(changeDetails == null)
+			{
+				return NotFound("CL {Change} not found in stream {StreamId}", changeNumber, streamId);
+			}
+
+			return PropertyFilter.Apply(new GetChangeDetailsResponse(changeDetails), filter);
+		}
+
+		/// <summary>
+		/// Gets the history of a step in the stream
+		/// </summary>
+		/// <param name="streamId">The stream id</param>
+		/// <param name="templateId"></param>
+		/// <param name="step">Name of the step to search for</param>
+		/// <param name="change">Maximum changelist number to return</param>
+		/// <param name="count">Number of results to return</param>
+		/// <param name="filter">The filter to apply to the results</param>
+		/// <returns>Http result code</returns>
+		[HttpGet]
+		[Route("/api/v1/streams/{streamId}/history")]
+		[ProducesResponseType(typeof(List<GetJobStepRefResponse>), 200)]
+		public async Task<ActionResult<List<object>>> GetStepHistoryAsync(StreamId streamId, [FromQuery] string templateId, [FromQuery] string step, [FromQuery] int? change = null, [FromQuery] int count = 10, [FromQuery] PropertyFilter? filter = null)
+		{
+			IStream? stream = await _streamService.GetStreamAsync(streamId);
+			if (stream == null)
+			{
+				return NotFound(streamId);
+			}
+			if (!await _streamService.AuthorizeAsync(stream, AclAction.ViewJob, User, null))
+			{
+				return Forbid(AclAction.ViewJob, streamId);
+			}
+
+			TemplateRefId templateIdValue = new TemplateRefId(templateId);
+
+			List<IJobStepRef> steps = await _jobStepRefCollection.GetStepsForNodeAsync(streamId, templateIdValue, step, change, true, count);
+			return steps.ConvertAll(x => PropertyFilter.Apply(new GetJobStepRefResponse(x), filter));
+		}
+
+		/// <summary>
+		/// Deletes a stream
+		/// </summary>
+		/// <param name="streamId">Id of the stream to update.</param>
+		/// <returns>Http result code</returns>
+		[HttpDelete]
+		[Route("/api/v1/streams/{streamId}")]
+		public async Task<ActionResult> DeleteStreamAsync(StreamId streamId)
+		{
+			IStream? stream = await _streamService.GetStreamAsync(streamId);
+			if (stream == null)
+			{
+				return NotFound(streamId);
+			}
+			if (!await _streamService.AuthorizeAsync(stream, AclAction.DeleteStream, User, null))
+			{
+				return Forbid(AclAction.DeleteStream, streamId);
+			}
+
+			await _streamService.DeleteStreamAsync(streamId);
+			return new OkResult();
+		}
+
+		/// <summary>
+		/// Gets a template for a stream
+		/// </summary>
+		/// <param name="streamId">Unique id of the stream to query</param>
+		/// <param name="templateId">Unique id of the template to query</param>
+		/// <param name="filter">Filter for properties to return</param>
+		/// <returns>Information about all the templates</returns>
+		[HttpGet]
+		[Route("/api/v1/streams/{streamId}/templates/{templateId}")]
+		[ProducesResponseType(typeof(List<GetTemplateResponse>), 200)]
+		public async Task<ActionResult<object>> GetTemplateAsync(StreamId streamId, TemplateRefId templateId, [FromQuery] PropertyFilter? filter = null)
+		{
+			IStream? stream = await _streamService.GetStreamAsync(streamId);
+			if (stream == null)
+			{
+				return NotFound(streamId);
+			}
+
+			TemplateRef? templateRef;
+			if (!stream.Templates.TryGetValue(templateId, out templateRef))
+			{
+				return NotFound(streamId, templateId);
+			}
+			if (!await _streamService.AuthorizeAsync(stream, templateRef, AclAction.ViewTemplate, User, null))
+			{
+				return Forbid(AclAction.ViewTemplate, streamId);
+			}
+
+			ITemplate? template = await _templateCollection.GetAsync(templateRef.Hash);
+			if(template == null)
+			{
+				return NotFound();
+			}
+
+			return new GetTemplateResponse(template).ApplyFilter(filter);
+		}
+
+		/// <summary>
+		/// Update a stream template ref
+		/// </summary>
+		[HttpPut]
+		[Authorize]
+		[Route("/api/v1/streams/{streamId}/templates/{templateRefId}")]
+		public async Task<ActionResult> UpdateStreamTemplateRefAsync(StreamId streamId, TemplateRefId templateRefId, [FromBody] UpdateTemplateRefRequest update)
+		{
+			if (!await _aclService.AuthorizeAsync(AclAction.AdminWrite, User))
+			{
+				return Forbid(AclAction.AdminWrite);
+			}
+
+			IStream? stream = await _streamService.GetStreamAsync(streamId);
+			if (stream == null)
+			{
+				return NotFound(streamId);
+			}
+
+			if (!stream.Templates.ContainsKey(templateRefId))
+			{
+				return NotFound(streamId, templateRefId);
+			}
+
+			await _streamService.TryUpdateTemplateRefAsync(stream, templateRefId, update.StepStates);
+
+			return Ok();
+
+		}
+
+	}
+}
diff --git a/Engine/Source/Programs/Horde/Horde.Build/Tasks/ConformTaskSource.cs b/Engine/Source/Programs/Horde/Horde.Build/Tasks/ConformTaskSource.cs
new file mode 100644
--- /dev/null
+++ b/Engine/Source/Programs/Horde/Horde.Build/Tasks/ConformTaskSource.cs
@@ -0,0 +1,383 @@
+// Copyright Epic Games, Inc. All Rights Reserved.
+
+using System;
+using System.Collections.Generic;
+using System.Linq;
+using System.Threading;
+using System.Threading.Tasks;
+using Google.Protobuf;
+using Google.Protobuf.WellKnownTypes;
+using Horde.Build.Api;
+using Horde.Build.Collections;
+using Horde.Build.Models;
+using Horde.Build.Server;
+using Horde.Build.Services;
+using Horde.Build.Utilities;
+using HordeCommon;
+using HordeCommon.Rpc.Tasks;
+using Microsoft.Extensions.Hosting;
+using Microsoft.Extensions.Logging;
+
+namespace Horde.Build.Tasks.Impl
+{
+	using JobId = ObjectId<IJob>;
+	using LeaseId = ObjectId<ILease>;
+	using LogId = ObjectId<ILogFile>;
+
+	/// <summary>
+	/// Generates tasks telling agents to sync their workspaces
+	/// </summary>
+	public sealed class ConformTaskSource : TaskSourceBase<ConformTask>, IHostedService, IDisposable
+	{
+		/// <inheritdoc/>
+		public override string Type => "Conform";
+
+		/// <inheritdoc/>
+		public override TaskSourceFlags Flags => TaskSourceFlags.AllowWhenDisabled;
+
+		readonly MongoService _mongoService;
+		readonly IAgentCollection _agentCollection;
+		readonly PoolService _poolService;
+		readonly SingletonDocument<ConformList> _conformList;
+		readonly PerforceLoadBalancer _perforceLoadBalancer;
+		readonly ILogFileService _logService;
+		readonly ILogger _logger;
+		readonly ITicker _tickConformList;
+
+		/// <summary>
+		/// Constructor
+		/// </summary>
+		public ConformTaskSource(MongoService mongoService, IAgentCollection agentCollection, PoolService poolService, ILogFileService logService, PerforceLoadBalancer perforceLoadBalancer, IClock clock, ILogger<ConformTaskSource> logger)
+		{
+			_mongoService = mongoService;
+			_agentCollection = agentCollection;
+			_poolService = poolService;
+			_conformList = new SingletonDocument<ConformList>(mongoService);
+			_perforceLoadBalancer = perforceLoadBalancer;
+			_logService = logService;
+			_logger = logger;
+			_tickConformList = clock.AddSharedTicker<ConformTaskSource>(TimeSpan.FromMinutes(1.0), CleanConformListAsync, logger);
+
+			OnLeaseStartedProperties.Add(nameof(ConformTask.LogId), x => new LogId(x.LogId));
+		}
+
+		/// <inheritdoc/>
+		public Task StartAsync(CancellationToken cancellationToken) => _tickConformList.StartAsync();
+
+		/// <inheritdoc/>
+		public Task StopAsync(CancellationToken cancellationToken) => _tickConformList.StopAsync();
+
+		/// <inheritdoc/>
+		public void Dispose() => _tickConformList.Dispose();
+
+		/// <summary>
+		/// Clean up the conform list of any outdated entries
+		/// </summary>
+		/// <param name="cancellationToken">Cancellation token</param>
+		/// <returns>Async task</returns>
+		async ValueTask CleanConformListAsync(CancellationToken cancellationToken)
+		{
+			DateTime utcNow = DateTime.UtcNow;
+			DateTime lastCheckTimeUtc = utcNow - TimeSpan.FromMinutes(30.0);
+
+			// Get the current state of the conform list
+			ConformList list = await _conformList.GetAsync();
+
+			// Update any leases that are older than LastCheckTimeUtc
+			Dictionary<LeaseId, bool> removeLeases = new Dictionary<LeaseId, bool>();
+			foreach (ConformListEntry entry in Enumerable.Concat(list.Entries, list.Servers.SelectMany(x => x.Entries)))
+			{
+				if (entry.LastCheckTimeUtc < lastCheckTimeUtc)
+				{
+					IAgent? agent = await _agentCollection.GetAsync(entry.AgentId);
+
+					bool remove = false;
+					if (agent == null || !agent.Leases.Any(x => x.Id == entry.LeaseId))
+					{
+						_logger.LogWarning("Removing invalid lease from conform list: {AgentId} lease {LeaseId}", entry.AgentId, entry.LeaseId);
+						remove = true;
+					}
+
+					removeLeases[entry.LeaseId] = remove;
+				}
+			}
+
+			// If there's anything to change, update the list
+			if (removeLeases.Count > 0)
+			{
+				await _conformList.UpdateAsync(list => UpdateConformList(list, utcNow, removeLeases));
+			}
+		}
+
+		/// <summary>
+		/// Remove items from the conform list, and update timestamps for them
+		/// </summary>
+		/// <param name="list">The list to update</param>
+		/// <param name="utcNow">Current time</param>
+		/// <param name="removeLeases">List of leases to update. Entries with values set to true will be removed, entries with values set to false will have their timestamp updated.</param>
+		static void UpdateConformList(ConformList list, DateTime utcNow, Dictionary<LeaseId, bool> removeLeases)
+		{
+			UpdateConformList(list.Entries, utcNow, removeLeases);
+			foreach (ConformListServer server in list.Servers)
+			{
+				UpdateConformList(server.Entries, utcNow, removeLeases);
+			}
+		}
+
+		/// <summary>
+		/// Remove items from the conform list, and update timestamps for them
+		/// </summary>
+		/// <param name="entries">The list to update</param>
+		/// <param name="utcNow">Current time</param>
+		/// <param name="removeLeases">List of leases to update. Entries with values set to true will be removed, entries with values set to false will have their timestamp updated.</param>
+		static void UpdateConformList(List<ConformListEntry> entries, DateTime utcNow, Dictionary<LeaseId, bool> removeLeases)
+		{
+			for (int idx = 0; idx < entries.Count; idx++)
+			{
+				bool remove;
+				if (removeLeases.TryGetValue(entries[idx].LeaseId, out remove))
+				{
+					if (remove)
+					{
+						entries.RemoveAt(idx--);
+					}
+					else if (entries[idx].LastCheckTimeUtc < utcNow)
+					{
+						entries[idx].LastCheckTimeUtc = utcNow;
+					}
+				}
+			}
+		}
+
+		/// <inheritdoc/>
+		public override async Task<Task<AgentLease?>> AssignLeaseAsync(IAgent agent, CancellationToken cancellationToken)
+		{
+			DateTime utcNow = DateTime.UtcNow;
+			if (!await IsConformPendingAsync(agent, utcNow))
+			{
+				return Skip(cancellationToken);
+			}
+
+			if (agent.Leases.Count == 0)
+			{
+				ConformTask task = new ConformTask();
+				if (await GetWorkspacesAsync(agent, task.Workspaces))
+				{
+					LeaseId leaseId = LeaseId.GenerateNewId();
+					if (await AllocateConformLeaseAsync(agent.Id, task.Workspaces, leaseId))
+					{
+						ILogFile log = await _logService.CreateLogFileAsync(JobId.Empty, agent.SessionId, LogType.Json);
+						task.LogId = log.Id.ToString();
+						task.RemoveUntrackedFiles = agent.RequestFullConform;
+
+						byte[] payload = Any.Pack(task).ToByteArray();
+
+						return Lease(new AgentLease(leaseId, "Updating workspaces", null, null, log.Id, LeaseState.Pending, null, true, payload));
+					}
+				}
+			}
+
+			if (agent.RequestConform || agent.RequestFullConform)
+			{
+				return await DrainAsync(cancellationToken);
+			}
+			else
+			{
+				return Skip(cancellationToken);
+			}
+		}
+
+		/// <inheritdoc/>
+		public override Task CancelLeaseAsync(IAgent agent, LeaseId leaseId, ConformTask payload)
+		{
+			return ReleaseConformLeaseAsync(leaseId);
+		}
+
+		/// <inheritdoc/>
+		public async Task<bool> GetWorkspacesAsync(IAgent agent, IList<HordeCommon.Rpc.Messages.AgentWorkspace> workspaces)
+		{
+			Globals globals = await _mongoService.GetGlobalsAsync();
+
+			HashSet<AgentWorkspace> conformWorkspaces = await _poolService.GetWorkspacesAsync(agent, DateTime.UtcNow);
+			foreach (AgentWorkspace conformWorkspace in conformWorkspaces)
+			{
+				PerforceCluster? cluster = globals.FindPerforceCluster(conformWorkspace.Cluster);
+				if (cluster == null || !await agent.TryAddWorkspaceMessage(conformWorkspace, cluster, _perforceLoadBalancer, workspaces))
+				{
+					return false;
+				}
+			}
+
+			return true;
+		}
+
+		/// <inheritdoc/>
+		public override Task OnLeaseFinishedAsync(IAgent agent, LeaseId leaseId, ConformTask payload, LeaseOutcome outcome, ReadOnlyMemory<byte> output, ILogger logger)
+		{
+			return ReleaseConformLeaseAsync(leaseId);
+		}
+
+		/// <summary>
+		/// Atempt to allocate a conform resource for the given lease
+		/// </summary>
+		/// <param name="agentId">The agent id</param>
+		/// <param name="workspaces">List of workspaces that are required</param>
+		/// <param name="leaseId">The lease id</param>
+		/// <returns>True if the resource was allocated, false otherwise</returns>
+		private async Task<bool> AllocateConformLeaseAsync(AgentId agentId, IEnumerable<HordeCommon.Rpc.Messages.AgentWorkspace> workspaces, LeaseId leaseId)
+		{
+			Globals globals = await _mongoService.GetGlobalsAsync();
+			for (; ; )
+			{
+				ConformList currentValue = await _conformList.GetAsync();
+				if (globals.MaxConformCount != 0 && currentValue.Entries.Count + currentValue.Servers.Sum(x => x.Entries.Count) >= globals.MaxConformCount)
+				{
+					return false;
+				}
+
+				HashSet<string> servers = new HashSet<string>(StringComparer.OrdinalIgnoreCase);
+				foreach (HordeCommon.Rpc.Messages.AgentWorkspace workspace in workspaces)
+				{
+					if (servers.Add(workspace.ServerAndPort))
+					{
+						PerforceCluster? cluster = globals.FindPerforceCluster(workspace.Cluster);
+						if (cluster == null)
+						{
+							_logger.LogWarning("Unable to find perforce cluster '{Cluster}' for conform", workspace.Cluster);
+							return false;
+						}
+
+						PerforceServer? server = cluster.Servers.FirstOrDefault(x => x.ServerAndPort.Equals(workspace.BaseServerAndPort, StringComparison.Ordinal));
+						if (server == null)
+						{
+							_logger.LogWarning("Unable to find perforce server '{Server}' in cluster '{Cluster}' for conform", workspace.BaseServerAndPort, workspace.Cluster);
+							return false;
+						}
+
+						ConformListServer? conformServer = currentValue.Servers.FirstOrDefault(x => x.Cluster.Equals(workspace.Cluster, StringComparison.Ordinal) && x.ServerAndPort.Equals(workspace.ServerAndPort, StringComparison.Ordinal));
+						if (conformServer == null)
+						{
+							conformServer = new ConformListServer { Cluster = workspace.Cluster, ServerAndPort = workspace.ServerAndPort };
+							currentValue.Servers.Add(conformServer);
+						}
+
+						if (server.MaxConformCount != 0 && conformServer.Entries.Count > server.MaxConformCount)
+						{
+							return false;
+						}
+
+						ConformListEntry entry = new ConformListEntry();
+						entry.AgentId = agentId;
+						entry.LeaseId = leaseId;
+						entry.LastCheckTimeUtc = DateTime.UtcNow;
+						conformServer.Entries.Add(entry);
+					}
+				}
+
+				if (await _conformList.TryUpdateAsync(currentValue))
+				{
+					_logger.LogInformation("Added conform lease {LeaseId}", leaseId);
+					return true;
+				}
+			}
+		}
+
+		/// <summary>
+		/// Terminate a conform lease
+		/// </summary>
+		/// <param name="leaseId">The lease id</param>
+		/// <returns>Async task</returns>
+		public async Task ReleaseConformLeaseAsync(LeaseId leaseId)
+		{
+			for (; ; )
+			{
+				ConformList currentValue = await _conformList.GetAsync();
+				if (currentValue.Entries.RemoveAll(x => x.LeaseId == leaseId) + currentValue.Servers.Sum(x => x.Entries.RemoveAll(x => x.LeaseId == leaseId)) == 0)
+				{
+					_logger.LogInformation("Conform lease {LeaseId} is not in singelton", leaseId);
+					break;
+				}
+				if (await _conformList.TryUpdateAsync(currentValue))
+				{
+					_logger.LogInformation("Removed conform lease {LeaseId}", leaseId);
+					break;
+				}
+			}
+		}
+
+		/// <summary>
+		/// Determine if an agent should be conformed
+		/// </summary>
+		/// <param name="agent">The agent to test</param>
+		/// <param name="utcNow">Current time</param>
+		/// <returns>True if the agent should be conformed</returns>
+		private async Task<bool> IsConformPendingAsync(IAgent agent, DateTime utcNow)
+		{
+			// If a conform was manually requested, allow it to run even if the agent is disabled
+			if (agent.RequestConform || agent.RequestFullConform)
+			{
+				return !IsConformCoolDownPeriod(agent, utcNow);
+			}
+
+			// Otherwise only run if the agent is enabled
+			if (agent.Enabled)
+			{
+				// If we've attempted (and failed) a conform, run again after a certain time
+				if (agent.ConformAttemptCount.HasValue)
+				{
+					return !IsConformCoolDownPeriod(agent, utcNow);
+				}
+
+				// Always run a conform every 24h
+				if (utcNow > agent.LastConformTime + TimeSpan.FromDays(1.0))
+				{
+					return true;
+				}
+
+				// Check if the workspaces have changed (first check against a cached list of workspaces, then an accurate one)
+				HashSet<AgentWorkspace> workspaces = await _poolService.GetWorkspacesAsync(agent, utcNow - TimeSpan.FromSeconds(30.0));
+				if (!workspaces.SetEquals(agent.Workspaces))
+				{
+					workspaces = await _poolService.GetWorkspacesAsync(agent, utcNow);
+					if (!workspaces.SetEquals(agent.Workspaces))
+					{
+						return true;
+					}
+				}
+			}
+			return false;
+		}
+
+		/// <summary>
+		/// Determines if the given agent is in a cooldown period, waiting to retry after a failed conform
+		/// </summary>
+		/// <param name="agent">The agent to test</param>
+		/// <param name="utcNow">Current time</param>
+		/// <returns></returns>
+		private static bool IsConformCoolDownPeriod(IAgent agent, DateTime utcNow)
+		{
+			if (!agent.ConformAttemptCount.HasValue)
+			{
+				return false;
+			}
+
+			TimeSpan retryTime;
+			switch (agent.ConformAttemptCount.Value)
+			{
+				case 1:
+					retryTime = TimeSpan.FromMinutes(5.0);
+					break;
+				case 2:
+					retryTime = TimeSpan.FromMinutes(20.0);
+					break;
+				case 3:
+					retryTime = TimeSpan.FromHours(1.0);
+					break;
+				default:
+					retryTime = TimeSpan.FromHours(6.0);
+					break;
+			}
+			return utcNow < agent.LastConformTime + retryTime;
+		}
+	}
+}
diff --git a/Engine/Source/Programs/Horde/Horde.Build/Tasks/RestartTaskSource.cs b/Engine/Source/Programs/Horde/Horde.Build/Tasks/RestartTaskSource.cs
new file mode 100644
--- /dev/null
+++ b/Engine/Source/Programs/Horde/Horde.Build/Tasks/RestartTaskSource.cs
@@ -0,0 +1,59 @@
+// Copyright Epic Games, Inc. All Rights Reserved.
+
+using System.Collections.Generic;
+using System.Runtime.CompilerServices;
+using System.Threading;
+using System.Threading.Tasks;
+using EpicGames.Core;
+using Google.Protobuf;
+using Google.Protobuf.WellKnownTypes;
+using Horde.Build.Api;
+using Horde.Build.Models;
+using Horde.Build.Services;
+using Horde.Build.Utilities;
+using HordeCommon;
+using HordeCommon.Rpc.Tasks;
+
+namespace Horde.Build.Tasks.Impl
+{
+	using JobId = ObjectId<IJob>;
+	using LeaseId = ObjectId<ILease>;
+	using LogId = ObjectId<ILogFile>;
+
+	class RestartTaskSource : TaskSourceBase<RestartTask>
+	{
+		public override string Type => "Restart";
+
+		public override TaskSourceFlags Flags => TaskSourceFlags.AllowWhenDisabled | TaskSourceFlags.AllowDuringDowntime;
+
+		readonly ILogFileService _logService;
+
+		public RestartTaskSource(ILogFileService logService)
+		{
+			_logService = logService;
+
+			OnLeaseStartedProperties.Add(nameof(RestartTask.LogId), x => new LogId(x.LogId));
+		}
+
+		public override async Task<Task<AgentLease?>> AssignLeaseAsync(IAgent agent, CancellationToken cancellationToken)
+		{
+			if (!agent.RequestRestart)
+			{
+				return Skip(cancellationToken);
+			}
+			if (agent.Leases.Count > 0)
+			{
+				return await DrainAsync(cancellationToken);
+			}
+
+			ILogFile log = await _logService.CreateLogFileAsync(JobId.Empty, agent.SessionId, LogType.Json);
+
+			RestartTask task = new RestartTask();
+			task.LogId = log.Id.ToString();
+
+			byte[] payload = Any.Pack(task).ToByteArray();
+
+			return Lease(new AgentLease(LeaseId.GenerateNewId(), "Restart", null, null, log.Id, LeaseState.Pending, null, true, payload));
+		}
+	}
+}
diff --git a/Engine/Source/Programs/Horde/Horde.Build/Tasks/ShutdownTaskSource.cs b/Engine/Source/Programs/Horde/Horde.Build/Tasks/ShutdownTaskSource.cs
new file mode 100644
--- /dev/null
+++ b/Engine/Source/Programs/Horde/Horde.Build/Tasks/ShutdownTaskSource.cs
@@ -0,0 +1,55 @@
+// Copyright Epic Games, Inc. All Rights Reserved.
+
+using System.Threading;
+using System.Threading.Tasks;
+using Google.Protobuf;
+using Google.Protobuf.WellKnownTypes;
+using Horde.Build.Api;
+using Horde.Build.Models;
+using Horde.Build.Services;
+using Horde.Build.Utilities;
+using HordeCommon;
+using HordeCommon.Rpc.Tasks;
+
+namespace Horde.Build.Tasks.Impl
+{
+	using JobId = ObjectId<IJob>;
+	using LeaseId = ObjectId<ILease>;
+	using LogId = ObjectId<ILogFile>;
+
+	class ShutdownTaskSource : TaskSourceBase<ShutdownTask>
+	{
+		public override string Type => "Shutdown";
+
+		public override TaskSourceFlags Flags => TaskSourceFlags.AllowWhenDisabled | TaskSourceFlags.AllowDuringDowntime;
+
+		readonly ILogFileService _logService;
+
+		public ShutdownTaskSource(ILogFileService logService)
+		{
+			_logService = logService;
+			OnLeaseStartedProperties.Add(nameof(ShutdownTask.LogId), x => new LogId(x.LogId));
+		}
+
+		public override async Task<Task<AgentLease?>> AssignLeaseAsync(IAgent agent, CancellationToken cancellationToken)
+		{
+			if (!agent.RequestShutdown)
+			{
+				return Skip(cancellationToken);
+			}
+			if (agent.Leases.Count > 0)
+			{
+				return await DrainAsync(cancellationToken);
+			}
+
+			ILogFile log = await _logService.CreateLogFileAsync(JobId.Empty, agent.SessionId, LogType.Json);
+
+			ShutdownTask task = new ShutdownTask();
+			task.LogId = log.Id.ToString();
+
+			byte[] payload = Any.Pack(task).ToByteArray();
+
+			return Lease(new AgentLease(LeaseId.GenerateNewId(), "Shutdown", null, null, log.Id, LeaseState.Pending, null, true, payload));
+		}
+	}
+}
diff --git a/Engine/Source/Programs/Horde/Horde.Build/Tasks/UpgradeTaskSource.cs b/Engine/Source/Programs/Horde/Horde.Build/Tasks/UpgradeTaskSource.cs
new file mode 100644
--- /dev/null
+++ b/Engine/Source/Programs/Horde/Horde.Build/Tasks/UpgradeTaskSource.cs
@@ -0,0 +1,89 @@
+// Copyright Epic Games, Inc. All Rights Reserved.
+
+using System;
+using System.Collections.Generic;
+using System.Runtime.CompilerServices;
+using System.Threading;
+using System.Threading.Tasks;
+using Google.Protobuf;
+using Google.Protobuf.WellKnownTypes;
+using Horde.Build.Api;
+using Horde.Build.Models;
+using Horde.Build.Services;
+using Horde.Build.Utilities;
+using HordeCommon;
+using HordeCommon.Rpc.Tasks;
+
+namespace Horde.Build.Tasks.Impl
+{
+	using AgentSoftwareChannelName = StringId<AgentSoftwareChannels>;
+	using JobId = ObjectId<IJob>;
+	using LeaseId = ObjectId<ILease>;
+	using LogId = ObjectId<ILogFile>;
+
+	class UpgradeTaskSource : TaskSourceBase<UpgradeTask>
+	{
+		public override string Type => "Upgrade";
+
+		public override TaskSourceFlags Flags => TaskSourceFlags.AllowWhenDisabled | TaskSourceFlags.AllowDuringDowntime;
+
+		readonly AgentSoftwareService _agentSoftwareService;
+		readonly ILogFileService _logService;
+		readonly IClock _clock;
+
+		public UpgradeTaskSource(AgentSoftwareService agentSoftwareService, ILogFileService logService, IClock clock)
+		{
+			_agentSoftwareService = agentSoftwareService;
+			_logService = logService;
+			_clock = clock;
+
+			OnLeaseStartedProperties.Add(nameof(UpgradeTask.LogId), x => new LogId(x.LogId));
+		}
+
+		public override async Task<Task<AgentLease?>> AssignLeaseAsync(IAgent agent, CancellationToken cancellationToken)
+		{
+			string? requiredVersion = await GetRequiredSoftwareVersion(agent);
+			if (requiredVersion == null || agent.Version == requiredVersion)
+			{
+				return Skip(cancellationToken);
+			}
+			if (agent.Leases.Count > 0 || !(agent.LastUpgradeTime == null || agent.LastUpgradeTime.Value + TimeSpan.FromMinutes(5.0) < _clock.UtcNow || agent.LastUpgradeVersion != requiredVersion.ToString()))
+			{
+				return await DrainAsync(cancellationToken);
+			}
+
+			ILogFile logFile = await _logService.CreateLogFileAsync(JobId.Empty, agent.SessionId, LogType.Json);
+
+			UpgradeTask task = new UpgradeTask();
+			task.SoftwareId = requiredVersion.ToString();
+			task.LogId = logFile.Id.ToString();
+
+			byte[] payload;
+			if (agent.Version == "5.0.0-17425336" || agent.Version == "5.0.0-17448746")
+			{
+				Any any = new Any();
+				any.TypeUrl = "type.googleapis.com/Horde.UpgradeTask";
+				any.Value = task.ToByteString();
+				payload = any.ToByteArray();
+			}
+			else
+			{
+				payload = Any.Pack(task).ToByteArray();
+			}
+
+			return Lease(new AgentLease(LeaseId.GenerateNewId(), $"Upgrade to {requiredVersion}", null, null, logFile.Id, LeaseState.Pending, null, true, payload));
+		}
+
+		/// <summary>
+		/// Determines the client software version that should be installed on an agent
+		/// </summary>
+		/// <param name="agent">The agent instance</param>
+		/// <returns>Unique id of the client version this agent should be running</returns>
+		public async Task<string?> GetRequiredSoftwareVersion(IAgent agent)
+		{
+			AgentSoftwareChannelName channelName = agent.Channel ?? AgentSoftwareService.DefaultChannelName;
+			IAgentSoftwareChannel? channel = await _agentSoftwareService.GetCachedChannelAsync(channelName);
+			return channel?.Version;
+		}
+	}
+}
diff --git a/Engine/Source/Programs/Horde/Horde.Build/Ugs/UgsController.Messages.cs b/Engine/Source/Programs/Horde/Horde.Build/Ugs/UgsController.Messages.cs
new file mode 100644
--- /dev/null
+++ b/Engine/Source/Programs/Horde/Horde.Build/Ugs/UgsController.Messages.cs
@@ -0,0 +1,561 @@
+// Copyright Epic Games, Inc. All Rights Reserved.
+
+using System;
+using System.Collections.Generic;
+using System.ComponentModel.DataAnnotations;
+using System.Linq;
+using EpicGames.Core;
+using Horde.Build.Models;
+
+namespace Horde.Build.Api
+{
+	/// <summary>
+	/// Review by a user of a particular change
+	/// </summary>
+	public enum UgsUserVote
+	{
+		/// <summary>
+		/// No vote for the current change
+		/// </summary>
+		None,
+
+		/// <summary>
+		/// Succesfully compiled the change
+		/// </summary>
+		CompileSuccess,
+
+		/// <summary>
+		/// Failed to compile the change
+		/// </summary>
+		CompileFailure,
+
+		/// <summary>
+		/// Manually marked the change as good
+		/// </summary>
+		Good,
+
+		/// <summary>
+		/// Manually marked the change as bad
+		/// </summary>
+		Bad
+	}
+
+	/// <summary>
+	/// State of a badge
+	/// </summary>
+	public enum UgsBadgeState
+	{
+		/// <summary>
+		/// Starting work on this badge, outcome currently unknown
+		/// </summary>
+		Starting,
+
+		/// <summary>
+		/// Badge failed
+		/// </summary>
+		Failure,
+
+		/// <summary>
+		/// Badge produced a warning
+		/// </summary>
+		Warning,
+
+		/// <summary>
+		/// Badge succeeded
+		/// </summary>
+		Success,
+
+		/// <summary>
+		/// Badge was skipped
+		/// </summary>
+		Skipped,
+	}
+
+	/// <summary>
+	/// Adds a new badge to the change
+	/// </summary>
+	public class AddUgsBadgeRequest
+	{
+		/// <summary>
+		/// Name of the badge
+		/// </summary>
+		[Required]
+		public string Name { get; set; } = String.Empty;
+
+		/// <summary>
+		/// Url for the badge
+		/// </summary>
+		public Uri? Url { get; set; }
+
+		/// <summary>
+		/// Current status of the badge
+		/// </summary>
+		public UgsBadgeState State { get; set; }
+	}
+
+	/// <summary>
+	/// Request object for adding new metadata to the server
+	/// </summary>
+	public class AddUgsMetadataRequest
+	{
+		/// <summary>
+		/// The stream name
+		/// </summary>
+		[Required]
+		public string Stream { get; set; } = null!;
+
+		/// <summary>
+		/// The changelist number
+		/// </summary>
+		[Required]
+		public int Change { get; set; }
+
+		/// <summary>
+		/// The project name
+		/// </summary>
+		public string? Project { get; set; }
+
+		/// <summary>
+		/// Name of the current user
+		/// </summary>
+		public string? UserName { get; set; }
+
+		/// <summary>
+		/// Whether this changelist has been synced by the user
+		/// </summary>
+		public bool? Synced { get; set; }
+
+		/// <summary>
+		/// State of the user
+		/// </summary>
+		public UgsUserVote? Vote { get; set; }
+
+		/// <summary>
+		/// New starred state for the issue
+		/// </summary>
+		public bool? Starred { get; set; }
+
+		/// <summary>
+		/// Whether the user is investigating
+		/// </summary>
+		public bool? Investigating { get; set; }
+
+		/// <summary>
+		/// Comment for this change
+		/// </summary>
+		public string? Comment { get; set; }
+
+		/// <summary>
+		/// List of badges to add
+		/// </summary>
+		public List<AddUgsBadgeRequest>? Badges { get; set; }
+	}
+
+	/// <summary>
+	/// Information about a user synced to a change
+	/// </summary>
+	public class GetUgsUserResponse
+	{
+		/// <summary>
+		/// Name of the user
+		/// </summary>
+		public string User { get; set; }
+
+		/// <summary>
+		/// Time that the change was synced
+		/// </summary>
+		public long? SyncTime { get; set; }
+
+		/// <summary>
+		/// State of the user
+		/// </summary>
+		public UgsUserVote? Vote { get; set; }
+
+		/// <summary>
+		/// Comment by this user
+		/// </summary>
+		public string? Comment { get; set; }
+
+		/// <summary>
+		/// Whether the user is investigating this change
+		/// </summary>
+		public bool? Investigating { get; set; }
+
+		/// <summary>
+		/// Whether this changelist is starred
+		/// </summary>
+		public bool? Starred { get; set; }
+
+		/// <summary>
+		/// Constructor
+		/// </summary>
+		/// <param name="userData">The user data object</param>
+		internal GetUgsUserResponse(IUgsUserData userData)
+		{
+			User = userData.User;
+			SyncTime = userData.SyncTime;
+			Vote = (userData.Vote == UgsUserVote.None) ? (UgsUserVote?)null : userData.Vote;
+			Comment = userData.Comment;
+			Investigating = userData.Investigating;
+			Starred = userData.Starred;
+		}
+	}
+
+	/// <summary>
+	/// Information about a badge
+	/// </summary>
+	public class GetUgsBadgeResponse
+	{
+		/// <summary>
+		/// Name of the badge
+		/// </summary>
+		public string Name { get; set; }
+
+		/// <summary>
+		/// Url for the badge
+		/// </summary>
+		public Uri? Url { get; set; }
+
+		/// <summary>
+		/// Current status of the badge
+		/// </summary>
+		public UgsBadgeState State { get; set; }
+
+		/// <summary>
+		/// Constructor
+		/// </summary>
+		/// <param name="badgeData">The badge data object</param>
+		internal GetUgsBadgeResponse(IUgsBadgeData badgeData)
+		{
+			Name = badgeData.Name;
+			Url = badgeData.Url;
+			State = badgeData.State;
+		}
+	}
+
+	/// <summary>
+	/// Response object for querying metadata
+	/// </summary>
+	public class GetUgsMetadataResponse
+	{
+		/// <summary>
+		/// Number of the changelist
+		/// </summary>
+		public int Change { get; set; }
+
+		/// <summary>
+		/// The project name
+		/// </summary>
+		public string? Project { get; set; }
+
+		/// <summary>
+		/// Information about a user synced to this change
+		/// </summary>
+		public List<GetUgsUserResponse>? Users { get; set; }
+
+		/// <summary>
+		/// Badges for this change
+		/// </summary>
+		public List<GetUgsBadgeResponse>? Badges { get; set; }
+
+		/// <summary>
+		/// Constructor
+		/// </summary>
+		/// <param name="metadata">Metadata object</param>
+		internal GetUgsMetadataResponse(IUgsMetadata metadata)
+		{
+			Change = metadata.Change;
+			Project = metadata.Project;
+			Users = metadata.Users?.ConvertAll(x => new GetUgsUserResponse(x));
+			Badges = metadata.Badges?.ConvertAll(x => new GetUgsBadgeResponse(x));
+		}
+	}
+
+	/// <summary>
+	/// Response object for querying metadata
+	/// </summary>
+	public class GetUgsMetadataListResponse
+	{
+		/// <summary>
+		/// Last time that the metadata was modified
+		/// </summary>
+		public long SequenceNumber { get; set; }
+
+		/// <summary>
+		/// List of changes matching the requested criteria
+		/// </summary>
+		public List<GetUgsMetadataResponse> Items { get; set; } = new List<GetUgsMetadataResponse>();
+	}
+
+	/// <summary>
+	/// Outcome of a particular build
+	/// </summary>
+	public enum UgsIssueBuildOutcome
+	{
+		/// <summary>
+		/// Unknown outcome
+		/// </summary>
+		Unknown,
+
+		/// <summary>
+		/// Build succeeded
+		/// </summary>
+		Success,
+
+		/// <summary>
+		/// Build failed
+		/// </summary>
+		Error,
+
+		/// <summary>
+		/// Build finished with warnings
+		/// </summary>
+		Warning,
+	}
+
+	/// <summary>
+	/// Legacy response describing a build
+	/// </summary>
+	public class GetUgsIssueBuildResponse
+	{
+		/// <summary>
+		/// Identifier for this build
+		/// </summary>
+		public long Id { get; set; }
+
+		/// <summary>
+		/// Path to the stream containing this build
+		/// </summary>
+		public string Stream { get; set; }
+
+		/// <summary>
+		/// The changelist number for this build
+		/// </summary>
+		public int Change { get; set; }
+
+		/// <summary>
+		/// Name of this job
+		/// </summary>
+		public string? JobName { get; set; }
+
+		/// <summary>
+		/// Link to the job
+		/// </summary>
+		public Uri? JobUrl { get; set; }
+
+		/// <summary>
+		/// Name of the job step
+		/// </summary>
+		public string? JobStepName { get; set; }
+
+		/// <summary>
+		/// Link to the job step
+		/// </summary>
+		public Uri? JobStepUrl { get; set; }
+
+		/// <summary>
+		/// Url for this particular error
+		/// </summary>
+		public Uri? ErrorUrl { get; set; }
+
+		/// <summary>
+		/// Outcome of this build (see <see cref="IssueBuildOutcome"/>)
+		/// </summary>
+		public int Outcome { get; set; }
+
+		/// <summary>
+		/// 
+		/// </summary>
+		/// <param name="stream"></param>
+		/// <param name="change"></param>
+		/// <param name="outcome"></param>
+		public GetUgsIssueBuildResponse(string stream, int change, IssueBuildOutcome outcome)
+		{
+			Stream = stream;
+			Change = change;
+			Outcome = (int)outcome;
+		}
+	}
+
+	/// <summary>
+	/// Information about a diagnostic
+	/// </summary>
+	public class GetUgsIssueDiagnosticResponse
+	{
+		/// <summary>
+		/// The corresponding build id
+		/// </summary>
+		public long? BuildId { get; set; }
+
+		/// <summary>
+		/// Message for the diagnostic
+		/// </summary>
+		public string Message { get; set; }
+
+		/// <summary>
+		/// Link to the error
+		/// </summary>
+		public Uri Url { get; set; }
+
+		/// <summary>
+		/// Constructor
+		/// </summary>
+		/// <param name="buildId">The corresponding build id</param>
+		/// <param name="message">Message for the diagnostic</param>
+		/// <param name="url">Link to the diagnostic</param>
+		public GetUgsIssueDiagnosticResponse(long? buildId, string message, Uri url)
+		{
+			BuildId = buildId;
+			Message = message;
+			Url = url;
+		}
+	}
+
+	/// <summary>
+	/// Stores information about a build health issue
+	/// </summary>
+	public class GetUgsIssueResponse
+	{
+		/// <summary>
+		/// Version number for this response
+		/// </summary>
+		public int Version { get; set; } = 1;
+
+		/// <summary>
+		/// The unique object id
+		/// </summary>
+		public int Id { get; set; }
+
+		/// <summary>
+		/// Time at which the issue was created
+		/// </summary>
+		public DateTime CreatedAt { get; set; }
+
+		/// <summary>
+		/// Time at which the issue was retrieved
+		/// </summary>
+		public DateTime RetrievedAt { get; set; }
+
+		/// <summary>
+		/// The associated project for the issue
+		/// </summary>
+		public string? Project { get; set; }
+
+		/// <summary>
+		/// The summary text for this issue
+		/// </summary>
+		public string Summary { get; set; }
+
+		/// <summary>
+		/// Owner of the issue
+		/// </summary>
+		public string? Owner { get; set; }
+
+		/// <summary>
+		/// User that nominated the current owner
+		/// </summary>
+		public string? NominatedBy { get; set; }
+
+		/// <summary>
+		/// Time that the issue was acknowledged
+		/// </summary>
+		public DateTime? AcknowledgedAt { get; set; }
+
+		/// <summary>
+		/// Changelist that fixed this issue
+		/// </summary>
+		public int? FixChange { get; set; }
+
+		/// <summary>
+		/// Time at which the issue was resolved
+		/// </summary>
+		public DateTime? ResolvedAt { get; set; }
+
+		/// <summary>
+		/// Whether to notify the user about this issue
+		/// </summary>
+		public bool BNotify { get; set; }
+
+		/// <summary>
+		/// Whether this issue just contains warnings
+		/// </summary>
+		public bool BWarning { get; set; }
+
+		/// <summary>
+		/// Link to the last build
+		/// </summary>
+		public Uri? BuildUrl { get; set; }
+
+		/// <summary>
+		/// List of streams affected by this issue
+		/// </summary>
+		public List<string> Streams { get; set; }
+
+		/// <summary>
+		/// Constructs a new issue
+		/// </summary>
+		/// <param name="details">Issue to construct from</param>
+		/// <param name="owner"></param>
+		/// <param name="nominatedBy"></param>
+		/// <param name="bNotify">Whether to notify the user about this issue</param>
+		/// <param name="buildUrl">Link to the last build for this issue</param>
+		public GetUgsIssueResponse(IIssueDetails details, IUser? owner, IUser? nominatedBy, bool bNotify, Uri? buildUrl)
+		{
+			IIssue issue = details.Issue;
+			Id = issue.Id;
+			CreatedAt = issue.CreatedAt;
+			RetrievedAt = DateTime.UtcNow;
+			Summary = issue.Summary;
+			Owner = owner?.Login;
+			NominatedBy = nominatedBy?.Login;
+			AcknowledgedAt = issue.AcknowledgedAt;
+			FixChange = issue.FixChange;
+			ResolvedAt = issue.ResolvedAt;
+			BNotify = bNotify;
+			BWarning = issue.Severity == IssueSeverity.Warning;
+			BuildUrl = buildUrl;
+			Streams = details.Spans.Select(x => x.StreamName ?? x.StreamId.ToString()).Distinct().ToList();
+		}
+	}
+
+	/// <summary>
+	/// Request an issue to be updated
+	/// </summary>
+	public class UpdateUgsIssueRequest
+	{
+		/// <summary>
+		/// New owner of the issue
+		/// </summary>
+		public string? Owner { get; set; }
+
+		/// <summary>
+		/// User than nominates the new owner
+		/// </summary>
+		public string? NominatedBy { get; set; }
+
+		/// <summary>
+		/// Whether the issue has been acknowledged
+		/// </summary>
+		public bool? Acknowledged { get; set; }
+
+		/// <summary>
+		/// Name of the user that declines the issue
+		/// </summary>
+		public string? DeclinedBy { get; set; }
+
+		/// <summary>
+		/// The change at which the issue is claimed fixed. 0 = not fixed, -1 = systemic issue.
+		/// </summary>
+		public int? FixChange { get; set; }
+
+		/// <summary>
+		/// Whether the issue should be marked as resolved
+		/// </summary>
+		public bool? Resolved { get; set; }
+
+		/// <summary>
+		/// Name of the user that resolved the issue
+		/// </summary>
+		public string? ResolvedBy { get; set; }
+	}
+}
\ No newline at end of file
