Horde: Finish re-org of Horde solution started a while ago.

diff --git a/Engine/Source/Programs/Horde/Horde.Build/Ugs/UgsController.cs b/Engine/Source/Programs/Horde/Horde.Build/Ugs/UgsController.cs
new file mode 100644
--- /dev/null
+++ b/Engine/Source/Programs/Horde/Horde.Build/Ugs/UgsController.cs
@@ -0,0 +1,395 @@
+// Copyright Epic Games, Inc. All Rights Reserved.
+
+using System;
+using System.Collections.Generic;
+using System.Linq;
+using System.Threading.Tasks;
+using Horde.Build.Api;
+using Horde.Build.Collections;
+using Horde.Build.Models;
+using Horde.Build.Services;
+using Horde.Build.Utilities;
+using Microsoft.AspNetCore.Mvc;
+using Microsoft.Extensions.Logging;
+using Microsoft.Extensions.Options;
+using MongoDB.Bson;
+
+namespace Horde.Build.Controllers
+{
+	using LogId = ObjectId<ILogFile>;
+	using UserId = ObjectId<IUser>;
+
+	/// <summary>
+	/// Controller for the /api/v1/issues endpoint
+	/// </summary>
+	[ApiController]
+	[Route("[controller]")]
+	public sealed class UgsController : ControllerBase
+	{
+		/// <summary>
+		/// Singleton instance of the issue service
+		/// </summary>
+		private readonly IssueService _issueService;
+
+		/// <summary>
+		/// Collection of metadata documents
+		/// </summary>
+		private readonly IUgsMetadataCollection _ugsMetadataCollection;
+
+		/// <summary>
+		/// Collection of users
+		/// </summary>
+		private readonly IUserCollection _userCollection;
+
+		/// <summary>
+		/// The log file service
+		/// </summary>
+		private readonly ILogFileService _logFileService;
+
+		/// <summary>
+		/// Server settings
+		/// </summary>
+		private readonly ServerSettings _settings;
+
+		/// <summary>
+		/// Logger 
+		/// </summary>
+		readonly ILogger _logger;
+
+		/// <summary>
+		/// Constructor
+		/// </summary>
+		public UgsController(IssueService issueService, IUgsMetadataCollection ugsMetadataCollection, IUserCollection userCollection, ILogFileService logFileService, IOptionsMonitor<ServerSettings> optionsMonitor, ILogger<UgsController> logger)
+		{
+			_issueService = issueService;
+			_ugsMetadataCollection = ugsMetadataCollection;
+			_userCollection = userCollection;
+			_logFileService = logFileService;
+			_settings = optionsMonitor.CurrentValue;
+			_logger = logger;
+		}
+
+		/// <summary>
+		/// Gets the latest version info
+		/// </summary>
+		/// <returns>Result code</returns>
+		[HttpGet]
+		[Route("/ugs/api/latest")]
+		public ActionResult<object> GetLatest()
+		{
+			return new { Version = 2 };
+		}
+
+		/// <summary>
+		/// Adds new metadata to the database
+		/// </summary>
+		/// <param name="request">Request object</param>
+		/// <returns>Result code</returns>
+		[HttpPost]
+		[Route("/ugs/api/metadata")]
+		public async Task<ActionResult> AddMetadataAsync(AddUgsMetadataRequest request)
+		{
+			IUgsMetadata metadata = await _ugsMetadataCollection.FindOrAddAsync(request.Stream, request.Change, request.Project);
+			if (request.Synced != null || request.Vote != null || request.Investigating != null || request.Starred != null || request.Comment != null)
+			{
+				if (request.UserName == null)
+				{
+					return BadRequest("Missing UserName field on request body");
+				}
+				metadata = await _ugsMetadataCollection.UpdateUserAsync(metadata, request.UserName, request.Synced, request.Vote, request.Investigating, request.Starred, request.Comment);
+			}
+			if (request.Badges != null)
+			{
+				foreach (AddUgsBadgeRequest badge in request.Badges)
+				{
+					metadata = await _ugsMetadataCollection.UpdateBadgeAsync(metadata, badge.Name, badge.Url, badge.State);
+				}
+			}
+			return Ok();
+		}
+
+		/// <summary>
+		/// Searches for metadata updates
+		/// </summary>
+		/// <param name="stream">THe stream to search for</param>
+		/// <param name="minChange">Minimum changelist number</param>
+		/// <param name="maxChange">Maximum changelist number</param>
+		/// <param name="project">The project identifiers to search for</param>
+		/// <param name="sequence">Last sequence number</param>
+		/// <returns>List of metadata updates</returns>
+		[HttpGet]
+		[Route("/ugs/api/metadata")]
+		public async Task<GetUgsMetadataListResponse> FindMetadataAsync([FromQuery] string stream, [FromQuery] int minChange, [FromQuery] int? maxChange = null, [FromQuery] string? project = null, [FromQuery] long? sequence = null)
+		{
+			List<IUgsMetadata> metadataList = await _ugsMetadataCollection.FindAsync(stream, minChange, maxChange, sequence);
+		
+			GetUgsMetadataListResponse response = new GetUgsMetadataListResponse();
+			if(sequence != null)
+			{
+				response.SequenceNumber = sequence.Value;
+			}
+
+			foreach (IUgsMetadata metadata in metadataList)
+			{
+				if (metadata.UpdateTicks > response.SequenceNumber)
+				{
+					response.SequenceNumber = metadata.UpdateTicks;
+				}
+				if (String.IsNullOrEmpty(metadata.Project) || metadata.Project.Equals(project, StringComparison.OrdinalIgnoreCase))
+				{
+					response.Items.Add(new GetUgsMetadataResponse(metadata));
+				}
+			}
+			return response;
+		}
+
+		/// <summary>
+		/// Retrieve information about open issues
+		/// </summary>
+		/// <param name="user"></param>
+		/// <param name="includeResolved">Whether to include resolved issues</param>
+		/// <param name="maxResults">Maximum number of results to return</param>
+		/// <returns>List of matching agents</returns>
+		[HttpGet]
+		[Route("/ugs/api/issues")]
+		[ProducesResponseType(typeof(GetUgsIssueResponse), 200)]
+		public async Task<ActionResult<List<GetUgsIssueResponse>>> GetIssuesAsync([FromQuery] string? user = null, [FromQuery] bool includeResolved = false, [FromQuery] int maxResults = 100)
+		{
+			IUser? userInfo = (user != null) ? await _userCollection.FindUserByLoginAsync(user) : null;
+
+			List<GetUgsIssueResponse> responses = new List<GetUgsIssueResponse>();
+			if (includeResolved)
+			{
+				List<IIssue> issues = await _issueService.Collection.FindIssuesAsync(null, resolved: null, count: maxResults);
+				foreach(IIssue issue in issues)
+				{
+					IIssueDetails details = await _issueService.GetIssueDetailsAsync(issue);
+					bool bNotify = userInfo != null && details.Suspects.Any(x => x.AuthorId == userInfo.Id);
+					responses.Add(await CreateIssueResponseAsync(details, bNotify));
+				}
+			}
+			else
+			{
+				foreach (IIssueDetails cachedOpenIssue in _issueService.CachedOpenIssues)
+				{
+					if (responses.Count >= maxResults)
+					{
+						break;
+					}
+
+					if (cachedOpenIssue.ShowNotifications())
+					{
+						bool bNotify = userInfo != null && cachedOpenIssue.IncludeForUser(userInfo.Id);
+						responses.Add(await CreateIssueResponseAsync(cachedOpenIssue, bNotify));
+					}
+				}
+			}
+			return responses;
+		}
+
+		/// <summary>
+		/// Retrieve information about a specific issue
+		/// </summary>
+		/// <param name="issueId">Id of the issue to get information about</param>
+		/// <param name="filter">Filter for the properties to return</param>
+		/// <returns>List of matching agents</returns>
+		[HttpGet]
+		[Route("/ugs/api/issues/{issueId}")]
+		[ProducesResponseType(typeof(GetUgsIssueBuildResponse), 200)]
+		public async Task<ActionResult<object>> GetIssueAsync(int issueId, [FromQuery] PropertyFilter? filter = null)
+		{
+			IIssueDetails? issue = await _issueService.GetIssueDetailsAsync(issueId);
+			if (issue == null)
+			{
+				return NotFound();
+			}
+
+			return PropertyFilter.Apply(await CreateIssueResponseAsync(issue, false), filter);
+		}
+
+		/// <summary>
+		/// Retrieve information about builds for a specific issue
+		/// </summary>
+		/// <param name="issueId">Id of the issue to get information about</param>
+		/// <param name="filter">Filter for the properties to return</param>
+		/// <returns>List of matching agents</returns>
+		[HttpGet]
+		[Route("/ugs/api/issues/{issueId}/builds")]
+		[ProducesResponseType(typeof(List<GetUgsIssueBuildResponse>), 200)]
+		public async Task<ActionResult<List<object>>> GetIssueBuildsAsync(int issueId, [FromQuery] PropertyFilter? filter = null)
+		{
+			IIssueDetails? issue = await _issueService.GetCachedIssueDetailsAsync(issueId);
+			if (issue == null)
+			{
+				return NotFound();
+			}
+
+			List<GetUgsIssueBuildResponse> responses = new List<GetUgsIssueBuildResponse>();
+			foreach (IIssueSpan span in issue.Spans)
+			{
+				if (span.LastSuccess != null)
+				{
+					responses.Add(CreateBuildResponse(span, span.LastSuccess, IssueBuildOutcome.Success));
+				}
+				foreach (IIssueStep step in issue.Steps)
+				{
+					responses.Add(CreateBuildResponse(span, step, IssueBuildOutcome.Error));
+				}
+				if (span.NextSuccess != null)
+				{
+					responses.Add(CreateBuildResponse(span, span.NextSuccess, IssueBuildOutcome.Success));
+				}
+			}
+
+			return responses.ConvertAll(x => PropertyFilter.Apply(x, filter));
+		}
+
+		/// <summary>
+		/// Retrieve information about builds for a specific issue
+		/// </summary>
+		/// <param name="issueId">Id of the issue to get information about</param>
+		/// <returns>List of matching agents</returns>
+		[HttpGet]
+		[Route("/ugs/api/issues/{issueId}/diagnostics")]
+		[ProducesResponseType(typeof(List<GetUgsIssueDiagnosticResponse>), 200)]
+		public async Task<ActionResult<List<GetUgsIssueDiagnosticResponse>>> GetIssueDiagnosticsAsync(int issueId)
+		{
+			List<GetUgsIssueDiagnosticResponse> diagnostics = new List<GetUgsIssueDiagnosticResponse>();
+
+			Dictionary<LogId, ILogFile?> logFiles = new Dictionary<LogId, ILogFile?>();
+
+			List<IIssueSpan> spans = await _issueService.Collection.FindSpansAsync(issueId);
+			List<ILogEvent> events = await _logFileService.FindEventsForSpansAsync(spans.Select(x => x.Id), null, 0, count: 10);
+
+			foreach (ILogEvent logEvent in events)
+			{
+				ILogFile? logFile;
+				if(!logFiles.TryGetValue(logEvent.LogId, out logFile))
+				{
+					logFile = await _logFileService.GetLogFileAsync(logEvent.LogId);
+					logFiles.Add(logEvent.LogId, logFile);
+				}
+				if (logFile != null)
+				{
+					ILogEventData eventData = await _logFileService.GetEventDataAsync(logFile, logEvent.LineIndex, logEvent.LineCount);
+					long buildId = logEvent.LogId.GetHashCode();
+					Uri url = new Uri(_settings.DashboardUrl, $"log/{logEvent.LogId}?lineindex={logEvent.LineIndex}");
+
+					GetUgsIssueDiagnosticResponse diagnostic = new GetUgsIssueDiagnosticResponse(buildId, eventData.Message, url);
+					diagnostics.Add(diagnostic);
+				}
+			}
+
+			return diagnostics;
+		}
+
+		/// <summary>
+		/// Gets the URL for a failing step in the
+		/// </summary>
+		/// <param name="details">The issue to get a URL for</param>
+		/// <param name="bNotify">Whether to show notifications for this issue</param>
+		/// <returns>The issue response</returns>
+		async Task<GetUgsIssueResponse> CreateIssueResponseAsync(IIssueDetails details, bool bNotify)
+		{
+			Uri? buildUrl = GetIssueBuildUrl(details);
+
+			IUser? owner = details.Issue.OwnerId.HasValue ? await _userCollection.GetCachedUserAsync(details.Issue.OwnerId.Value) : null;
+			IUser? nominatedBy = details.Issue.NominatedById.HasValue ? await _userCollection.GetCachedUserAsync(details.Issue.NominatedById.Value) : null;
+
+			return new GetUgsIssueResponse(details, owner, nominatedBy, bNotify, buildUrl);
+		}
+
+		/// <summary>
+		/// Gets the URL for a failing step in the given issue
+		/// </summary>
+		/// <param name="issue">The issue details</param>
+		/// <returns>The build URL</returns>
+		Uri? GetIssueBuildUrl(IIssueDetails issue)
+		{
+			HashSet<ObjectId> unresolvedSpans = new HashSet<ObjectId>(issue.Spans.Where(x => x.NextSuccess == null).Select(x => x.Id));
+
+			IIssueStep? step = issue.Steps.OrderByDescending(x => unresolvedSpans.Contains(x.SpanId)).ThenByDescending(x => x.Change).FirstOrDefault();
+			if (step == null)
+			{
+				return null;
+			}
+
+			return new Uri(_settings.DashboardUrl, $"job/{step.JobId}?step={step.StepId}");
+		}
+
+		/// <summary>
+		/// Creates the response for a particular build
+		/// </summary>
+		/// <param name="span">Span containing the step</param>
+		/// <param name="step">The step to describe</param>
+		/// <param name="outcome">Outcome of this step</param>
+		/// <returns>Response object</returns>
+		GetUgsIssueBuildResponse CreateBuildResponse(IIssueSpan span, IIssueStep step, IssueBuildOutcome outcome)
+		{
+			GetUgsIssueBuildResponse response = new GetUgsIssueBuildResponse(span.StreamName, step.Change, outcome);
+			response.Id = step.LogId.GetHashCode();
+			response.JobName = $"{step.JobName}: {span.NodeName}";
+			response.JobUrl = new Uri(_settings.DashboardUrl, $"job/{step.JobId}");
+			response.JobStepName = span.NodeName;
+			response.JobStepUrl = new Uri(_settings.DashboardUrl, $"job/{step.JobId}?step={step.StepId}");
+			response.ErrorUrl = response.JobStepUrl;
+			return response;
+		}
+
+		/// <summary>
+		/// Update an issue
+		/// </summary>
+		/// <param name="issueId">Id of the issue to get information about</param>
+		/// <param name="request">The update information</param>
+		/// <returns>List of matching agents</returns>
+		[HttpPut]
+		[Route("/ugs/api/issues/{issueId}")]
+		public async Task<ActionResult> UpdateIssueAsync(int issueId, [FromBody] UpdateUgsIssueRequest request)
+		{
+			UserId? newOwnerId = null;
+			if (!String.IsNullOrEmpty(request.Owner))
+			{
+				newOwnerId = (await _userCollection.FindOrAddUserByLoginAsync(request.Owner))?.Id;
+			}
+
+			UserId? newNominatedById = null;
+			if (!String.IsNullOrEmpty(request.NominatedBy))
+			{
+				newNominatedById = (await _userCollection.FindOrAddUserByLoginAsync(request.NominatedBy))?.Id;
+			}
+
+			UserId? newDeclinedById = null;
+			if (!String.IsNullOrEmpty(request.DeclinedBy))
+			{
+				newDeclinedById = (await _userCollection.FindOrAddUserByLoginAsync(request.DeclinedBy))?.Id;
+			}
+
+			UserId? newResolvedById = null;
+			if (!String.IsNullOrEmpty(request.ResolvedBy))
+			{
+				newResolvedById = (await _userCollection.FindOrAddUserByLoginAsync(request.ResolvedBy))?.Id;
+			}
+			if (newResolvedById == null && request.Resolved.HasValue)
+			{
+				newResolvedById = request.Resolved.Value ? IIssue.ResolvedByUnknownId : UserId.Empty;
+			}
+
+			if (!await _issueService.UpdateIssueAsync(issueId, ownerId: newOwnerId, nominatedById: newNominatedById, acknowledged: request.Acknowledged, declinedById: newDeclinedById, fixChange: request.FixChange, resolvedById: newResolvedById))
+			{
+				return NotFound();
+			}
+			return Ok();
+		}
+
+		/// <summary>
+		/// Post information about net core installation. 
+		/// </summary>
+		[HttpPost]
+		[Route("/ugs/api/netcore")]
+		public ActionResult<object> PostNetCoreInfo(string? user = null, string? machine = null, bool netCore = false)
+		{
+			_logger.LogInformation("NetCore: User={User}, Machine={Machine}, NetCore={NetCore}", user, machine, netCore);
+			return new { };
+		}
+	}
+}
diff --git a/Engine/Source/Programs/Horde/Horde.Build/Users/IServiceAccountCollection.cs b/Engine/Source/Programs/Horde/Horde.Build/Users/IServiceAccountCollection.cs
new file mode 100644
--- /dev/null
+++ b/Engine/Source/Programs/Horde/Horde.Build/Users/IServiceAccountCollection.cs
@@ -0,0 +1,55 @@
+// Copyright Epic Games, Inc. All Rights Reserved.
+
+using System.Collections.Generic;
+using System.Threading.Tasks;
+using Horde.Build.Models;
+using MongoDB.Bson;
+
+namespace Horde.Build.Collections
+{
+	/// <summary>
+	/// Interface for a collection of service account documents
+	/// </summary>
+	public interface IServiceAccountCollection
+	{
+		/// <summary>
+		/// Adds a new service account to the collection
+		/// </summary>
+		/// <param name="secretToken">Secret token used for identifying this service account</param>
+		/// <param name="claims">Dictionary of claims to assign</param>
+		/// <param name="description">Description of the account</param>
+		Task<IServiceAccount> AddAsync(string secretToken, List<string> claims, string description);
+
+		/// <summary>
+		/// Get service account via ID
+		/// </summary>
+		/// <param name="id">The unique service account id</param>
+		/// <returns>The service account</returns>
+		Task<IServiceAccount?> GetAsync(ObjectId id);
+		
+		/// <summary>
+		/// Get service account via secret token
+		/// </summary>
+		/// <param name="secretToken">Secret token to use for searching</param>
+		/// <returns>The service account</returns>
+		Task<IServiceAccount?> GetBySecretTokenAsync(string secretToken);
+
+		/// <summary>
+		/// Update a service account from the collection
+		/// </summary>
+		/// <param name="id">The service account</param>
+		/// <param name="secretToken">If set, secret token will be set</param>
+		/// <param name="claims">If set, claims will be set</param>
+		/// <param name="enabled">If set, enabled flag will be set</param>
+		/// <param name="description">If set, description will be set</param>
+		/// <returns>Async task</returns>
+		Task UpdateAsync(ObjectId id, string? secretToken, List<string>? claims, bool? enabled, string? description);
+
+		/// <summary>
+		/// Delete a service account from the collection
+		/// </summary>
+		/// <param name="id">The service account</param>
+		/// <returns>Async task</returns>
+		Task DeleteAsync(ObjectId id);
+	}
+}
diff --git a/Engine/Source/Programs/Horde/Horde.Build/Users/ServiceAccountCollection.cs b/Engine/Source/Programs/Horde/Horde.Build/Users/ServiceAccountCollection.cs
new file mode 100644
--- /dev/null
+++ b/Engine/Source/Programs/Horde/Horde.Build/Users/ServiceAccountCollection.cs
@@ -0,0 +1,165 @@
+// Copyright Epic Games, Inc. All Rights Reserved.
+
+using System;
+using System.Collections.Generic;
+using System.Linq;
+using System.Threading.Tasks;
+using Horde.Build.Models;
+using Horde.Build.Server;
+using MongoDB.Bson;
+using MongoDB.Bson.Serialization.Attributes;
+using MongoDB.Driver;
+
+namespace Horde.Build.Collections.Impl
+{
+	/// <summary>
+	/// Collection of service account documents
+	/// </summary>
+	public class ServiceAccountCollection : IServiceAccountCollection
+	{
+		/// <summary>
+		/// Concrete implementation of IServiceAccount
+		/// </summary>
+		class ServiceAccountDocument : IServiceAccount
+		{
+			public const string ClaimSeparator = "###";
+			
+			[BsonRequired, BsonId]
+			public ObjectId Id { get; set; }
+
+			[BsonRequired]
+			public string SecretToken { get; set; } = "empty";
+
+			[BsonRequired]
+			public List<string> Claims { get; set; } = new List<string>();
+			
+			public bool Enabled { get; set; }
+			public string Description { get; set; } = "empty";
+
+			[BsonConstructor]
+			private ServiceAccountDocument()
+			{
+			}
+
+			public ServiceAccountDocument(ObjectId id, string secretToken, List<string> claims, bool enabled, string description)
+			{
+				Id = id;
+				SecretToken = secretToken;
+				Claims = claims;
+				Enabled = enabled;
+				Description = description;
+			}
+			
+			/// <inheritdoc/>
+			public void AddClaim(string type, string value)
+			{
+				Claims.Add(type + ClaimSeparator + value);
+			}
+
+			/// <inheritdoc/>
+			public IReadOnlyList<(string Type, string Value)> GetClaims()
+			{
+				return Claims.Select(x =>
+				{
+					string[] split = x.Split(ClaimSeparator);
+					return (split[0], split[1]);
+				}).ToList();
+			}
+
+			protected bool Equals(ServiceAccountDocument other)
+			{
+				bool areClaimsEqual = !Claims.Except(other.Claims).Any();
+
+                return Id.Equals(other.Id) && SecretToken == other.SecretToken && areClaimsEqual && Enabled == other.Enabled && Description == other.Description;
+			}
+
+			public override bool Equals(object? obj)
+			{
+				if (obj is null)
+				{
+					return false;
+				}
+				if (ReferenceEquals(this, obj))
+				{
+					return true;
+				}
+				if (obj.GetType() != GetType())
+				{
+					return false;
+				}
+				return Equals((ServiceAccountDocument) obj);
+			}
+
+			public override int GetHashCode()
+			{
+				return HashCode.Combine(Id, SecretToken, Claims, Enabled, Description);
+			}
+		}
+
+		/// <summary>
+		/// Collection of session documents
+		/// </summary>
+		private readonly IMongoCollection<ServiceAccountDocument> _serviceAccounts;
+
+		/// <summary>
+		/// Constructor
+		/// </summary>
+		/// <param name="mongoService">The database service</param>
+		public ServiceAccountCollection(MongoService mongoService)
+		{
+			_serviceAccounts = mongoService.GetCollection<ServiceAccountDocument>("ServiceAccounts", keys => keys.Ascending(x => x.SecretToken));
+		}
+
+		/// <inheritdoc/>
+		public async Task<IServiceAccount> AddAsync(string secretToken, List<string> claims, string description)
+		{
+			ServiceAccountDocument newSession = new ServiceAccountDocument(ObjectId.GenerateNewId(), secretToken, claims, true, description);
+			await _serviceAccounts.InsertOneAsync(newSession);
+			return newSession;
+		}
+
+		/// <inheritdoc/>
+		public async Task<IServiceAccount?> GetAsync(ObjectId id)
+		{
+			return await _serviceAccounts.Find(x => x.Id == id).FirstOrDefaultAsync();
+		}
+
+		/// <inheritdoc/>
+		public async Task<IServiceAccount?> GetBySecretTokenAsync(string secretToken)
+		{
+			return await _serviceAccounts.Find(x => x.SecretToken == secretToken).FirstOrDefaultAsync();
+		}
+
+		/// <inheritdoc/>
+		public Task UpdateAsync(ObjectId id, string? secretToken, List<string>? claims, bool? enabled, string? description)
+		{
+			UpdateDefinitionBuilder<ServiceAccountDocument> update = Builders<ServiceAccountDocument>.Update;
+			List<UpdateDefinition<ServiceAccountDocument>> updates = new List<UpdateDefinition<ServiceAccountDocument>>();
+
+			if (secretToken != null)
+			{
+				updates.Add(update.Set(x => x.SecretToken, secretToken));
+			}
+			if (claims != null)
+			{
+				updates.Add(update.Set(x => x.Claims, claims));
+			}
+			if (enabled != null)
+			{
+				updates.Add(update.Set(x => x.Enabled, enabled));
+			}
+			if (description != null)
+			{
+				updates.Add(update.Set(x => x.Description, description));
+			}
+
+			return _serviceAccounts.FindOneAndUpdateAsync(x => x.Id == id, update.Combine(updates));
+		}
+
+		/// <inheritdoc/>
+		public Task DeleteAsync(ObjectId sessionId)
+		{
+			return _serviceAccounts.DeleteOneAsync(x => x.Id == sessionId);
+		}
+	}
+}
diff --git a/Engine/Source/Programs/Horde/Horde.Build/Users/UserCollection.cs b/Engine/Source/Programs/Horde/Horde.Build/Users/UserCollection.cs
new file mode 100644
--- /dev/null
+++ b/Engine/Source/Programs/Horde/Horde.Build/Users/UserCollection.cs
@@ -0,0 +1,207 @@
+// Copyright Epic Games, Inc. All Rights Reserved.
+
+using System;
+using System.Collections.Generic;
+using System.Linq;
+using System.Security.Claims;
+using System.Threading.Tasks;
+using Horde.Build.Models;
+using Horde.Build.Server;
+using Horde.Build.Utilities;
+using MongoDB.Bson;
+using MongoDB.Bson.Serialization;
+using MongoDB.Bson.Serialization.Attributes;
+using MongoDB.Driver;
+
+namespace Horde.Build.Collections.Impl
+{
+	using JobId = ObjectId<IJob>;
+	using UserId = ObjectId<IUser>;
+
+	/// <summary>
+	/// Manages user documents
+	/// </summary>
+	class UserCollectionV1 : IUserCollection
+	{
+		class UserDocument : IUser, IUserClaims, IUserSettings
+		{
+			public UserId Id { get; set; }
+
+			public ClaimDocument PrimaryClaim { get; set; } = null!;
+			public List<ClaimDocument> Claims { get; set; } = new List<ClaimDocument>();
+
+			[BsonDefaultValue(false), BsonIgnoreIfDefault]
+			public bool EnableExperimentalFeatures { get; set; }
+
+			public BsonValue DashboardSettings { get; set; } = BsonNull.Value;
+			public List<JobId> PinnedJobIds { get; set; } = new List<JobId>();
+
+			string IUser.Name => Claims.FirstOrDefault(x => String.Equals(x.Type, "name", StringComparison.Ordinal))?.Value ?? PrimaryClaim.Value;
+			string IUser.Login => Claims.FirstOrDefault(x => String.Equals(x.Type, ClaimTypes.Name, StringComparison.Ordinal))?.Value ?? PrimaryClaim.Value;
+			string? IUser.Email => Claims.FirstOrDefault(x => String.Equals(x.Type, ClaimTypes.Email, StringComparison.Ordinal))?.Value;
+
+			UserId IUserClaims.UserId => Id;
+			IReadOnlyList<IUserClaim> IUserClaims.Claims => Claims;
+
+			UserId IUserSettings.UserId => Id;
+			IReadOnlyList<JobId> IUserSettings.PinnedJobIds => PinnedJobIds;
+		}
+
+		class ClaimDocument : IUserClaim
+		{
+			public string Type { get; set; }
+			public string Value { get; set; }
+
+			private ClaimDocument()
+			{
+				Type = null!;
+				Value = null!;
+			}
+
+			public ClaimDocument(string type, string value)
+			{
+				Type = type;
+				Value = value;
+			}
+
+			public ClaimDocument(IUserClaim other)
+			{
+				Type = other.Type;
+				Value = other.Value;
+			}
+		}
+
+		readonly IMongoCollection<UserDocument> _users;
+
+		/// <summary>
+		/// Static constructor
+		/// </summary>
+		static UserCollectionV1()
+		{
+			BsonSerializer.RegisterDiscriminatorConvention(typeof(UserDocument), NullDiscriminatorConvention.Instance);
+			BsonSerializer.RegisterDiscriminatorConvention(typeof(ClaimDocument), NullDiscriminatorConvention.Instance);
+		}
+
+		/// <summary>
+		/// Constructor
+		/// </summary>
+		/// <param name="mongoService"></param>
+		public UserCollectionV1(MongoService mongoService)
+		{
+			_users = mongoService.GetCollection<UserDocument>("Users", keys => keys.Ascending(x => x.PrimaryClaim), unique: true);
+		}
+
+		/// <inheritdoc/>
+		public async Task<IUser?> GetUserAsync(UserId id)
+		{
+			return await _users.Find(x => x.Id == id).FirstOrDefaultAsync();
+		}
+
+		/// <inheritdoc/>
+		public async ValueTask<IUser?> GetCachedUserAsync(UserId? id)
+		{
+			if (id == null)
+			{
+				return null;
+			}
+			else
+			{
+				return await GetUserAsync(id.Value);
+			}
+		}
+
+		/// <inheritdoc/>
+		public async Task<List<IUser>> FindUsersAsync(IEnumerable<UserId>? ids, string? nameRegex, int? index, int? count)
+		{
+			FilterDefinition<UserDocument> filter = Builders<UserDocument>.Filter.In(x => x.Id, ids);
+			return await _users.Find(filter).Range(index, count).ToListAsync<UserDocument, IUser>();
+		}
+
+		/// <inheritdoc/>
+		public async Task<IUser?> FindUserByLoginAsync(string login)
+		{
+			ClaimDocument primaryClaim = new ClaimDocument(ClaimTypes.Name, login);
+			return await _users.Find(x => x.PrimaryClaim == primaryClaim).FirstOrDefaultAsync();
+		}
+
+		/// <inheritdoc/>
+		public async Task<IUser> FindOrAddUserByLoginAsync(string login, string? name, string? email)
+		{
+			ClaimDocument newPrimaryClaim = new ClaimDocument(ClaimTypes.Name, login);
+			UpdateDefinition<UserDocument> update = Builders<UserDocument>.Update.SetOnInsert(x => x.Id, UserId.GenerateNewId());
+			return await _users.FindOneAndUpdateAsync<UserDocument>(x => x.PrimaryClaim == newPrimaryClaim, update, new FindOneAndUpdateOptions<UserDocument> { IsUpsert = true, ReturnDocument = ReturnDocument.After });
+		}
+
+		/// <inheritdoc/>
+		public async Task<IUserClaims> GetClaimsAsync(UserId userId)
+		{
+			return await _users.Find(x => x.Id == userId).FirstOrDefaultAsync() ?? new UserDocument { Id = userId };
+		}
+
+		/// <inheritdoc/>
+		public async Task UpdateClaimsAsync(UserId userId, IEnumerable<IUserClaim> claims)
+		{
+			List<ClaimDocument> newClaims = claims.Select(x => new ClaimDocument(x)).ToList();
+			await _users.FindOneAndUpdateAsync(x => x.Id == userId, Builders<UserDocument>.Update.Set(x => x.Claims, newClaims));
+		}
+
+		/// <inheritdoc/>
+		public async Task<IUserSettings> GetSettingsAsync(UserId userId)
+		{
+			return await _users.Find(x => x.Id == userId).FirstOrDefaultAsync() ?? new UserDocument { Id = userId };
+		}
+
+		/// <inheritdoc/>
+		public async Task UpdateSettingsAsync(UserId userId, bool? enableExperimentalFeatures, BsonValue? dashboardSettings = null, IEnumerable<JobId>? addPinnedJobIds = null, IEnumerable<JobId>? removePinnedJobIds = null)
+		{
+			if (addPinnedJobIds != null)
+			{
+				foreach (JobId pinnedJobId in addPinnedJobIds)
+				{
+					FilterDefinition<UserDocument> filter = Builders<UserDocument>.Filter.Eq(x => x.Id, userId) & Builders<UserDocument>.Filter.AnyNin<JobId>(x => x.PinnedJobIds, new[] { pinnedJobId });
+					UpdateDefinition<UserDocument> update = Builders<UserDocument>.Update.PushEach(x => x.PinnedJobIds, new[] { pinnedJobId }, -50);
+					await _users.UpdateOneAsync(filter, update);
+				}
+			}
+
+			List<UpdateDefinition<UserDocument>> updates = new List<UpdateDefinition<UserDocument>>();
+			if (enableExperimentalFeatures != null)
+			{
+				updates.Add(Builders<UserDocument>.Update.SetOrUnsetNull(x => x.EnableExperimentalFeatures, enableExperimentalFeatures));
+			}
+			if (dashboardSettings != null)
+			{
+				updates.Add(Builders<UserDocument>.Update.Set(x => x.DashboardSettings, dashboardSettings));
+			}
+			if (removePinnedJobIds != null && removePinnedJobIds.Any())
+			{
+				updates.Add(Builders<UserDocument>.Update.PullAll(x => x.PinnedJobIds, removePinnedJobIds));
+			}
+			if (updates.Count > 0)
+			{
+				await _users.UpdateOneAsync<UserDocument>(x => x.Id == userId, Builders<UserDocument>.Update.Combine(updates));
+			}
+		}
+
+		/// <summary>
+		/// Enumerate all the documents in this collection
+		/// </summary>
+		/// <returns></returns>
+		public async IAsyncEnumerable<(IUser, IUserClaims, IUserSettings)> EnumerateDocumentsAsync()
+		{
+			using (IAsyncCursor<UserDocument> cursor = await _users.Find(FilterDefinition<UserDocument>.Empty).ToCursorAsync())
+			{
+				while (await cursor.MoveNextAsync())
+				{
+					foreach (UserDocument document in cursor.Current)
+					{
+						if (document.Claims.Count > 0)
+						{
+							yield return (document, document, document);
+						}
+					}
+				}
+			}
+		}
+	}
+}
diff --git a/Engine/Source/Programs/Horde/Horde.Build/Users/UserCollectionV2.cs b/Engine/Source/Programs/Horde/Horde.Build/Users/UserCollectionV2.cs
new file mode 100644
--- /dev/null
+++ b/Engine/Source/Programs/Horde/Horde.Build/Users/UserCollectionV2.cs
@@ -0,0 +1,367 @@
+// Copyright Epic Games, Inc. All Rights Reserved.
+
+using System;
+using System.Collections.Generic;
+using System.Linq;
+using System.Threading.Tasks;
+using Horde.Build.Models;
+using Horde.Build.Server;
+using Horde.Build.Utilities;
+using Microsoft.Extensions.Caching.Memory;
+using Microsoft.Extensions.Logging;
+using MongoDB.Bson;
+using MongoDB.Bson.Serialization;
+using MongoDB.Bson.Serialization.Attributes;
+using MongoDB.Driver;
+
+namespace Horde.Build.Collections.Impl
+{
+	using JobId = ObjectId<IJob>;
+	using UserId = ObjectId<IUser>;
+
+	/// <summary>
+	/// Manages user documents
+	/// </summary>
+	class UserCollectionV2 : IUserCollection, IDisposable
+	{
+		class UserDocument : IUser
+		{
+			public UserId Id { get; set; }
+
+			public string Name { get; set; }
+			public string Login { get; set; }
+			public string LoginUpper { get; set; }
+			public string? Email { get; set; }
+			public string? EmailUpper { get; set; }
+
+			[BsonIgnoreIfDefault]
+			public bool? Hidden { get; set; }
+
+			[BsonConstructor]
+			private UserDocument()
+			{
+				Name = null!;
+				Login = null!;
+				LoginUpper = null!;
+			}
+
+			public UserDocument(IUser other)
+				: this(other.Id, other.Name, other.Login, other.Email)
+			{
+			}
+
+			public UserDocument(UserId id, string name, string login, string? email)
+			{
+				Id = id;
+				Name = name;
+				Login = login;
+				LoginUpper = login.ToUpperInvariant();
+				Email = email;
+				EmailUpper = email?.ToUpperInvariant();
+			}
+		}
+
+		class UserClaimsDocument : IUserClaims
+		{
+			public UserId Id { get; set; }
+			public List<UserClaim> Claims { get; set; } = new List<UserClaim>();
+
+			UserId IUserClaims.UserId => Id;
+			IReadOnlyList<IUserClaim> IUserClaims.Claims => Claims;
+
+			[BsonConstructor]
+			private UserClaimsDocument()
+			{
+			}
+
+			public UserClaimsDocument(UserId id)
+			{
+				Id = id;
+			}
+
+			public UserClaimsDocument(IUserClaims other)
+				: this(other.UserId)
+			{
+				Claims.AddRange(other.Claims.Select(x => new UserClaim(x)));
+			}
+		}
+
+		class UserSettingsDocument : IUserSettings
+		{
+			public UserId Id { get; set; }
+
+			[BsonDefaultValue(false), BsonIgnoreIfDefault]
+			public bool EnableExperimentalFeatures { get; set; }
+
+			public BsonValue DashboardSettings { get; set; } = BsonNull.Value;
+			public List<JobId> PinnedJobIds { get; set; } = new List<JobId>();
+
+			UserId IUserSettings.UserId => Id;
+			IReadOnlyList<JobId> IUserSettings.PinnedJobIds => PinnedJobIds;
+
+			[BsonConstructor]
+			private UserSettingsDocument()
+			{
+			}
+
+			public UserSettingsDocument(UserId id)
+			{
+				Id = id;
+			}
+
+			public UserSettingsDocument(IUserSettings other)
+				: this(other.UserId)
+			{
+				EnableExperimentalFeatures = other.EnableExperimentalFeatures;
+				DashboardSettings = other.DashboardSettings;
+				PinnedJobIds = new List<JobId>(other.PinnedJobIds);
+			}
+		}
+
+		class ClaimDocument : IUserClaim
+		{
+			public string Type { get; set; }
+			public string Value { get; set; }
+
+			private ClaimDocument()
+			{
+				Type = null!;
+				Value = null!;
+			}
+
+			public ClaimDocument(string type, string value)
+			{
+				Type = type;
+				Value = value;
+			}
+
+			public ClaimDocument(IUserClaim other)
+			{
+				Type = other.Type;
+				Value = other.Value;
+			}
+		}
+
+		readonly IMongoCollection<UserDocument> _users;
+		readonly IMongoCollection<UserClaimsDocument> _userClaims;
+		readonly IMongoCollection<UserSettingsDocument> _userSettings;
+		readonly ILogger<UserCollectionV2> _logger;
+		readonly MemoryCache _userCache;
+
+		/// <summary>
+		/// Static constructor
+		/// </summary>
+		static UserCollectionV2()
+		{
+			BsonSerializer.RegisterDiscriminatorConvention(typeof(UserDocument), NullDiscriminatorConvention.Instance);
+			BsonSerializer.RegisterDiscriminatorConvention(typeof(UserClaimsDocument), NullDiscriminatorConvention.Instance);
+			BsonSerializer.RegisterDiscriminatorConvention(typeof(UserSettingsDocument), NullDiscriminatorConvention.Instance);
+
+			BsonSerializer.RegisterDiscriminatorConvention(typeof(ClaimDocument), NullDiscriminatorConvention.Instance);
+		}
+
+		/// <summary>
+		/// Constructor
+		/// </summary>
+		/// <param name="mongoService"></param>
+		/// <param name="logger"></param>
+		public UserCollectionV2(MongoService mongoService, ILogger<UserCollectionV2> logger)
+		{
+			_logger = logger;
+
+			List<MongoIndex<UserDocument>> userIndexes = new List<MongoIndex<UserDocument>>();
+			userIndexes.Add(keys => keys.Ascending(x => x.LoginUpper), unique: true);
+			userIndexes.Add(keys => keys.Ascending(x => x.EmailUpper));
+			_users = mongoService.GetCollection<UserDocument>("UsersV2", userIndexes);
+
+			_userClaims = mongoService.GetCollection<UserClaimsDocument>("UserClaimsV2");
+			_userSettings = mongoService.GetCollection<UserSettingsDocument>("UserSettingsV2");
+
+			MemoryCacheOptions options = new MemoryCacheOptions();
+			_userCache = new MemoryCache(options);
+		}
+
+		/// <inheritdoc/>
+		public void Dispose()
+		{
+			_userCache.Dispose();
+		}
+
+		/// <inheritdoc/>
+		public async Task<IUser?> GetUserAsync(UserId id)
+		{
+			IUser? user = await _users.Find(x => x.Id == id).FirstOrDefaultAsync();
+			using (ICacheEntry entry = _userCache.CreateEntry(id))
+			{
+				entry.SetValue(user);
+				entry.SetSlidingExpiration(TimeSpan.FromMinutes(30.0));
+			}
+			return user;
+		}
+
+		/// <inheritdoc/>
+		public async ValueTask<IUser?> GetCachedUserAsync(UserId? id)
+		{
+			IUser? user;
+			if(id == null)
+			{
+				return null;
+			}
+			else if (_userCache.TryGetValue(id.Value, out user))
+			{
+				return user;
+			}
+			else
+			{
+				return await GetUserAsync(id.Value);
+			}
+		}
+
+		/// <inheritdoc/>
+		public async Task<List<IUser>> FindUsersAsync(IEnumerable<UserId>? ids, string? nameRegex = null, int? index = null, int? count = null)
+		{
+			FilterDefinition<UserDocument> filter = FilterDefinition<UserDocument>.Empty;
+			if (ids != null)
+			{
+				filter &= Builders<UserDocument>.Filter.In(x => x.Id, ids);
+			}
+			
+			if (nameRegex != null)
+			{
+				BsonRegularExpression regex = new BsonRegularExpression(nameRegex, "i");
+				filter &= Builders<UserDocument>.Filter.Regex(x => x.Name, regex);
+			}
+
+			filter &= Builders<UserDocument>.Filter.Ne(x => x.Hidden, true);
+
+			return await _users.Find(filter).Range(index, count ?? 100).ToListAsync<UserDocument, IUser>();
+		}
+
+		/// <inheritdoc/>
+		public async Task<IUser?> FindUserByLoginAsync(string login)
+		{
+			string loginUpper = login.ToUpperInvariant();
+			return await _users.Find(x => x.LoginUpper == loginUpper).FirstOrDefaultAsync();
+		}
+
+		/// <inheritdoc/>
+		public async Task<IUser> FindOrAddUserByLoginAsync(string login, string? name, string? email)
+		{
+			if (login.Contains(' ', StringComparison.Ordinal))
+			{
+				_logger.LogWarning("Potentially invalid login name: {Login} (from {Trace})", login, Environment.StackTrace);
+			}
+
+			UserId newUserId = UserId.GenerateNewId();
+			UpdateDefinition<UserDocument> update = Builders<UserDocument>.Update.SetOnInsert(x => x.Id, newUserId).SetOnInsert(x => x.Login, login).Unset(x => x.Hidden);
+
+			if (name == null)
+			{
+				update = update.SetOnInsert(x => x.Name, login);
+			}
+			else
+			{
+				update = update.Set(x => x.Name, name);
+			}
+
+			if (email != null)
+			{
+				update = update.Set(x => x.Email, email).Set(x => x.EmailUpper, email.ToUpperInvariant());
+			}
+
+			string loginUpper = login.ToUpperInvariant();
+
+			IUser user = await _users.FindOneAndUpdateAsync<UserDocument>(x => x.LoginUpper == loginUpper, update, new FindOneAndUpdateOptions<UserDocument, UserDocument> { IsUpsert = true, ReturnDocument = ReturnDocument.After });
+			if (user.Id == newUserId)
+			{
+				_logger.LogInformation("Added new user {Name} ({UserId}, {Login}, {Email})", user.Name, user.Id, user.Login, user.Email);
+			}
+
+			return user;
+		}
+
+		/// <inheritdoc/>
+		public async Task<IUserClaims> GetClaimsAsync(UserId userId)
+		{
+			IUserClaims? claims = await _userClaims.Find(x => x.Id == userId).FirstOrDefaultAsync();
+			if(claims == null)
+			{
+				claims = new UserClaimsDocument(userId);
+			}
+			return claims;
+		}
+
+		/// <inheritdoc/>
+		public async Task UpdateClaimsAsync(UserId userId, IEnumerable<IUserClaim> claims)
+		{
+			UserClaimsDocument newDocument = new UserClaimsDocument(userId);
+			newDocument.Claims.AddRange(claims.Select(x => new UserClaim(x)));
+			await _userClaims.ReplaceOneAsync(x => x.Id == userId, newDocument, new ReplaceOptions { IsUpsert = true });
+		}
+
+		/// <inheritdoc/>
+		public async Task<IUserSettings> GetSettingsAsync(UserId userId)
+		{
+			IUserSettings? settings = await _userSettings.Find(x => x.Id == userId).FirstOrDefaultAsync();
+			if (settings == null)
+			{
+				settings = new UserSettingsDocument(userId);
+			}
+			return settings;
+		}
+
+		/// <inheritdoc/>
+		public async Task UpdateSettingsAsync(UserId userId, bool? enableExperimentalFeatures = null, BsonValue? dashboardSettings = null, IEnumerable<JobId>? addPinnedJobIds = null, IEnumerable<JobId>? removePinnedJobIds = null)
+		{
+			List<UpdateDefinition<UserSettingsDocument>> updates = new List<UpdateDefinition<UserSettingsDocument>>();
+			if (enableExperimentalFeatures != null)
+			{
+				updates.Add(Builders<UserSettingsDocument>.Update.SetOrUnsetNull(x => x.EnableExperimentalFeatures, enableExperimentalFeatures));
+			}
+			if (dashboardSettings != null)
+			{
+				updates.Add(Builders<UserSettingsDocument>.Update.Set(x => x.DashboardSettings, dashboardSettings));
+			}
+			if (addPinnedJobIds != null && addPinnedJobIds.Any())
+			{
+				updates.Add(Builders<UserSettingsDocument>.Update.AddToSetEach(x => x.PinnedJobIds, addPinnedJobIds));
+			}
+			if (removePinnedJobIds != null && removePinnedJobIds.Any())
+			{
+				updates.Add(Builders<UserSettingsDocument>.Update.PullAll(x => x.PinnedJobIds, removePinnedJobIds));
+			}
+			if (updates.Count > 0)
+			{
+				await _userSettings.UpdateOneAsync<UserSettingsDocument>(x => x.Id == userId, Builders<UserSettingsDocument>.Update.Combine(updates), new UpdateOptions { IsUpsert = true });
+			}
+		}
+
+		/// <summary>
+		/// Upgrade from V1 collection
+		/// </summary>
+		/// <param name="userCollectionV1"></param>
+		/// <returns></returns>
+		public async Task ResaveDocumentsAsync(UserCollectionV1 userCollectionV1)
+		{
+			await foreach ((IUser user, IUserClaims claims, IUserSettings settings) in userCollectionV1.EnumerateDocumentsAsync())
+			{
+				try
+				{
+					await _users.ReplaceOneAsync(x => x.Id == user.Id, new UserDocument(user), new ReplaceOptions { IsUpsert = true });
+					await _userClaims.ReplaceOneAsync(x => x.Id == user.Id, new UserClaimsDocument(claims), new ReplaceOptions { IsUpsert = true });
+					await _userSettings.ReplaceOneAsync(x => x.Id == user.Id, new UserSettingsDocument(settings), new ReplaceOptions { IsUpsert = true });
+					_logger.LogDebug("Updated user {UserId}", user.Id);
+				}
+				catch (MongoWriteException ex)
+				{
+					_logger.LogWarning(ex, "Unable to resave user {UserId}", user.Id);
+				}
+
+				if(settings.PinnedJobIds.Count > 0)
+				{
+					await UpdateSettingsAsync(user.Id, addPinnedJobIds: settings.PinnedJobIds);
+				}
+			}
+		}
+	}
+}
diff --git a/Engine/Source/Programs/Horde/Horde.Build/Users/UserController.cs b/Engine/Source/Programs/Horde/Horde.Build/Users/UserController.cs
new file mode 100644
--- /dev/null
+++ b/Engine/Source/Programs/Horde/Horde.Build/Users/UserController.cs
@@ -0,0 +1,86 @@
+// Copyright Epic Games, Inc. All Rights Reserved.
+
+using System.Collections.Generic;
+using System.Linq;
+using System.Threading.Tasks;
+using Horde.Build.Api;
+using Horde.Build.Collections;
+using Horde.Build.Models;
+using Horde.Build.Services;
+using Horde.Build.Utilities;
+using Microsoft.AspNetCore.Authorization;
+using Microsoft.AspNetCore.Mvc;
+
+namespace Horde.Build.Controllers
+{
+	using JobId = ObjectId<IJob>;
+	using UserId = ObjectId<IUser>;
+
+	/// <summary>
+	/// Controller for the /api/v1/user endpoint
+	/// </summary>
+	[ApiController]
+	[Authorize]
+	[Route("[controller]")]
+	public class UserController : ControllerBase
+	{
+		/// <summary>
+		/// The user collection instance
+		/// </summary>
+		IUserCollection UserCollection { get; set; }
+
+		/// <summary>
+		/// The avatar service
+		/// </summary>
+		IAvatarService? AvatarService { get; set; }
+
+		/// <summary>
+		/// Constructor
+		/// </summary>
+		/// <param name="userCollection"></param>
+		/// <param name="avatarService"></param>
+		public UserController(IUserCollection userCollection, IAvatarService? avatarService)
+		{
+			UserCollection = userCollection;
+			AvatarService = avatarService;
+		}
+
+		/// <summary>
+		/// Gets information about the logged in user
+		/// </summary>
+		/// <returns>Http result code</returns>
+		[HttpGet]
+		[Route("/api/v1/user")]
+		[ProducesResponseType(typeof(List<GetUserResponse>), 200)]
+		public async Task<ActionResult<object>> GetUserAsync([FromQuery] PropertyFilter? filter = null)
+		{
+			IUser? internalUser = await UserCollection.GetUserAsync(User);
+			if (internalUser == null)
+			{
+				return NotFound();
+			}
+
+			IAvatar? avatar = (AvatarService == null)? (IAvatar?)null : await AvatarService.GetAvatarAsync(internalUser);
+			IUserClaims claims = await UserCollection.GetClaimsAsync(internalUser.Id);
+			IUserSettings settings = await UserCollection.GetSettingsAsync(internalUser.Id);
+			return PropertyFilter.Apply(new GetUserResponse(internalUser, avatar, claims, settings), filter);
+		}
+
+		/// <summary>
+		/// Updates the logged in user
+		/// </summary>
+		/// <returns>Http result code</returns>
+		[HttpPut]
+		[Route("/api/v1/user")]
+		public async Task<ActionResult> UpdateUserAsync(UpdateUserRequest request)
+		{
+			UserId? userId = User.GetUserId();
+			if(userId == null)
+			{
+				return BadRequest("Current user does not have a registered profile");
+			}
+			await UserCollection.UpdateSettingsAsync(userId.Value, request.EnableExperimentalFeatures, request.DashboardSettings?.ToBsonValue(), request.AddPinnedJobIds?.Select(x => new JobId(x)), request.RemovePinnedJobIds?.Select(x => new JobId(x)));
+			return Ok();
+		}
+	}
+}
diff --git a/Engine/Source/Programs/Horde/Horde.Build/Users/UsersController.cs b/Engine/Source/Programs/Horde/Horde.Build/Users/UsersController.cs
new file mode 100644
--- /dev/null
+++ b/Engine/Source/Programs/Horde/Horde.Build/Users/UsersController.cs
@@ -0,0 +1,123 @@
+// Copyright Epic Games, Inc. All Rights Reserved.
+
+using System;
+using System.Collections.Generic;
+using System.Linq;
+using System.Threading.Tasks;
+using Horde.Build.Api;
+using Horde.Build.Collections;
+using Horde.Build.Models;
+using Horde.Build.Services;
+using Horde.Build.Utilities;
+using Microsoft.AspNetCore.Authorization;
+using Microsoft.AspNetCore.Mvc;
+
+namespace Horde.Build.Controllers
+{
+	using UserId = ObjectId<IUser>;
+
+	/// <summary>
+	/// Controller for the /api/v1/users endpoint
+	/// </summary>
+	[ApiController]
+	[Authorize]
+	[Route("[controller]")]
+	public class UsersController : HordeControllerBase
+	{
+		/// <summary>
+		/// The user collection instance
+		/// </summary>
+		IUserCollection UserCollection { get; set; }
+
+		/// <summary>
+		/// The avatar service
+		/// </summary>
+		IAvatarService? AvatarService { get; set; }
+
+		/// <summary>
+		/// Constructor
+		/// </summary>
+		/// <param name="userCollection"></param>
+		/// <param name="avatarService"></param>
+		public UsersController(IUserCollection userCollection, IAvatarService? avatarService)
+		{
+			UserCollection = userCollection;
+			AvatarService = avatarService;
+		}
+
+		/// <summary>
+		/// Gets information about a user by id, specify "current" for id to get the currently logged in user
+		/// </summary>
+		/// <returns>Http result code</returns>
+		[HttpGet]
+		[Route("/api/v1/users/{id}")]
+		[ProducesResponseType(typeof(List<GetUserResponse>), 200)]
+		public async Task<ActionResult<object>> GetUserAsync(string id, [FromQuery] PropertyFilter? filter = null)
+		{
+			UserId? userId = ParseUserId(id);
+			if(userId == null)
+			{
+				return BadRequest("Invalid user id '{Id}'", id);
+			}
+
+			IUser? user = await UserCollection.GetUserAsync(userId.Value);
+			if (user == null)
+			{
+				return NotFound(userId.Value);
+			}
+
+			IAvatar? avatar = (AvatarService == null) ? (IAvatar?)null : await AvatarService.GetAvatarAsync(user);
+			IUserClaims? claims = await UserCollection.GetClaimsAsync(user.Id);
+			IUserSettings? settings = await UserCollection.GetSettingsAsync(user.Id);
+			return PropertyFilter.Apply(new GetUserResponse(user, avatar, claims, settings), filter);
+		}
+
+		/// <summary>
+		/// Gets a list of users
+		/// </summary>
+		/// <returns>List of user responses</returns>
+		[HttpGet]
+		[Route("/api/v1/users")]
+		[ProducesResponseType(typeof(List<GetUserResponse>), 200)]
+		public async Task<ActionResult<List<GetUserResponse>>> FindUsersAsync(
+			[FromQuery] string[]? ids = null,
+			[FromQuery] string? nameRegex = null,
+			[FromQuery] int index = 0,
+			[FromQuery] int count = 100,
+			[FromQuery] bool includeClaims = false,
+			[FromQuery] bool includeAvatar = false)
+		{
+
+			UserId[]? userIds = null;
+			if (ids != null && ids.Length > 0)
+			{
+				userIds = ids.Select(x => new UserId(x)).ToArray();
+			}
+
+			List<IUser> users = await UserCollection.FindUsersAsync(userIds, nameRegex, index, count);
+
+			List<GetUserResponse> response = new List<GetUserResponse>();
+			foreach (IUser user in users)
+			{
+				IAvatar? avatar = (AvatarService == null || !includeAvatar) ? (IAvatar?)null : await AvatarService.GetAvatarAsync(user);
+				IUserClaims? claims = (!includeClaims) ? null : await UserCollection.GetClaimsAsync(user.Id);
+				response.Add(new GetUserResponse(user, avatar, claims, null));
+			}
+
+			return response;
+		}
+
+		UserId? ParseUserId(string id)
+		{
+			if (id.Equals("current", StringComparison.OrdinalIgnoreCase))
+			{
+				return User.GetUserId();
+			}
+			else if(UserId.TryParse(id, out UserId result))
+			{
+				return result;
+			}
+			return null;
+		}
+	}
+}
