Horde: Remove "Default" pool type which is problematic and support setting a pool to use default
#rnx

diff --git a/Engine/Source/Programs/Horde/Horde.Build/Agents/Fleet/AutoscaleServiceV2.cs b/Engine/Source/Programs/Horde/Horde.Build/Agents/Fleet/AutoscaleServiceV2.cs
--- a/Engine/Source/Programs/Horde/Horde.Build/Agents/Fleet/AutoscaleServiceV2.cs
+++ b/Engine/Source/Programs/Horde/Horde.Build/Agents/Fleet/AutoscaleServiceV2.cs
@@ -80,7 +80,7 @@ namespace Horde.Build.Fleet.Autoscale
 			Dictionary<PoolSizeStrategy, List<PoolSizeData>> poolSizeDataByStrategy = await GetPoolSizeDataByStrategyType();
 			foreach ((PoolSizeStrategy strategyType, List<PoolSizeData> currentData) in poolSizeDataByStrategy)
 			{
-				List<PoolSizeData> newData = await GetPoolSizeStrategy(strategyType, _settings.DefaultAgentPoolSizeStrategy).CalcDesiredPoolSizesAsync(currentData);
+				List<PoolSizeData> newData = await GetPoolSizeStrategy(strategyType).CalcDesiredPoolSizesAsync(currentData);
 				await ResizePools(newData);
 			}
 			
@@ -98,7 +98,7 @@ namespace Horde.Build.Fleet.Autoscale
 			foreach (PoolSizeStrategy strategyType in Enum.GetValues<PoolSizeStrategy>())
 			{
 				result[strategyType] = pools
-					.Where(x => (x.SizeStrategy == PoolSizeStrategy.Default ? _settings.DefaultAgentPoolSizeStrategy : x.SizeStrategy) == strategyType)
+					.Where(x => x.SizeStrategy == strategyType)
 					.Select(x => new PoolSizeData(x, GetAgentsInPool(x.Id), null))
 					.ToList();
 			}
@@ -190,14 +190,8 @@ namespace Horde.Build.Fleet.Autoscale
 			_noOpStrategy = noOpStrategy;
 		}
 		
-		private IPoolSizeStrategy GetPoolSizeStrategy(PoolSizeStrategy type, PoolSizeStrategy defaultStrategy)
+		private IPoolSizeStrategy GetPoolSizeStrategy(PoolSizeStrategy type)
 		{
-
-			if (type == PoolSizeStrategy.Default)
-			{
-				type = defaultStrategy;
-			}
-
 			return type switch
 			{
 				PoolSizeStrategy.LeaseUtilization => _leaseUtilizationStrategy,
diff --git a/Engine/Source/Programs/Horde/Horde.Build/Agents/Fleet/IPoolSizeStrategy.cs b/Engine/Source/Programs/Horde/Horde.Build/Agents/Fleet/IPoolSizeStrategy.cs
--- a/Engine/Source/Programs/Horde/Horde.Build/Agents/Fleet/IPoolSizeStrategy.cs
+++ b/Engine/Source/Programs/Horde/Horde.Build/Agents/Fleet/IPoolSizeStrategy.cs
@@ -12,11 +12,6 @@ namespace Horde.Build.Fleet.Autoscale
 	/// </summary>
 	public enum PoolSizeStrategy
 	{
-		/// <summary>
-		/// Use the default strategy configured in server settings
-		/// </summary>
-		Default,
-
 		/// <summary>
 		/// Strategy based on lease utilization
 		/// <see cref="Horde.Build.Fleet.Autoscale.LeaseUtilizationStrategy"/>
diff --git a/Engine/Source/Programs/Horde/Horde.Build/Agents/Pools/IPoolCollection.cs b/Engine/Source/Programs/Horde/Horde.Build/Agents/Pools/IPoolCollection.cs
--- a/Engine/Source/Programs/Horde/Horde.Build/Agents/Pools/IPoolCollection.cs
+++ b/Engine/Source/Programs/Horde/Horde.Build/Agents/Pools/IPoolCollection.cs
@@ -92,6 +92,7 @@ namespace Horde.Build.Collections
 		/// <param name="sizeStrategy">Pool sizing strategy</param>
 		/// <param name="leaseUtilizationSettings">Settings for lease utilization strategy</param>
 		/// <param name="jobQueueSettings">Settings for job queue strategy</param>
+		/// <param name="useDefaultStrategy">Whether to use the default strategy</param>
 		/// <returns>Async task</returns>
 		Task<IPool?> TryUpdateAsync(
 			IPool pool,
@@ -109,6 +110,7 @@ namespace Horde.Build.Collections
 			TimeSpan? scaleInCooldown = null,
 			PoolSizeStrategy? sizeStrategy = null,
 			LeaseUtilizationSettings? leaseUtilizationSettings = null,
-			JobQueueSettings? jobQueueSettings = null);
+			JobQueueSettings? jobQueueSettings = null,
+			bool? useDefaultStrategy = null);
 	}
 }
diff --git a/Engine/Source/Programs/Horde/Horde.Build/Agents/Pools/PoolCollection.cs b/Engine/Source/Programs/Horde/Horde.Build/Agents/Pools/PoolCollection.cs
--- a/Engine/Source/Programs/Horde/Horde.Build/Agents/Pools/PoolCollection.cs
+++ b/Engine/Source/Programs/Horde/Horde.Build/Agents/Pools/PoolCollection.cs
@@ -9,6 +9,7 @@ using Horde.Build.Fleet.Autoscale;
 using Horde.Build.Models;
 using Horde.Build.Server;
 using Horde.Build.Utilities;
+using Microsoft.Extensions.Options;
 using MongoDB.Bson;
 using MongoDB.Bson.Serialization.Attributes;
 using MongoDB.Driver;
@@ -42,14 +43,14 @@ namespace Horde.Build.Collections.Impl
 			public TimeSpan? ScaleOutCooldown { get; set; }
 			public TimeSpan? ScaleInCooldown { get; set; }
 			public PoolSizeStrategy? SizeStrategy { get; set; }
-			PoolSizeStrategy IPool.SizeStrategy => SizeStrategy ?? PoolSizeStrategy.Default;
+			PoolSizeStrategy IPool.SizeStrategy => SizeStrategy ?? defaultStrategy!.Value;
 			public LeaseUtilizationSettings? LeaseUtilizationSettings { get; set; }
 			public JobQueueSettings? JobQueueSettings { get; set; }
 			public int UpdateIndex { get; set; }
 
 			// Read-only wrappers
 			IReadOnlyList<AgentWorkspace> IPool.Workspaces => Workspaces;
-			IReadOnlyDictionary<string, string> IPool.Properties => Properties;
+			IReadOnlyDictionary<string, string> IPool.Properties => Properties;			
 
 			public PoolDocument()
 			{
@@ -82,13 +83,17 @@ namespace Horde.Build.Collections.Impl
 		/// </summary>
 		readonly IMongoCollection<PoolDocument> _pools;
 
+		static PoolSizeStrategy? defaultStrategy;
+
 		/// <summary>
 		/// Constructor
 		/// </summary>
 		/// <param name="mongoService">The database service instance</param>
-		public PoolCollection(MongoService mongoService)
+		/// <param name="serverSettings">The server settings</param>
+		public PoolCollection(MongoService mongoService, IOptions<ServerSettings> serverSettings)
 		{
 			_pools = mongoService.GetCollection<PoolDocument>("Pools");
+			defaultStrategy = serverSettings.Value.DefaultAgentPoolSizeStrategy;
 		}
 
 		/// <inheritdoc/>
@@ -216,7 +221,8 @@ namespace Horde.Build.Collections.Impl
 			TimeSpan? scaleInCooldown, 
 			PoolSizeStrategy? sizeStrategy,
 			LeaseUtilizationSettings? leaseUtilizationSettings,
-			JobQueueSettings? jobQueueSettings)
+			JobQueueSettings? jobQueueSettings, 
+			bool? useDefaultStrategy)
 		{
 			TransactionBuilder<PoolDocument> transaction = new TransactionBuilder<PoolDocument>();
 			if (newName != null)
@@ -297,16 +303,13 @@ namespace Horde.Build.Collections.Impl
 			}
 			if (sizeStrategy != null)
 			{
-				if (sizeStrategy == PoolSizeStrategy.Default)
-				{
-					transaction.Set(x => x.SizeStrategy, null);
-				}
-				else
-				{
-					transaction.Set(x => x.SizeStrategy, sizeStrategy);
-				}
-				
+				transaction.Set(x => x.SizeStrategy, sizeStrategy);				
+			}
+			else if (useDefaultStrategy != null && useDefaultStrategy.Value)
+			{
+				transaction.Set(x => x.SizeStrategy, null);
 			}
+
 			if (leaseUtilizationSettings != null)
 			{
 				transaction.Set(x => x.LeaseUtilizationSettings, leaseUtilizationSettings);
diff --git a/Engine/Source/Programs/Horde/Horde.Build/Agents/Pools/PoolService.cs b/Engine/Source/Programs/Horde/Horde.Build/Agents/Pools/PoolService.cs
--- a/Engine/Source/Programs/Horde/Horde.Build/Agents/Pools/PoolService.cs
+++ b/Engine/Source/Programs/Horde/Horde.Build/Agents/Pools/PoolService.cs
@@ -121,6 +121,7 @@ namespace Horde.Build.Services
 		/// <param name="scaleOutCooldown">Cooldown time between scale-out events</param>
 		/// <param name="scaleInCooldown">Cooldown time between scale-in events</param>
 		/// <param name="sizeStrategy">New pool sizing strategy for the pool</param>
+		/// <param name="useDefaultStrategy">Whether to use the default strategy</param>
 		/// <returns>Async task object</returns>
 		public async Task<IPool?> UpdatePoolAsync(
 			IPool? pool,
@@ -132,7 +133,8 @@ namespace Horde.Build.Services
 			Dictionary<string, string?>? newProperties = null,
 			TimeSpan? scaleOutCooldown = null,
 			TimeSpan? scaleInCooldown = null,
-			PoolSizeStrategy? sizeStrategy = null)
+			PoolSizeStrategy? sizeStrategy = null,
+			bool? useDefaultStrategy = null)
 		{
 			for (; pool != null; pool = await _pools.GetAsync(pool.Id))
 			{
@@ -146,7 +148,8 @@ namespace Horde.Build.Services
 					newProperties: newProperties,
 					scaleOutCooldown: scaleOutCooldown,
 					scaleInCooldown: scaleInCooldown,
-					sizeStrategy: sizeStrategy);
+					sizeStrategy: sizeStrategy,
+					useDefaultStrategy: useDefaultStrategy);
 				
 				if (newPool != null)
 				{
diff --git a/Engine/Source/Programs/Horde/Horde.Build/Agents/Pools/PoolsController.Messages.cs b/Engine/Source/Programs/Horde/Horde.Build/Agents/Pools/PoolsController.Messages.cs
--- a/Engine/Source/Programs/Horde/Horde.Build/Agents/Pools/PoolsController.Messages.cs
+++ b/Engine/Source/Programs/Horde/Horde.Build/Agents/Pools/PoolsController.Messages.cs
@@ -185,7 +185,12 @@ namespace Horde.Build.Api
 		/// Pool sizing strategy
 		/// </summary>
 		public PoolSizeStrategy? SizeStrategy { get; set; }
-		
+
+		/// <summary>
+		/// Set pool to use default strategy
+		/// </summary>
+		public bool? UseDefaultStrategy { get; set; }
+
 		/// <summary>
 		/// Settings for lease utilization pool sizing strategy (if used)
 		/// </summary>
diff --git a/Engine/Source/Programs/Horde/Horde.Build/Agents/Pools/PoolsController.cs b/Engine/Source/Programs/Horde/Horde.Build/Agents/Pools/PoolsController.cs
--- a/Engine/Source/Programs/Horde/Horde.Build/Agents/Pools/PoolsController.cs
+++ b/Engine/Source/Programs/Horde/Horde.Build/Agents/Pools/PoolsController.cs
@@ -1,4 +1,4 @@
-﻿// Copyright Epic Games, Inc. All Rights Reserved.
+// Copyright Epic Games, Inc. All Rights Reserved.
 
 using System;
 using System.Collections.Generic;
@@ -149,7 +149,7 @@ namespace Horde.Build.Controllers
 			TimeSpan? scaleInCooldown = update.ScaleInCooldown == null ? null : TimeSpan.FromSeconds(update.ScaleInCooldown.Value);
 
 			await _poolService.UpdatePoolAsync(pool, update.Name, update.Condition, update.EnableAutoscaling,
-				update.MinAgents, update.NumReserveAgents, update.Properties, scaleOutCooldown, scaleInCooldown, update.SizeStrategy);
+				update.MinAgents, update.NumReserveAgents, update.Properties, scaleOutCooldown, scaleInCooldown, update.SizeStrategy, update.UseDefaultStrategy);
 			return new OkResult();
 		}
 
