Horde: Add missing file from previous commit
#rnx

diff --git a/Engine/Source/Programs/Horde/Horde.Build.Tests/Fleet/AutoscaleServiceV2Test.cs b/Engine/Source/Programs/Horde/Horde.Build.Tests/Fleet/AutoscaleServiceV2Test.cs
--- a/Engine/Source/Programs/Horde/Horde.Build.Tests/Fleet/AutoscaleServiceV2Test.cs
+++ b/Engine/Source/Programs/Horde/Horde.Build.Tests/Fleet/AutoscaleServiceV2Test.cs
@@ -83,7 +83,7 @@ namespace Horde.Build.Tests.Fleet
 
 			await service.TickLeaderAsync(CancellationToken.None);
 			
-			Assert.AreEqual(2, leaseUtilizationSpy.CallCount);
+			Assert.AreEqual(1, leaseUtilizationSpy.CallCount);
 			Assert.IsTrue(leaseUtilizationSpy.PoolIdsSeen.Contains(pool1.Id));
 			Assert.IsTrue(leaseUtilizationSpy.PoolIdsSeen.Contains(pool2.Id));
 			
