Horde: Move files into namespaces corresponding to their location on disk.

diff --git a/Engine/Source/Programs/Horde/Horde.Build.Tests/AgentControllerTest.cs b/Engine/Source/Programs/Horde/Horde.Build.Tests/AgentControllerTest.cs
--- a/Engine/Source/Programs/Horde/Horde.Build.Tests/AgentControllerTest.cs
+++ b/Engine/Source/Programs/Horde/Horde.Build.Tests/AgentControllerTest.cs
@@ -1,8 +1,7 @@
 // Copyright Epic Games, Inc. All Rights Reserved.
 
 using System.Threading.Tasks;
-using Horde.Build.Api;
-using Horde.Build.Models;
+using Horde.Build.Agents;
 using Microsoft.AspNetCore.Mvc;
 using Microsoft.VisualStudio.TestTools.UnitTesting;
 
diff --git a/Engine/Source/Programs/Horde/Horde.Build.Tests/AgentServiceTest.cs b/Engine/Source/Programs/Horde/Horde.Build.Tests/AgentServiceTest.cs
--- a/Engine/Source/Programs/Horde/Horde.Build.Tests/AgentServiceTest.cs
+++ b/Engine/Source/Programs/Horde/Horde.Build.Tests/AgentServiceTest.cs
@@ -4,8 +4,8 @@ using System;
 using System.Collections.Generic;
 using System.Security.Claims;
 using System.Threading.Tasks;
-using Horde.Build.Api;
-using Horde.Build.Models;
+using Horde.Build.Agents;
+using Horde.Build.Jobs;
 using Horde.Build.Utilities;
 using HordeCommon;
 using Microsoft.AspNetCore.Mvc;
diff --git a/Engine/Source/Programs/Horde/Horde.Build.Tests/ControllerIntegrationTest.cs b/Engine/Source/Programs/Horde/Horde.Build.Tests/ControllerIntegrationTest.cs
--- a/Engine/Source/Programs/Horde/Horde.Build.Tests/ControllerIntegrationTest.cs
+++ b/Engine/Source/Programs/Horde/Horde.Build.Tests/ControllerIntegrationTest.cs
@@ -4,18 +4,21 @@ using System;
 using System.Collections.Generic;
 using System.Net.Http;
 using System.Threading.Tasks;
-using Horde.Build.Collections;
-using Horde.Build.Collections.Impl;
-using Horde.Build.Config;
 using Horde.Build.Jobs;
 using Horde.Build.Server;
-using Horde.Build.Services;
+using Horde.Build.Perforce;
 using Horde.Build.Tests.Stubs.Services;
 using Microsoft.AspNetCore.Hosting;
 using Microsoft.AspNetCore.Mvc.Testing;
 using Microsoft.AspNetCore.TestHost;
 using Microsoft.Extensions.Configuration;
 using Microsoft.Extensions.DependencyInjection;
+using Horde.Build.Agents;
+using Horde.Build.Streams;
+using Horde.Build.Jobs.Templates;
+using Horde.Build.Configuration;
+using Horde.Build.Jobs.Graphs;
+using Horde.Build.Jobs.Artifacts;
 
 namespace Horde.Build.Tests;
 
diff --git a/Engine/Source/Programs/Horde/Horde.Build.Tests/Fixture.cs b/Engine/Source/Programs/Horde/Horde.Build.Tests/Fixture.cs
--- a/Engine/Source/Programs/Horde/Horde.Build.Tests/Fixture.cs
+++ b/Engine/Source/Programs/Horde/Horde.Build.Tests/Fixture.cs
@@ -5,13 +5,15 @@ using System.IO;
 using System.Text;
 using System.Threading.Tasks;
 using EpicGames.Core;
-using Horde.Build.Api;
-using Horde.Build.Collections;
-using Horde.Build.Config;
+using Horde.Build.Agents;
+using Horde.Build.Configuration;
 using Horde.Build.Jobs;
-using Horde.Build.Models;
+using Horde.Build.Jobs.Artifacts;
+using Horde.Build.Jobs.Graphs;
+using Horde.Build.Jobs.Templates;
+using Horde.Build.Projects;
 using Horde.Build.Server;
-using Horde.Build.Services;
+using Horde.Build.Streams;
 using Horde.Build.Utilities;
 using HordeCommon;
 
diff --git a/Engine/Source/Programs/Horde/Horde.Build.Tests/Fleet/AutoscaleServiceV2Test.cs b/Engine/Source/Programs/Horde/Horde.Build.Tests/Fleet/AutoscaleServiceV2Test.cs
--- a/Engine/Source/Programs/Horde/Horde.Build.Tests/Fleet/AutoscaleServiceV2Test.cs
+++ b/Engine/Source/Programs/Horde/Horde.Build.Tests/Fleet/AutoscaleServiceV2Test.cs
@@ -4,9 +4,9 @@ using System;
 using System.Collections.Generic;
 using System.Threading;
 using System.Threading.Tasks;
-using Horde.Build.Fleet.Autoscale;
-using Horde.Build.Models;
-using Horde.Build.Services;
+using Horde.Build.Agents;
+using Horde.Build.Agents.Fleet;
+using Horde.Build.Agents.Pools;
 using Horde.Build.Utilities;
 using Microsoft.Extensions.DependencyInjection;
 using Microsoft.Extensions.Logging;
diff --git a/Engine/Source/Programs/Horde/Horde.Build.Tests/Fleet/JobQueueStrategyTest.cs b/Engine/Source/Programs/Horde/Horde.Build.Tests/Fleet/JobQueueStrategyTest.cs
--- a/Engine/Source/Programs/Horde/Horde.Build.Tests/Fleet/JobQueueStrategyTest.cs
+++ b/Engine/Source/Programs/Horde/Horde.Build.Tests/Fleet/JobQueueStrategyTest.cs
@@ -4,12 +4,16 @@ using System;
 using System.Collections.Generic;
 using System.Threading.Tasks;
 using EpicGames.Core;
-using Horde.Build.Api;
-using Horde.Build.Fleet.Autoscale;
-using Horde.Build.Models;
+using Horde.Build.Agents.Pools;
+using Horde.Build.Agents.Fleet;
+using Horde.Build.Projects;
+using Horde.Build.Streams;
+using Horde.Build.Jobs;
+using Horde.Build.Jobs.Graphs;
 using Horde.Build.Server;
 using Horde.Build.Utilities;
 using Microsoft.VisualStudio.TestTools.UnitTesting;
+using Horde.Build.Agents;
 
 namespace Horde.Build.Tests.Fleet
 {
diff --git a/Engine/Source/Programs/Horde/Horde.Build.Tests/Fleet/LeaseUtilizationStrategyTest.cs b/Engine/Source/Programs/Horde/Horde.Build.Tests/Fleet/LeaseUtilizationStrategyTest.cs
--- a/Engine/Source/Programs/Horde/Horde.Build.Tests/Fleet/LeaseUtilizationStrategyTest.cs
+++ b/Engine/Source/Programs/Horde/Horde.Build.Tests/Fleet/LeaseUtilizationStrategyTest.cs
@@ -5,8 +5,11 @@ using System.Collections.Generic;
 using System.Threading.Tasks;
 using Google.Protobuf;
 using Google.Protobuf.WellKnownTypes;
-using Horde.Build.Fleet.Autoscale;
-using Horde.Build.Models;
+using Horde.Build.Agents;
+using Horde.Build.Agents.Fleet;
+using Horde.Build.Agents.Leases;
+using Horde.Build.Streams;
+using Horde.Build.Agents.Pools;
 using Horde.Build.Utilities;
 using HordeCommon;
 using HordeCommon.Rpc.Tasks;
diff --git a/Engine/Source/Programs/Horde/Horde.Build.Tests/IssueServiceTests.cs b/Engine/Source/Programs/Horde/Horde.Build.Tests/IssueServiceTests.cs
--- a/Engine/Source/Programs/Horde/Horde.Build.Tests/IssueServiceTests.cs
+++ b/Engine/Source/Programs/Horde/Horde.Build.Tests/IssueServiceTests.cs
@@ -9,12 +9,15 @@ using System.Threading.Tasks;
 using EpicGames.Core;
 using Horde.Agent.Parser;
 using Horde.Agent.Utility;
-using Horde.Build.IssueHandlers.Impl;
-using Horde.Build.Models;
+using Horde.Build.Jobs.Graphs;
+using Horde.Build.Jobs;
+using Horde.Build.Issues;
+using Horde.Build.Logs;
+using Horde.Build.Users;
+using Horde.Build.Projects;
+using Horde.Build.Streams;
 using Horde.Build.Server;
-using Horde.Build.Services;
 using Horde.Build.Tests.Stubs.Services;
-using Horde.Build.Collections;
 using Horde.Build.Utilities;
 using HordeCommon;
 using Microsoft.Extensions.DependencyInjection;
@@ -23,6 +26,7 @@ using Microsoft.VisualStudio.TestTools.UnitTesting;
 using MongoDB.Bson;
 using MongoDB.Driver;
 using Moq;
+using Horde.Build.Issues.Handlers;
 
 namespace Horde.Build.Tests
 {
@@ -115,8 +119,8 @@ namespace Horde.Build.Tests
 			string revision = $"config:{streamId}";
 
 			StreamConfig streamConfig = new StreamConfig { Name = streamName };
-			streamConfig.Tabs.Add(new Api.CreateJobsTabRequest { Title = "General", Templates = new List<string> { "test-template" } });
-			streamConfig.Templates.Add(new Api.TemplateRefConfig { Id = new TemplateRefId("test-template") });
+			streamConfig.Tabs.Add(new CreateJobsTabRequest { Title = "General", Templates = new List<string> { "test-template" } });
+			streamConfig.Templates.Add(new TemplateRefConfig { Id = new TemplateRefId("test-template") });
 			await ConfigCollection.AddConfigAsync(revision, streamConfig);
 
 			return Deref(await StreamCollection.TryCreateOrReplaceAsync(streamId, null, revision, projectId));
@@ -199,7 +203,7 @@ namespace Horde.Build.Tests
 				{
 					SubResourceId stepId = new SubResourceId((ushort)((groupIdx * 100) + nodeIdx));
 
-					ILogFile logFile = LogFileService.CreateLogFileAsync(jobId, null, Horde.Build.Api.LogType.Json).Result;
+					ILogFile logFile = LogFileService.CreateLogFileAsync(jobId, null, LogType.Json).Result;
 
 					Mock<IJobStep> step = new Mock<IJobStep>(MockBehavior.Strict);
 					step.SetupGet(x => x.Id).Returns(stepId);
@@ -1563,7 +1567,7 @@ namespace Horde.Build.Tests
 				spans.Add(MockSpan(new StreamId("ue5-release-staging"), MockSuspect(101, 1), MockSuspect(102, 2)));
 				spans.Add(MockSpan(new StreamId("ue5-release"), MockSuspect(1, null), MockSuspect(2, null)));
 
-				List<IIssueSpan> results = Horde.Build.Services.IssueService.FindMergeOriginSpans(spans);
+				List<IIssueSpan> results = Horde.Build.Issues.IssueService.FindMergeOriginSpans(spans);
 				Assert.AreEqual(1, results.Count);
 				Assert.AreEqual(new StreamId("ue5-release"), results[0].StreamId);
 			}
@@ -1574,7 +1578,7 @@ namespace Horde.Build.Tests
 				spans.Add(MockSpan(new StreamId("ue5-release-staging"), MockSuspect(101, 1), MockSuspect(102, null)));
 				spans.Add(MockSpan(new StreamId("ue5-release"), MockSuspect(1, null), MockSuspect(2, null)));
 
-				List<IIssueSpan> results = Horde.Build.Services.IssueService.FindMergeOriginSpans(spans);
+				List<IIssueSpan> results = Horde.Build.Issues.IssueService.FindMergeOriginSpans(spans);
 				Assert.AreEqual(1, results.Count);
 				Assert.AreEqual(new StreamId("ue5-release"), results[0].StreamId);
 			}
@@ -1585,7 +1589,7 @@ namespace Horde.Build.Tests
 				spans.Add(MockSpan(new StreamId("ue5-release-staging"), MockSuspect(101, 1), MockSuspect(102, null)));
 				spans.Add(MockSpan(new StreamId("ue5-main"), MockSuspect(201, null), MockSuspect(202, 2)));
 
-				List<IIssueSpan> results = Horde.Build.Services.IssueService.FindMergeOriginSpans(spans);
+				List<IIssueSpan> results = Horde.Build.Issues.IssueService.FindMergeOriginSpans(spans);
 				Assert.AreEqual(1, results.Count);
 				Assert.AreEqual(new StreamId("ue5-release"), results[0].StreamId);
 			}
diff --git a/Engine/Source/Programs/Horde/Horde.Build.Tests/JobCollectionTests.cs b/Engine/Source/Programs/Horde/Horde.Build.Tests/JobCollectionTests.cs
--- a/Engine/Source/Programs/Horde/Horde.Build.Tests/JobCollectionTests.cs
+++ b/Engine/Source/Programs/Horde/Horde.Build.Tests/JobCollectionTests.cs
@@ -5,13 +5,20 @@ using System.Collections.Generic;
 using System.Linq;
 using System.Threading.Tasks;
 using EpicGames.Core;
-using Horde.Build.Api;
-using Horde.Build.Models;
+using Horde.Build.Jobs;
+using Horde.Build.Jobs.Graphs;
+using Horde.Build.Agents;
+using Horde.Build.Logs;
+using Horde.Build.Agents.Leases;
+using Horde.Build.Agents.Pools;
+using Horde.Build.Streams;
+using Horde.Build.Jobs.Templates;
 using Horde.Build.Utilities;
 using HordeCommon;
 using Microsoft.VisualStudio.TestTools.UnitTesting;
 using MongoDB.Bson;
 using Moq;
+using Horde.Build.Agents.Sessions;
 
 namespace Horde.Build.Tests
 {
diff --git a/Engine/Source/Programs/Horde/Horde.Build.Tests/JobServiceTests.cs b/Engine/Source/Programs/Horde/Horde.Build.Tests/JobServiceTests.cs
--- a/Engine/Source/Programs/Horde/Horde.Build.Tests/JobServiceTests.cs
+++ b/Engine/Source/Programs/Horde/Horde.Build.Tests/JobServiceTests.cs
@@ -5,9 +5,16 @@ using System.Collections.Generic;
 using System.Text;
 using System.Threading.Tasks;
 using EpicGames.Core;
-using Horde.Build.Api;
-using Horde.Build.Models;
+using Horde.Build.Agents;
+using Horde.Build.Agents.Pools;
+using Horde.Build.Jobs;
+using Horde.Build.Jobs.Graphs;
+using Horde.Build.Jobs.Templates;
+using Horde.Build.Logs;
+using Horde.Build.Projects;
 using Horde.Build.Server;
+using Horde.Build.Streams;
+using Horde.Build.Users;
 using Horde.Build.Utilities;
 using HordeCommon;
 using Microsoft.VisualStudio.TestTools.UnitTesting;
diff --git a/Engine/Source/Programs/Horde/Horde.Build.Tests/JobTaskSourceTests.cs b/Engine/Source/Programs/Horde/Horde.Build.Tests/JobTaskSourceTests.cs
--- a/Engine/Source/Programs/Horde/Horde.Build.Tests/JobTaskSourceTests.cs
+++ b/Engine/Source/Programs/Horde/Horde.Build.Tests/JobTaskSourceTests.cs
@@ -3,15 +3,17 @@
 using System.Collections.Generic;
 using System.Threading;
 using System.Threading.Tasks;
-using Horde.Build.Api;
-using Horde.Build.Models;
+using Horde.Build.Jobs;
+using Horde.Build.Jobs.Graphs;
+using Horde.Build.Users;
 using Horde.Build.Utilities;
 using HordeCommon;
 using Microsoft.VisualStudio.TestTools.UnitTesting;
-using Horde.Build.Services;
 using System.Linq;
 using System;
-
+using Horde.Build.Agents;
+using Horde.Build.Agents.Pools;
+using Horde.Build.Streams;
 
 namespace Horde.Build.Tests
 {
diff --git a/Engine/Source/Programs/Horde/Horde.Build.Tests/JobsControllerDbTest.cs b/Engine/Source/Programs/Horde/Horde.Build.Tests/JobsControllerDbTest.cs
--- a/Engine/Source/Programs/Horde/Horde.Build.Tests/JobsControllerDbTest.cs
+++ b/Engine/Source/Programs/Horde/Horde.Build.Tests/JobsControllerDbTest.cs
@@ -2,9 +2,8 @@
 
 using System.Collections.Generic;
 using System.Threading.Tasks;
-using Horde.Build.Api;
-using Horde.Build.Controllers;
-using Horde.Build.Models;
+using Horde.Build.Jobs;
+using Horde.Build.Jobs.Graphs;
 using Microsoft.AspNetCore.Mvc;
 using Microsoft.VisualStudio.TestTools.UnitTesting;
 
diff --git a/Engine/Source/Programs/Horde/Horde.Build.Tests/JobsControllerTest.cs b/Engine/Source/Programs/Horde/Horde.Build.Tests/JobsControllerTest.cs
--- a/Engine/Source/Programs/Horde/Horde.Build.Tests/JobsControllerTest.cs
+++ b/Engine/Source/Programs/Horde/Horde.Build.Tests/JobsControllerTest.cs
@@ -3,6 +3,7 @@
 using System;
 using System.Net;
 using System.Threading.Tasks;
+using Horde.Build.Jobs.Artifacts;
 using Microsoft.VisualStudio.TestTools.UnitTesting;
 
 namespace Horde.Build.Tests
@@ -19,7 +20,7 @@ namespace Horde.Build.Tests
         public async Task GetArtifactDataByFilenameTest()
         {
 			Fixture? fixture = await GetFixture();
-			Models.IArtifact? art = fixture.Job1Artifact;
+			IArtifact? art = fixture.Job1Artifact;
 
 			// Test existing filename
 			System.Net.Http.HttpResponseMessage? res = await Client.GetAsync(GetUri(art.JobId.ToString(), art.StepId.ToString()!, art.Name));
diff --git a/Engine/Source/Programs/Horde/Horde.Build.Tests/LeasesControllerDbTest.cs b/Engine/Source/Programs/Horde/Horde.Build.Tests/LeasesControllerDbTest.cs
--- a/Engine/Source/Programs/Horde/Horde.Build.Tests/LeasesControllerDbTest.cs
+++ b/Engine/Source/Programs/Horde/Horde.Build.Tests/LeasesControllerDbTest.cs
@@ -3,8 +3,9 @@
 using System;
 using System.Collections.Generic;
 using System.Threading.Tasks;
-using Horde.Build.Api;
-using Horde.Build.Models;
+using Horde.Build.Agents;
+using Horde.Build.Agents.Leases;
+using Horde.Build.Agents.Sessions;
 using Horde.Build.Utilities;
 using HordeCommon;
 using Microsoft.AspNetCore.Mvc;
diff --git a/Engine/Source/Programs/Horde/Horde.Build.Tests/LogBuilderTests.cs b/Engine/Source/Programs/Horde/Horde.Build.Tests/LogBuilderTests.cs
--- a/Engine/Source/Programs/Horde/Horde.Build.Tests/LogBuilderTests.cs
+++ b/Engine/Source/Programs/Horde/Horde.Build.Tests/LogBuilderTests.cs
@@ -4,10 +4,9 @@ using System;
 using System.Collections.Generic;
 using System.Text;
 using System.Threading.Tasks;
-using Horde.Build.Api;
 using Horde.Build.Logs;
 using Horde.Build.Logs.Builder;
-using Horde.Build.Models;
+using Horde.Build.Logs.Data;
 using Horde.Build.Server;
 using Horde.Build.Utilities;
 using Microsoft.Extensions.Logging.Abstractions;
diff --git a/Engine/Source/Programs/Horde/Horde.Build.Tests/LogChunkTests.cs b/Engine/Source/Programs/Horde/Horde.Build.Tests/LogChunkTests.cs
--- a/Engine/Source/Programs/Horde/Horde.Build.Tests/LogChunkTests.cs
+++ b/Engine/Source/Programs/Horde/Horde.Build.Tests/LogChunkTests.cs
@@ -3,8 +3,8 @@
 using System;
 using System.Collections.Generic;
 using System.Text;
-using Horde.Build.Api;
 using Horde.Build.Logs;
+using Horde.Build.Logs.Data;
 using Microsoft.Extensions.Logging.Abstractions;
 using Microsoft.VisualStudio.TestTools.UnitTesting;
 
diff --git a/Engine/Source/Programs/Horde/Horde.Build.Tests/LogFileServiceTests.cs b/Engine/Source/Programs/Horde/Horde.Build.Tests/LogFileServiceTests.cs
--- a/Engine/Source/Programs/Horde/Horde.Build.Tests/LogFileServiceTests.cs
+++ b/Engine/Source/Programs/Horde/Horde.Build.Tests/LogFileServiceTests.cs
@@ -7,13 +7,12 @@ using System.Linq;
 using System.Security.Claims;
 using System.Text;
 using System.Threading.Tasks;
-using Horde.Build.Api;
-using Horde.Build.Collections.Impl;
+using Horde.Build.Agents.Sessions;
+using Horde.Build.Jobs;
 using Horde.Build.Logs;
 using Horde.Build.Logs.Builder;
-using Horde.Build.Logs.Readers;
-using Horde.Build.Models;
-using Horde.Build.Services;
+using Horde.Build.Logs.Data;
+using Horde.Build.Logs.Storage;
 using Horde.Build.Storage;
 using Horde.Build.Storage.Backends;
 using Horde.Build.Utilities;
diff --git a/Engine/Source/Programs/Horde/Horde.Build.Tests/LogIndexingTests.cs b/Engine/Source/Programs/Horde/Horde.Build.Tests/LogIndexingTests.cs
--- a/Engine/Source/Programs/Horde/Horde.Build.Tests/LogIndexingTests.cs
+++ b/Engine/Source/Programs/Horde/Horde.Build.Tests/LogIndexingTests.cs
@@ -6,21 +6,18 @@ using System.IO;
 using System.Linq;
 using System.Text;
 using System.Threading.Tasks;
-using Horde.Build.Api;
-using Horde.Build.Collections.Impl;
 using Horde.Build.Compute.Tests.Properties;
 using Horde.Build.Logs;
 using Horde.Build.Logs.Builder;
 using Horde.Build.Logs.Storage;
-using Horde.Build.Logs.Storage.Impl;
-using Horde.Build.Models;
-using Horde.Build.Services;
+using Horde.Build.Jobs;
 using Horde.Build.Utilities;
 using HordeCommon;
 using Microsoft.Extensions.Configuration;
 using Microsoft.Extensions.DependencyInjection;
 using Microsoft.Extensions.Logging;
 using Microsoft.VisualStudio.TestTools.UnitTesting;
+using Horde.Build.Logs.Data;
 
 namespace Horde.Build.Tests
 {
diff --git a/Engine/Source/Programs/Horde/Horde.Build.Tests/Models/Job/GetLabelStatesTests.cs b/Engine/Source/Programs/Horde/Horde.Build.Tests/Models/Job/GetLabelStatesTests.cs
--- a/Engine/Source/Programs/Horde/Horde.Build.Tests/Models/Job/GetLabelStatesTests.cs
+++ b/Engine/Source/Programs/Horde/Horde.Build.Tests/Models/Job/GetLabelStatesTests.cs
@@ -1,8 +1,8 @@
 // Copyright Epic Games, Inc. All Rights Reserved.
 
 using System.Collections.Generic;
-using Horde.Build.Api;
-using Horde.Build.Models;
+using Horde.Build.Jobs;
+using Horde.Build.Jobs.Graphs;
 using Horde.Build.Utilities;
 using HordeCommon;
 using Microsoft.VisualStudio.TestTools.UnitTesting;
diff --git a/Engine/Source/Programs/Horde/Horde.Build.Tests/NotificationServiceTests.cs b/Engine/Source/Programs/Horde/Horde.Build.Tests/NotificationServiceTests.cs
--- a/Engine/Source/Programs/Horde/Horde.Build.Tests/NotificationServiceTests.cs
+++ b/Engine/Source/Programs/Horde/Horde.Build.Tests/NotificationServiceTests.cs
@@ -3,16 +3,20 @@
 using System;
 using System.Collections.Generic;
 using System.Threading.Tasks;
-using Horde.Build.Api;
-using Horde.Build.Issues.Impl;
-using Horde.Build.Models;
+using Horde.Build.Agents.Pools;
+using Horde.Build.Devices;
+using Horde.Build.Issues;
+using Horde.Build.Jobs;
+using Horde.Build.Jobs.Graphs;
 using Horde.Build.Notifications;
-using Horde.Build.Notifications.Impl;
+using Horde.Build.Users;
+using Horde.Build.Streams;
 using Horde.Build.Utilities;
 using HordeCommon;
 using Microsoft.Extensions.DependencyInjection;
 using Microsoft.VisualStudio.TestTools.UnitTesting;
 using Moq;
+using Horde.Build.Logs;
 
 namespace Horde.Build.Tests
 {
diff --git a/Engine/Source/Programs/Horde/Horde.Build.Tests/PoolServiceTests.cs b/Engine/Source/Programs/Horde/Horde.Build.Tests/PoolServiceTests.cs
--- a/Engine/Source/Programs/Horde/Horde.Build.Tests/PoolServiceTests.cs
+++ b/Engine/Source/Programs/Horde/Horde.Build.Tests/PoolServiceTests.cs
@@ -3,13 +3,15 @@
 using System;
 using System.Collections.Generic;
 using System.Threading.Tasks;
-using Horde.Build.Models;
-using Horde.Build.Services;
+using Horde.Build.Agents;
+using Horde.Build.Agents.Pools;
+using Horde.Build.Utilities;
 using Microsoft.VisualStudio.TestTools.UnitTesting;
-using PoolId = Horde.Build.Utilities.StringId<Horde.Build.Models.IPool>;
 
 namespace Horde.Build.Tests
 {
+	using PoolId = StringId<IPool>;
+
 	[TestClass]
     public class PoolServiceTests : TestSetup
     {
diff --git a/Engine/Source/Programs/Horde/Horde.Build.Tests/PoolsControllerTest.cs b/Engine/Source/Programs/Horde/Horde.Build.Tests/PoolsControllerTest.cs
--- a/Engine/Source/Programs/Horde/Horde.Build.Tests/PoolsControllerTest.cs
+++ b/Engine/Source/Programs/Horde/Horde.Build.Tests/PoolsControllerTest.cs
@@ -2,9 +2,9 @@
 
 using System.Collections.Generic;
 using System.Threading.Tasks;
-using Horde.Build.Api;
-using Horde.Build.Fleet.Autoscale;
-using Horde.Build.Models;
+using Horde.Build.Agents;
+using Horde.Build.Agents.Fleet;
+using Horde.Build.Agents.Pools;
 using Horde.Build.Utilities;
 using Microsoft.AspNetCore.Mvc;
 using Microsoft.VisualStudio.TestTools.UnitTesting;
@@ -35,7 +35,7 @@ namespace Horde.Build.Tests
 		        ScaleOutCooldown = 111,
 		        ScaleInCooldown = 222,
 		        SizeStrategy = PoolSizeStrategy.JobQueue,
-		        JobQueueSettings = new Horde.Build.Api.JobQueueSettings(new Horde.Build.Fleet.Autoscale.JobQueueSettings(0.35, 0.85))
+		        JobQueueSettings = new JobQueueSettingsMessage(new JobQueueSettings(0.35, 0.85))
 	        };
 	        
 	        ActionResult<CreatePoolResponse> result = await PoolsController.CreatePoolAsync(request);
diff --git a/Engine/Source/Programs/Horde/Horde.Build.Tests/RpcServiceTest.cs b/Engine/Source/Programs/Horde/Horde.Build.Tests/RpcServiceTest.cs
--- a/Engine/Source/Programs/Horde/Horde.Build.Tests/RpcServiceTest.cs
+++ b/Engine/Source/Programs/Horde/Horde.Build.Tests/RpcServiceTest.cs
@@ -1,4 +1,4 @@
-﻿// Copyright Epic Games, Inc. All Rights Reserved.
+// Copyright Epic Games, Inc. All Rights Reserved.
 
 using System;
 using System.Collections.Generic;
@@ -10,9 +10,11 @@ using System.Threading.Channels;
 using System.Threading.Tasks;
 using Google.Protobuf;
 using Grpc.Core;
-using Horde.Build.Models;
-using Horde.Build.Services;
 using Horde.Build.Utilities;
+using Horde.Build.Agents.Pools;
+using Horde.Build.Agents.Leases;
+using Horde.Build.Logs;
+using Horde.Build.Server;
 using HordeCommon.Rpc;
 using Microsoft.AspNetCore.Http;
 using Microsoft.AspNetCore.Http.Features;
@@ -21,12 +23,14 @@ using Microsoft.VisualStudio.TestTools.UnitTesting;
 using MongoDB.Bson;
 using AgentCapabilities = HordeCommon.Rpc.Messages.AgentCapabilities;
 using ISession = Microsoft.AspNetCore.Http.ISession;
-using PoolId = Horde.Build.Utilities.StringId<Horde.Build.Models.IPool>;
+using Horde.Build.Jobs.Artifacts;
+using Horde.Build.Agents;
 
 namespace Horde.Build.Tests
 {
 	using LeaseId = ObjectId<ILease>;
 	using LogId = ObjectId<ILogFile>;
+	using PoolId = StringId<IPool>;
 
 	public class AppLifetimeStub : IHostApplicationLifetime
 	{
@@ -398,7 +402,7 @@ namespace Horde.Build.Tests
 
 			// Set the session ID on the job batch to pass auth later
 			Deref(await JobCollection.TryAssignLeaseAsync(fixture.Job1, 0, new PoolId("foo"),
-				new AgentId("test"), new ObjectId<Horde.Build.Models.ISession>(sessionId),
+				new AgentId("test"), new ObjectId<Horde.Build.Agents.Sessions.ISession>(sessionId),
 				LeaseId.GenerateNewId(), LogId.GenerateNewId()));
 /*
 			TestAsyncStreamReader<UploadArtifactRequest> RequestStream = new TestAsyncStreamReader<UploadArtifactRequest>(Context);
diff --git a/Engine/Source/Programs/Horde/Horde.Build.Tests/SchedulerTests.cs b/Engine/Source/Programs/Horde/Horde.Build.Tests/SchedulerTests.cs
--- a/Engine/Source/Programs/Horde/Horde.Build.Tests/SchedulerTests.cs
+++ b/Engine/Source/Programs/Horde/Horde.Build.Tests/SchedulerTests.cs
@@ -4,15 +4,20 @@ using System;
 using System.Collections.Generic;
 using System.Linq;
 using System.Threading.Tasks;
-using Horde.Build.Api;
-using Horde.Build.Collections;
-using Horde.Build.Models;
+using Horde.Build.Jobs;
+using Horde.Build.Logs;
 using Horde.Build.Server;
-using Horde.Build.Services;
+using Horde.Build.Projects;
+using Horde.Build.Streams;
+using Horde.Build.Jobs.Templates;
+using Horde.Build.Jobs.Schedules;
+using Horde.Build.Users;
 using Horde.Build.Utilities;
 using HordeCommon;
 using Microsoft.Extensions.DependencyInjection;
 using Microsoft.VisualStudio.TestTools.UnitTesting;
+using Horde.Build.Jobs.Graphs;
+using Horde.Build.Perforce;
 
 namespace Horde.Build.Tests
 {
diff --git a/Engine/Source/Programs/Horde/Horde.Build.Tests/ServiceAccountAuthTest.cs b/Engine/Source/Programs/Horde/Horde.Build.Tests/ServiceAccountAuthTest.cs
--- a/Engine/Source/Programs/Horde/Horde.Build.Tests/ServiceAccountAuthTest.cs
+++ b/Engine/Source/Programs/Horde/Horde.Build.Tests/ServiceAccountAuthTest.cs
@@ -5,8 +5,7 @@ using System.Linq;
 using System.Security.Claims;
 using System.Threading.Tasks;
 using Horde.Build.Authentication;
-using Horde.Build.Collections;
-using Horde.Build.Collections.Impl;
+using Horde.Build.Users;
 using Microsoft.AspNetCore.Authentication;
 using Microsoft.AspNetCore.Http;
 using Microsoft.Extensions.DependencyInjection;
diff --git a/Engine/Source/Programs/Horde/Horde.Build.Tests/ServiceAccountCollectionTests.cs b/Engine/Source/Programs/Horde/Horde.Build.Tests/ServiceAccountCollectionTests.cs
--- a/Engine/Source/Programs/Horde/Horde.Build.Tests/ServiceAccountCollectionTests.cs
+++ b/Engine/Source/Programs/Horde/Horde.Build.Tests/ServiceAccountCollectionTests.cs
@@ -2,10 +2,8 @@
 
 using System.Collections.Generic;
 using System.Threading.Tasks;
-using Horde.Build.Collections;
-using Horde.Build.Collections.Impl;
-using Horde.Build.Models;
 using Horde.Build.Server;
+using Horde.Build.Users;
 using Microsoft.VisualStudio.TestTools.UnitTesting;
 using MongoDB.Bson;
 
diff --git a/Engine/Source/Programs/Horde/Horde.Build.Tests/StreamServiceTest.cs b/Engine/Source/Programs/Horde/Horde.Build.Tests/StreamServiceTest.cs
--- a/Engine/Source/Programs/Horde/Horde.Build.Tests/StreamServiceTest.cs
+++ b/Engine/Source/Programs/Horde/Horde.Build.Tests/StreamServiceTest.cs
@@ -3,7 +3,7 @@
 using System;
 using System.Globalization;
 using System.Threading.Tasks;
-using Horde.Build.Models;
+using Horde.Build.Streams;
 using Microsoft.VisualStudio.TestTools.UnitTesting;
 
 namespace Horde.Build.Tests
diff --git a/Engine/Source/Programs/Horde/Horde.Build.Tests/Stubs/Collections/JobStepRefCollectionStub.cs b/Engine/Source/Programs/Horde/Horde.Build.Tests/Stubs/Collections/JobStepRefCollectionStub.cs
--- a/Engine/Source/Programs/Horde/Horde.Build.Tests/Stubs/Collections/JobStepRefCollectionStub.cs
+++ b/Engine/Source/Programs/Horde/Horde.Build.Tests/Stubs/Collections/JobStepRefCollectionStub.cs
@@ -3,8 +3,11 @@
 using System;
 using System.Collections.Generic;
 using System.Threading.Tasks;
-using Horde.Build.Collections;
-using Horde.Build.Models;
+using Horde.Build.Logs;
+using Horde.Build.Agents;
+using Horde.Build.Agents.Pools;
+using Horde.Build.Streams;
+using Horde.Build.Jobs;
 using Horde.Build.Utilities;
 using HordeCommon;
 
diff --git a/Engine/Source/Programs/Horde/Horde.Build.Tests/Stubs/Collections/SingletonDocumentStub.cs b/Engine/Source/Programs/Horde/Horde.Build.Tests/Stubs/Collections/SingletonDocumentStub.cs
--- a/Engine/Source/Programs/Horde/Horde.Build.Tests/Stubs/Collections/SingletonDocumentStub.cs
+++ b/Engine/Source/Programs/Horde/Horde.Build.Tests/Stubs/Collections/SingletonDocumentStub.cs
@@ -1,7 +1,7 @@
 // Copyright Epic Games, Inc. All Rights Reserved.
 
 using System.Threading.Tasks;
-using Horde.Build.Models;
+using Horde.Build.Server;
 using Horde.Build.Utilities;
 
 namespace Horde.Build.Tests.Stubs.Collections
diff --git a/Engine/Source/Programs/Horde/Horde.Build.Tests/Stubs/Services/PerforceServiceStub.cs b/Engine/Source/Programs/Horde/Horde.Build.Tests/Stubs/Services/PerforceServiceStub.cs
--- a/Engine/Source/Programs/Horde/Horde.Build.Tests/Stubs/Services/PerforceServiceStub.cs
+++ b/Engine/Source/Programs/Horde/Horde.Build.Tests/Stubs/Services/PerforceServiceStub.cs
@@ -5,9 +5,8 @@ using System.Collections.Generic;
 using System.Linq;
 using System.Threading.Tasks;
 using EpicGames.Perforce;
-using Horde.Build.Collections;
-using Horde.Build.Models;
-using Horde.Build.Services;
+using Horde.Build.Perforce;
+using Horde.Build.Users;
 using Horde.Build.Utilities;
 
 namespace Horde.Build.Tests.Stubs.Services
diff --git a/Engine/Source/Programs/Horde/Horde.Build.Tests/TestSetup.cs b/Engine/Source/Programs/Horde/Horde.Build.Tests/TestSetup.cs
--- a/Engine/Source/Programs/Horde/Horde.Build.Tests/TestSetup.cs
+++ b/Engine/Source/Programs/Horde/Horde.Build.Tests/TestSetup.cs
@@ -10,26 +10,20 @@ using Datadog.Trace;
 using EpicGames.Core;
 using EpicGames.Horde.Storage;
 using Horde.Build.Acls;
-using Horde.Build.Collections;
-using Horde.Build.Collections.Impl;
-using Horde.Build.Commits;
-using Horde.Build.Commits.Impl;
-using Horde.Build.Config;
-using Horde.Build.Controllers;
+using Horde.Build.Agents;
+using Horde.Build.Agents.Pools;
+using Horde.Build.Agents.Leases;
+using Horde.Build.Agents.Sessions;
+using Horde.Build.Projects;
+using Horde.Build.Streams;
+using Horde.Build.Issues;
 using Horde.Build.Jobs;
 using Horde.Build.Logs;
 using Horde.Build.Logs.Builder;
-using Horde.Build.Logs.Readers;
-using Horde.Build.Logs.Storage.Impl;
-using Horde.Build.Models;
 using Horde.Build.Notifications;
-using Horde.Build.Notifications.Impl;
 using Horde.Build.Server;
-using Horde.Build.Services;
-using Horde.Build.Services.Impl;
 using Horde.Build.Storage;
 using Horde.Build.Storage.Services;
-using Horde.Build.Tasks.Impl;
 using Horde.Build.Tests.Stubs.Collections;
 using Horde.Build.Tests.Stubs.Services;
 using Horde.Build.Tools;
@@ -44,6 +38,24 @@ using Microsoft.Extensions.Hosting;
 using Microsoft.Extensions.Logging;
 using Microsoft.Extensions.Options;
 using StatsdClient;
+using Horde.Build.Users;
+using Horde.Build.Agents.Software;
+using Horde.Build.Ugs;
+using Horde.Build.Jobs.Timing;
+using Horde.Build.Configuration;
+using Horde.Build.Jobs.Graphs;
+using Horde.Build.Jobs.TestData;
+using Horde.Build.Jobs.Templates;
+using Horde.Build.Jobs.Artifacts;
+using Horde.Build.Secrets;
+using Horde.Build.Jobs.Schedules;
+using Horde.Build.Perforce;
+using Horde.Build.Agents.Fleet.Providers;
+using Horde.Build.Agents.Fleet;
+using Horde.Build.Agents.Telemetry;
+using Horde.Build.Logs.Storage;
+using Horde.Build.Tasks;
+using Horde.Build.Auditing;
 
 namespace Horde.Build.Tests
 {
diff --git a/Engine/Source/Programs/Horde/Horde.Build/Acls/Acl.cs b/Engine/Source/Programs/Horde/Horde.Build/Acls/Acl.cs
--- a/Engine/Source/Programs/Horde/Horde.Build/Acls/Acl.cs
+++ b/Engine/Source/Programs/Horde/Horde.Build/Acls/Acl.cs
@@ -5,7 +5,6 @@ using System.Collections.Generic;
 using System.Linq;
 using System.Linq.Expressions;
 using System.Security.Claims;
-using Horde.Build.Api;
 using MongoDB.Bson;
 using MongoDB.Bson.Serialization;
 using MongoDB.Bson.Serialization.Attributes;
diff --git a/Engine/Source/Programs/Horde/Horde.Build/Acls/AclMessages.cs b/Engine/Source/Programs/Horde/Horde.Build/Acls/AclMessages.cs
--- a/Engine/Source/Programs/Horde/Horde.Build/Acls/AclMessages.cs
+++ b/Engine/Source/Programs/Horde/Horde.Build/Acls/AclMessages.cs
@@ -1,10 +1,9 @@
-﻿// Copyright Epic Games, Inc. All Rights Reserved.
+// Copyright Epic Games, Inc. All Rights Reserved.
 
 using System.Collections.Generic;
 using System.ComponentModel.DataAnnotations;
-using Horde.Build.Acls;
 
-namespace Horde.Build.Api
+namespace Horde.Build.Acls
 {
 	/// <summary>
 	/// New claim to create
diff --git a/Engine/Source/Programs/Horde/Horde.Build/Acls/AclService.cs b/Engine/Source/Programs/Horde/Horde.Build/Acls/AclService.cs
--- a/Engine/Source/Programs/Horde/Horde.Build/Acls/AclService.cs
+++ b/Engine/Source/Programs/Horde/Horde.Build/Acls/AclService.cs
@@ -6,8 +6,10 @@ using System.IdentityModel.Tokens.Jwt;
 using System.Linq;
 using System.Security.Claims;
 using System.Threading.Tasks;
-using Horde.Build.Models;
+using Horde.Build.Agents;
+using Horde.Build.Agents.Sessions;
 using Horde.Build.Server;
+using Horde.Build.Users;
 using Horde.Build.Utilities;
 using Microsoft.Extensions.Options;
 using Microsoft.IdentityModel.Tokens;
diff --git a/Engine/Source/Programs/Horde/Horde.Build/Agents/AgentCollection.cs b/Engine/Source/Programs/Horde/Horde.Build/Agents/AgentCollection.cs
--- a/Engine/Source/Programs/Horde/Horde.Build/Agents/AgentCollection.cs
+++ b/Engine/Source/Programs/Horde/Horde.Build/Agents/AgentCollection.cs
@@ -8,9 +8,11 @@ using System.Threading.Tasks;
 using EpicGames.Redis;
 using Google.Protobuf.WellKnownTypes;
 using Horde.Build.Acls;
-using Horde.Build.Models;
+using Horde.Build.Agents.Pools;
+using Horde.Build.Agents.Sessions;
+using Horde.Build.Agents.Software;
+using Horde.Build.Auditing;
 using Horde.Build.Server;
-using Horde.Build.Services;
 using Horde.Build.Utilities;
 using HordeCommon;
 using HordeCommon.Rpc.Tasks;
@@ -18,7 +20,7 @@ using MongoDB.Bson;
 using MongoDB.Bson.Serialization.Attributes;
 using MongoDB.Driver;
 
-namespace Horde.Build.Collections.Impl
+namespace Horde.Build.Agents
 {
 	using AgentSoftwareChannelName = StringId<AgentSoftwareChannels>;
 	using PoolId = StringId<IPool>;
diff --git a/Engine/Source/Programs/Horde/Horde.Build/Agents/AgentId.cs b/Engine/Source/Programs/Horde/Horde.Build/Agents/AgentId.cs
--- a/Engine/Source/Programs/Horde/Horde.Build/Agents/AgentId.cs
+++ b/Engine/Source/Programs/Horde/Horde.Build/Agents/AgentId.cs
@@ -9,7 +9,7 @@ using MongoDB.Bson.Serialization.Attributes;
 using MongoDB.Bson.Serialization.Serializers;
 using StackExchange.Redis;
 
-namespace Horde.Build.Utilities
+namespace Horde.Build.Agents
 {
 	/// <summary>
 	/// Normalized hostname of an agent
diff --git a/Engine/Source/Programs/Horde/Horde.Build/Agents/AgentService.cs b/Engine/Source/Programs/Horde/Horde.Build/Agents/AgentService.cs
--- a/Engine/Source/Programs/Horde/Horde.Build/Agents/AgentService.cs
+++ b/Engine/Source/Programs/Horde/Horde.Build/Agents/AgentService.cs
@@ -12,8 +12,11 @@ using EpicGames.Redis;
 using EpicGames.Serialization;
 using Google.Protobuf.WellKnownTypes;
 using Horde.Build.Acls;
-using Horde.Build.Collections;
-using Horde.Build.Models;
+using Horde.Build.Agents.Leases;
+using Horde.Build.Agents.Pools;
+using Horde.Build.Agents.Sessions;
+using Horde.Build.Agents.Software;
+using Horde.Build.Auditing;
 using Horde.Build.Server;
 using Horde.Build.Tasks;
 using Horde.Build.Utilities;
@@ -26,7 +29,7 @@ using OpenTracing.Util;
 using StackExchange.Redis;
 using StatsdClient;
 
-namespace Horde.Build.Services
+namespace Horde.Build.Agents
 {
 	using AgentSoftwareChannelName = StringId<AgentSoftwareChannels>;
 	using LeaseId = ObjectId<ILease>;
diff --git a/Engine/Source/Programs/Horde/Horde.Build/Agents/AgentsController.Messages.cs b/Engine/Source/Programs/Horde/Horde.Build/Agents/AgentsController.Messages.cs
--- a/Engine/Source/Programs/Horde/Horde.Build/Agents/AgentsController.Messages.cs
+++ b/Engine/Source/Programs/Horde/Horde.Build/Agents/AgentsController.Messages.cs
@@ -4,11 +4,14 @@ using System;
 using System.Collections.Generic;
 using System.ComponentModel.DataAnnotations;
 using System.Linq;
+using Horde.Build.Acls;
+using Horde.Build.Agents;
 using EpicGames.Core;
-using Horde.Build.Models;
 using HordeCommon;
+using Horde.Build.Agents.Leases;
+using Horde.Build.Agents.Sessions;
 
-namespace Horde.Build.Api
+namespace Horde.Build.Agents
 {
 	/// <summary>
 	/// Parameters to register a new agent
diff --git a/Engine/Source/Programs/Horde/Horde.Build/Agents/AgentsController.cs b/Engine/Source/Programs/Horde/Horde.Build/Agents/AgentsController.cs
--- a/Engine/Source/Programs/Horde/Horde.Build/Agents/AgentsController.cs
+++ b/Engine/Source/Programs/Horde/Horde.Build/Agents/AgentsController.cs
@@ -5,16 +5,17 @@ using System.Collections.Generic;
 using System.Linq;
 using System.Threading.Tasks;
 using Horde.Build.Acls;
-using Horde.Build.Api;
-using Horde.Build.Collections;
-using Horde.Build.Models;
-using Horde.Build.Services;
+using Horde.Build.Agents.Leases;
+using Horde.Build.Agents.Pools;
+using Horde.Build.Agents.Sessions;
+using Horde.Build.Agents.Software;
+using Horde.Build.Auditing;
 using Horde.Build.Utilities;
 using Microsoft.AspNetCore.Authorization;
 using Microsoft.AspNetCore.Mvc;
 using Microsoft.Extensions.Logging;
 
-namespace Horde.Build.Controllers
+namespace Horde.Build.Agents
 {
 	using AgentSoftwareChannelName = StringId<AgentSoftwareChannels>;
 	using LeaseId = ObjectId<ILease>;
diff --git a/Engine/Source/Programs/Horde/Horde.Build/Agents/ConformList.cs b/Engine/Source/Programs/Horde/Horde.Build/Agents/ConformList.cs
--- a/Engine/Source/Programs/Horde/Horde.Build/Agents/ConformList.cs
+++ b/Engine/Source/Programs/Horde/Horde.Build/Agents/ConformList.cs
@@ -2,11 +2,13 @@
 
 using System;
 using System.Collections.Generic;
+using Horde.Build.Agents.Leases;
+using Horde.Build.Server;
 using Horde.Build.Utilities;
 using MongoDB.Bson;
 using MongoDB.Bson.Serialization.Attributes;
 
-namespace Horde.Build.Models
+namespace Horde.Build.Agents
 {
 	using LeaseId = ObjectId<ILease>;
 
diff --git a/Engine/Source/Programs/Horde/Horde.Build/Agents/Fleet/AutoscaleService.cs b/Engine/Source/Programs/Horde/Horde.Build/Agents/Fleet/AutoscaleService.cs
--- a/Engine/Source/Programs/Horde/Horde.Build/Agents/Fleet/AutoscaleService.cs
+++ b/Engine/Source/Programs/Horde/Horde.Build/Agents/Fleet/AutoscaleService.cs
@@ -7,8 +7,8 @@ using System.Linq;
 using System.Threading;
 using System.Threading.Tasks;
 using Google.Protobuf.WellKnownTypes;
-using Horde.Build.Collections;
-using Horde.Build.Models;
+using Horde.Build.Agents.Leases;
+using Horde.Build.Agents.Pools;
 using Horde.Build.Server;
 using Horde.Build.Utilities;
 using HordeCommon;
@@ -19,7 +19,7 @@ using OpenTracing;
 using OpenTracing.Util;
 using StatsdClient;
 
-namespace Horde.Build.Services
+namespace Horde.Build.Agents.Fleet
 {
 	using PoolId = StringId<IPool>;
 
diff --git a/Engine/Source/Programs/Horde/Horde.Build/Agents/Fleet/AutoscaleServiceV2.cs b/Engine/Source/Programs/Horde/Horde.Build/Agents/Fleet/AutoscaleServiceV2.cs
--- a/Engine/Source/Programs/Horde/Horde.Build/Agents/Fleet/AutoscaleServiceV2.cs
+++ b/Engine/Source/Programs/Horde/Horde.Build/Agents/Fleet/AutoscaleServiceV2.cs
@@ -6,9 +6,7 @@ using System.Diagnostics;
 using System.Linq;
 using System.Threading;
 using System.Threading.Tasks;
-using Horde.Build.Collections;
-using Horde.Build.Models;
-using Horde.Build.Services;
+using Horde.Build.Agents.Pools;
 using Horde.Build.Utilities;
 using HordeCommon;
 using Microsoft.Extensions.Hosting;
@@ -18,7 +16,7 @@ using OpenTracing;
 using OpenTracing.Util;
 using StatsdClient;
 
-namespace Horde.Build.Fleet.Autoscale
+namespace Horde.Build.Agents.Fleet
 {
 	using PoolId = StringId<IPool>;
 
diff --git a/Engine/Source/Programs/Horde/Horde.Build/Agents/Fleet/IFleetManager.cs b/Engine/Source/Programs/Horde/Horde.Build/Agents/Fleet/IFleetManager.cs
--- a/Engine/Source/Programs/Horde/Horde.Build/Agents/Fleet/IFleetManager.cs
+++ b/Engine/Source/Programs/Horde/Horde.Build/Agents/Fleet/IFleetManager.cs
@@ -2,9 +2,9 @@
 
 using System.Collections.Generic;
 using System.Threading.Tasks;
-using Horde.Build.Models;
+using Horde.Build.Agents.Pools;
 
-namespace Horde.Build.Services
+namespace Horde.Build.Agents.Fleet
 {
 	/// <summary>
 	/// Service to manage a fleet of machines
diff --git a/Engine/Source/Programs/Horde/Horde.Build/Agents/Fleet/IPoolSizeStrategy.cs b/Engine/Source/Programs/Horde/Horde.Build/Agents/Fleet/IPoolSizeStrategy.cs
--- a/Engine/Source/Programs/Horde/Horde.Build/Agents/Fleet/IPoolSizeStrategy.cs
+++ b/Engine/Source/Programs/Horde/Horde.Build/Agents/Fleet/IPoolSizeStrategy.cs
@@ -3,9 +3,9 @@
 using System.Collections.Generic;
 using System.Linq;
 using System.Threading.Tasks;
-using Horde.Build.Models;
+using Horde.Build.Agents.Pools;
 
-namespace Horde.Build.Fleet.Autoscale
+namespace Horde.Build.Agents.Fleet
 {
 	/// <summary>
 	/// Available pool sizing strategies
@@ -14,19 +14,19 @@ namespace Horde.Build.Fleet.Autoscale
 	{
 		/// <summary>
 		/// Strategy based on lease utilization
-		/// <see cref="Horde.Build.Fleet.Autoscale.LeaseUtilizationStrategy"/>
+		/// <see cref="LeaseUtilizationStrategy"/>
 		/// </summary>
 		LeaseUtilization,
 		
 		/// <summary>
 		/// Strategy based on size of job build queue
-		/// <see cref="Horde.Build.Fleet.Autoscale.JobQueueStrategy"/> 
+		/// <see cref="JobQueueStrategy"/> 
 		/// </summary>
 		JobQueue,
 		
 		/// <summary>
 		/// No-op strategy used as fallback/default behavior
-		/// <see cref="Horde.Build.Fleet.Autoscale.NoOpPoolSizeStrategy"/> 
+		/// <see cref="NoOpPoolSizeStrategy"/> 
 		/// </summary>
 		NoOp
 	}
diff --git a/Engine/Source/Programs/Horde/Horde.Build/Agents/Fleet/JobQueueStrategy.cs b/Engine/Source/Programs/Horde/Horde.Build/Agents/Fleet/JobQueueStrategy.cs
--- a/Engine/Source/Programs/Horde/Horde.Build/Agents/Fleet/JobQueueStrategy.cs
+++ b/Engine/Source/Programs/Horde/Horde.Build/Agents/Fleet/JobQueueStrategy.cs
@@ -4,14 +4,14 @@ using System;
 using System.Collections.Generic;
 using System.Linq;
 using System.Threading.Tasks;
-using Horde.Build.Api;
-using Horde.Build.Collections;
-using Horde.Build.Models;
-using Horde.Build.Services;
+using Horde.Build.Agents.Pools;
+using Horde.Build.Jobs;
+using Horde.Build.Jobs.Graphs;
+using Horde.Build.Streams;
 using Horde.Build.Utilities;
 using HordeCommon;
 
-namespace Horde.Build.Fleet.Autoscale
+namespace Horde.Build.Agents.Fleet
 {
 	using PoolId = StringId<IPool>;
 	using StreamId = StringId<IStream>;
@@ -51,7 +51,7 @@ namespace Horde.Build.Fleet.Autoscale
 	/// Calculate pool size by observing the number of jobs in waiting state
 	///
 	/// Allows for more proactive scaling compared to LeaseUtilizationStrategy.
-	/// <see cref="Horde.Build.Fleet.Autoscale.LeaseUtilizationStrategy"/> 
+	/// <see cref="LeaseUtilizationStrategy"/> 
 	/// </summary>
 	public class JobQueueStrategy : IPoolSizeStrategy
 	{
diff --git a/Engine/Source/Programs/Horde/Horde.Build/Agents/Fleet/LeaseUtilizationStrategy.cs b/Engine/Source/Programs/Horde/Horde.Build/Agents/Fleet/LeaseUtilizationStrategy.cs
--- a/Engine/Source/Programs/Horde/Horde.Build/Agents/Fleet/LeaseUtilizationStrategy.cs
+++ b/Engine/Source/Programs/Horde/Horde.Build/Agents/Fleet/LeaseUtilizationStrategy.cs
@@ -6,15 +6,15 @@ using System.Linq;
 using System.Text;
 using System.Threading.Tasks;
 using Google.Protobuf.WellKnownTypes;
-using Horde.Build.Collections;
-using Horde.Build.Models;
+using Horde.Build.Agents.Leases;
+using Horde.Build.Agents.Pools;
 using Horde.Build.Utilities;
 using HordeCommon;
 using HordeCommon.Rpc.Tasks;
 using OpenTracing;
 using OpenTracing.Util;
 
-namespace Horde.Build.Fleet.Autoscale
+namespace Horde.Build.Agents.Fleet
 {
 	using PoolId = StringId<IPool>;
 
diff --git a/Engine/Source/Programs/Horde/Horde.Build/Agents/Fleet/Providers/AwsFleetManager.cs b/Engine/Source/Programs/Horde/Horde.Build/Agents/Fleet/Providers/AwsFleetManager.cs
--- a/Engine/Source/Programs/Horde/Horde.Build/Agents/Fleet/Providers/AwsFleetManager.cs
+++ b/Engine/Source/Programs/Horde/Horde.Build/Agents/Fleet/Providers/AwsFleetManager.cs
@@ -7,14 +7,13 @@ using System.Threading.Tasks;
 using Amazon;
 using Amazon.EC2;
 using Amazon.EC2.Model;
-using Horde.Build.Collections;
-using Horde.Build.Models;
-using Horde.Build.Utilities;
+using Horde.Build.Agents.Pools;
+using Horde.Build.Auditing;
 using Microsoft.Extensions.Logging;
 using OpenTracing;
 using OpenTracing.Util;
 
-namespace Horde.Build.Services.Impl
+namespace Horde.Build.Agents.Fleet.Providers
 {
 	/// <summary>
 	/// Fleet manager for handling AWS EC2 instances
diff --git a/Engine/Source/Programs/Horde/Horde.Build/Agents/Fleet/Providers/DefaultFleetManager.cs b/Engine/Source/Programs/Horde/Horde.Build/Agents/Fleet/Providers/DefaultFleetManager.cs
--- a/Engine/Source/Programs/Horde/Horde.Build/Agents/Fleet/Providers/DefaultFleetManager.cs
+++ b/Engine/Source/Programs/Horde/Horde.Build/Agents/Fleet/Providers/DefaultFleetManager.cs
@@ -2,10 +2,10 @@
 
 using System.Collections.Generic;
 using System.Threading.Tasks;
-using Horde.Build.Models;
+using Horde.Build.Agents.Pools;
 using Microsoft.Extensions.Logging;
 
-namespace Horde.Build.Services.Impl
+namespace Horde.Build.Agents.Fleet.Providers
 {
 	/// <summary>
 	/// Default implementation of <see cref="IFleetManager"/>
diff --git a/Engine/Source/Programs/Horde/Horde.Build/Agents/IAgent.cs b/Engine/Source/Programs/Horde/Horde.Build/Agents/IAgent.cs
--- a/Engine/Source/Programs/Horde/Horde.Build/Agents/IAgent.cs
+++ b/Engine/Source/Programs/Horde/Horde.Build/Agents/IAgent.cs
@@ -10,14 +10,21 @@ using EpicGames.Horde.Common;
 using EpicGames.Horde.Compute;
 using Google.Protobuf.WellKnownTypes;
 using Horde.Build.Acls;
-using Horde.Build.Services;
+using Horde.Build.Agents.Leases;
+using Horde.Build.Agents.Pools;
+using Horde.Build.Agents.Sessions;
+using Horde.Build.Agents.Software;
+using Horde.Build.Logs;
+using Horde.Build.Perforce;
+using Horde.Build.Server;
+using Horde.Build.Streams;
 using Horde.Build.Utilities;
 using HordeCommon;
 using HordeCommon.Rpc.Tasks;
 using MongoDB.Bson;
 using MongoDB.Bson.Serialization.Attributes;
 
-namespace Horde.Build.Models
+namespace Horde.Build.Agents
 {
 	using AgentSoftwareChannelName = StringId<AgentSoftwareChannels>;
 	using LeaseId = ObjectId<ILease>;
diff --git a/Engine/Source/Programs/Horde/Horde.Build/Agents/IAgentCollection.cs b/Engine/Source/Programs/Horde/Horde.Build/Agents/IAgentCollection.cs
--- a/Engine/Source/Programs/Horde/Horde.Build/Agents/IAgentCollection.cs
+++ b/Engine/Source/Programs/Horde/Horde.Build/Agents/IAgentCollection.cs
@@ -1,15 +1,17 @@
-﻿// Copyright Epic Games, Inc. All Rights Reserved.
+// Copyright Epic Games, Inc. All Rights Reserved.
 
 using System;
 using System.Collections.Generic;
 using System.Threading.Tasks;
 using Horde.Build.Acls;
-using Horde.Build.Models;
-using Horde.Build.Services;
+using Horde.Build.Agents.Pools;
+using Horde.Build.Agents.Sessions;
+using Horde.Build.Agents.Software;
+using Horde.Build.Auditing;
 using Horde.Build.Utilities;
 using HordeCommon;
 
-namespace Horde.Build.Collections
+namespace Horde.Build.Agents
 {
 	using AgentSoftwareChannelName = StringId<AgentSoftwareChannels>;
 	using PoolId = StringId<IPool>;
diff --git a/Engine/Source/Programs/Horde/Horde.Build/Agents/Leases/ILease.cs b/Engine/Source/Programs/Horde/Horde.Build/Agents/Leases/ILease.cs
--- a/Engine/Source/Programs/Horde/Horde.Build/Agents/Leases/ILease.cs
+++ b/Engine/Source/Programs/Horde/Horde.Build/Agents/Leases/ILease.cs
@@ -3,10 +3,14 @@
 using System;
 using Google.Protobuf;
 using Google.Protobuf.WellKnownTypes;
+using Horde.Build.Agents.Pools;
+using Horde.Build.Agents.Sessions;
+using Horde.Build.Logs;
+using Horde.Build.Streams;
 using Horde.Build.Utilities;
 using HordeCommon;
 
-namespace Horde.Build.Models
+namespace Horde.Build.Agents.Leases
 {
 	using LeaseId = ObjectId<ILease>;
 	using LogId = ObjectId<ILogFile>;
diff --git a/Engine/Source/Programs/Horde/Horde.Build/Agents/Leases/ILeaseCollection.cs b/Engine/Source/Programs/Horde/Horde.Build/Agents/Leases/ILeaseCollection.cs
--- a/Engine/Source/Programs/Horde/Horde.Build/Agents/Leases/ILeaseCollection.cs
+++ b/Engine/Source/Programs/Horde/Horde.Build/Agents/Leases/ILeaseCollection.cs
@@ -3,11 +3,14 @@
 using System;
 using System.Collections.Generic;
 using System.Threading.Tasks;
-using Horde.Build.Models;
+using Horde.Build.Agents.Pools;
+using Horde.Build.Agents.Sessions;
+using Horde.Build.Logs;
+using Horde.Build.Streams;
 using Horde.Build.Utilities;
 using HordeCommon;
 
-namespace Horde.Build.Collections
+namespace Horde.Build.Agents.Leases
 {
 	using LeaseId = ObjectId<ILease>;
 	using LogId = ObjectId<ILogFile>;
diff --git a/Engine/Source/Programs/Horde/Horde.Build/Agents/Leases/LeaseCollection.cs b/Engine/Source/Programs/Horde/Horde.Build/Agents/Leases/LeaseCollection.cs
--- a/Engine/Source/Programs/Horde/Horde.Build/Agents/Leases/LeaseCollection.cs
+++ b/Engine/Source/Programs/Horde/Horde.Build/Agents/Leases/LeaseCollection.cs
@@ -4,15 +4,17 @@ using System;
 using System.Collections.Generic;
 using System.Linq;
 using System.Threading.Tasks;
-using Horde.Build.Models;
+using Horde.Build.Agents.Pools;
+using Horde.Build.Agents.Sessions;
+using Horde.Build.Logs;
 using Horde.Build.Server;
+using Horde.Build.Streams;
 using Horde.Build.Utilities;
 using HordeCommon;
-using MongoDB.Bson;
 using MongoDB.Bson.Serialization.Attributes;
 using MongoDB.Driver;
 
-namespace Horde.Build.Collections.Impl
+namespace Horde.Build.Agents.Leases
 {
 	using LeaseId = ObjectId<ILease>;
 	using LogId = ObjectId<ILogFile>;
diff --git a/Engine/Source/Programs/Horde/Horde.Build/Agents/Leases/LeasesController.Messages.cs b/Engine/Source/Programs/Horde/Horde.Build/Agents/Leases/LeasesController.Messages.cs
--- a/Engine/Source/Programs/Horde/Horde.Build/Agents/Leases/LeasesController.Messages.cs
+++ b/Engine/Source/Programs/Horde/Horde.Build/Agents/Leases/LeasesController.Messages.cs
@@ -1,6 +1,6 @@
 // Copyright Epic Games, Inc. All Rights Reserved.
 
-namespace Horde.Build.Api
+namespace Horde.Build.Agents.Leases
 {
 
 	/// <summary>
diff --git a/Engine/Source/Programs/Horde/Horde.Build/Agents/Leases/LeasesController.cs b/Engine/Source/Programs/Horde/Horde.Build/Agents/Leases/LeasesController.cs
--- a/Engine/Source/Programs/Horde/Horde.Build/Agents/Leases/LeasesController.cs
+++ b/Engine/Source/Programs/Horde/Horde.Build/Agents/Leases/LeasesController.cs
@@ -5,16 +5,14 @@ using System.Collections.Generic;
 using System.Linq;
 using System.Threading.Tasks;
 using Horde.Build.Acls;
-using Horde.Build.Api;
-using Horde.Build.Models;
-using Horde.Build.Services;
+using Horde.Build.Agents.Sessions;
 using Horde.Build.Utilities;
 using Microsoft.AspNetCore.Authorization;
 using Microsoft.AspNetCore.Mvc;
 using OpenTracing;
 using OpenTracing.Util;
 
-namespace Horde.Build.Controllers
+namespace Horde.Build.Agents.Leases
 {
 	using LeaseId = ObjectId<ILease>;
 	using SessionId = ObjectId<ISession>;
diff --git a/Engine/Source/Programs/Horde/Horde.Build/Agents/Pools/IPool.cs b/Engine/Source/Programs/Horde/Horde.Build/Agents/Pools/IPool.cs
--- a/Engine/Source/Programs/Horde/Horde.Build/Agents/Pools/IPool.cs
+++ b/Engine/Source/Programs/Horde/Horde.Build/Agents/Pools/IPool.cs
@@ -3,10 +3,10 @@
 using System;
 using System.Collections.Generic;
 using EpicGames.Horde.Common;
-using Horde.Build.Fleet.Autoscale;
+using Horde.Build.Agents.Fleet;
 using Horde.Build.Utilities;
 
-namespace Horde.Build.Models
+namespace Horde.Build.Agents.Pools
 {
 	using PoolId = StringId<IPool>;
 
diff --git a/Engine/Source/Programs/Horde/Horde.Build/Agents/Pools/IPoolCollection.cs b/Engine/Source/Programs/Horde/Horde.Build/Agents/Pools/IPoolCollection.cs
--- a/Engine/Source/Programs/Horde/Horde.Build/Agents/Pools/IPoolCollection.cs
+++ b/Engine/Source/Programs/Horde/Horde.Build/Agents/Pools/IPoolCollection.cs
@@ -4,11 +4,10 @@ using System;
 using System.Collections.Generic;
 using System.Threading.Tasks;
 using EpicGames.Horde.Common;
-using Horde.Build.Fleet.Autoscale;
-using Horde.Build.Models;
+using Horde.Build.Agents.Fleet;
 using Horde.Build.Utilities;
 
-namespace Horde.Build.Collections
+namespace Horde.Build.Agents.Pools
 {
 	using PoolId = StringId<IPool>;
 
diff --git a/Engine/Source/Programs/Horde/Horde.Build/Agents/Pools/PoolCollection.cs b/Engine/Source/Programs/Horde/Horde.Build/Agents/Pools/PoolCollection.cs
--- a/Engine/Source/Programs/Horde/Horde.Build/Agents/Pools/PoolCollection.cs
+++ b/Engine/Source/Programs/Horde/Horde.Build/Agents/Pools/PoolCollection.cs
@@ -5,8 +5,7 @@ using System.Collections.Generic;
 using System.Linq;
 using System.Threading.Tasks;
 using EpicGames.Horde.Common;
-using Horde.Build.Fleet.Autoscale;
-using Horde.Build.Models;
+using Horde.Build.Agents.Fleet;
 using Horde.Build.Server;
 using Horde.Build.Utilities;
 using Microsoft.Extensions.Options;
@@ -14,7 +13,7 @@ using MongoDB.Bson;
 using MongoDB.Bson.Serialization.Attributes;
 using MongoDB.Driver;
 
-namespace Horde.Build.Collections.Impl
+namespace Horde.Build.Agents.Pools
 {
 	using PoolId = StringId<IPool>;
 
diff --git a/Engine/Source/Programs/Horde/Horde.Build/Agents/Pools/PoolService.cs b/Engine/Source/Programs/Horde/Horde.Build/Agents/Pools/PoolService.cs
--- a/Engine/Source/Programs/Horde/Horde.Build/Agents/Pools/PoolService.cs
+++ b/Engine/Source/Programs/Horde/Horde.Build/Agents/Pools/PoolService.cs
@@ -6,14 +6,12 @@ using System.Linq;
 using System.Threading;
 using System.Threading.Tasks;
 using EpicGames.Horde.Common;
-using Horde.Build.Collections;
-using Horde.Build.Fleet.Autoscale;
-using Horde.Build.Models;
+using Horde.Build.Agents.Fleet;
 using Horde.Build.Server;
 using Horde.Build.Utilities;
 using HordeCommon;
 
-namespace Horde.Build.Services
+namespace Horde.Build.Agents.Pools
 {
 	using PoolId = StringId<IPool>;
 
diff --git a/Engine/Source/Programs/Horde/Horde.Build/Agents/Pools/PoolUpdateService.cs b/Engine/Source/Programs/Horde/Horde.Build/Agents/Pools/PoolUpdateService.cs
--- a/Engine/Source/Programs/Horde/Horde.Build/Agents/Pools/PoolUpdateService.cs
+++ b/Engine/Source/Programs/Horde/Horde.Build/Agents/Pools/PoolUpdateService.cs
@@ -5,14 +5,13 @@ using System.Collections.Generic;
 using System.Linq;
 using System.Threading;
 using System.Threading.Tasks;
-using Horde.Build.Collections;
-using Horde.Build.Models;
+using Horde.Build.Streams;
 using Horde.Build.Utilities;
 using HordeCommon;
 using Microsoft.Extensions.Hosting;
 using Microsoft.Extensions.Logging;
 
-namespace Horde.Build.Services
+namespace Horde.Build.Agents.Pools
 {
 	using PoolId = StringId<IPool>;
 
diff --git a/Engine/Source/Programs/Horde/Horde.Build/Agents/Pools/PoolsController.Messages.cs b/Engine/Source/Programs/Horde/Horde.Build/Agents/Pools/PoolsController.Messages.cs
--- a/Engine/Source/Programs/Horde/Horde.Build/Agents/Pools/PoolsController.Messages.cs
+++ b/Engine/Source/Programs/Horde/Horde.Build/Agents/Pools/PoolsController.Messages.cs
@@ -4,20 +4,19 @@ using System.Collections.Generic;
 using System.ComponentModel.DataAnnotations;
 using System.Linq;
 using EpicGames.Horde.Common;
-using Horde.Build.Fleet.Autoscale;
-using Horde.Build.Models;
+using Horde.Build.Agents.Fleet;
 
-namespace Horde.Build.Api
+namespace Horde.Build.Agents.Pools
 {
-	/// <see cref="Horde.Build.Fleet.Autoscale.LeaseUtilizationSettings" />
-	public class LeaseUtilizationSettings
+	/// <see cref="LeaseUtilizationSettings" />
+	public class LeaseUtilizationSettingsMessage
 	{
 		/// <summary>
 		/// Construct a public REST API representation from the internal one
 		/// </summary>
 		/// <param name="settings"></param>
 #pragma warning disable IDE0060 // Remove unused parameter
-		public LeaseUtilizationSettings(Horde.Build.Fleet.Autoscale.LeaseUtilizationSettings settings)
+		public LeaseUtilizationSettingsMessage(LeaseUtilizationSettings settings)
 #pragma warning restore IDE0060 // Remove unused parameter
 		{
 		}
@@ -26,14 +25,14 @@ namespace Horde.Build.Api
 		/// Convert this public REST API object to an internal representation
 		/// </summary>
 		/// <returns></returns>
-		public Horde.Build.Fleet.Autoscale.LeaseUtilizationSettings Convert()
+		public Horde.Build.Agents.Fleet.LeaseUtilizationSettings Convert()
 		{
-			return new Horde.Build.Fleet.Autoscale.LeaseUtilizationSettings();
+			return new Horde.Build.Agents.Fleet.LeaseUtilizationSettings();
 		}
 	}
 	
-	/// <see cref="Horde.Build.Fleet.Autoscale.JobQueueSettings" />
-	public class JobQueueSettings
+	/// <see cref="JobQueueSettings" />
+	public class JobQueueSettingsMessage
 	{
 		/// <summary>
 		/// Factor translating queue size to additional agents to grow the pool with
@@ -53,7 +52,7 @@ namespace Horde.Build.Api
 		/// Construct a public REST API representation from the internal one
 		/// </summary>
 		/// <param name="settings"></param>
-		public JobQueueSettings(Horde.Build.Fleet.Autoscale.JobQueueSettings settings)
+		public JobQueueSettingsMessage(JobQueueSettings settings)
 		{
 			ScaleInFactor = settings.ScaleInFactor;
 			ScaleOutFactor = settings.ScaleOutFactor;
@@ -63,9 +62,9 @@ namespace Horde.Build.Api
 		/// Convert this public REST API object to an internal representation
 		/// </summary>
 		/// <returns></returns>
-		public Horde.Build.Fleet.Autoscale.JobQueueSettings Convert()
+		public JobQueueSettings Convert()
 		{
-			return new Horde.Build.Fleet.Autoscale.JobQueueSettings(ScaleOutFactor, ScaleInFactor);
+			return new JobQueueSettings(ScaleOutFactor, ScaleInFactor);
 		}
 	}
 	
@@ -108,12 +107,12 @@ namespace Horde.Build.Api
 		/// <summary>
 		/// Settings for lease utilization pool sizing strategy (if used)
 		/// </summary>
-		public LeaseUtilizationSettings? LeaseUtilizationSettings { get; set; }
+		public LeaseUtilizationSettingsMessage? LeaseUtilizationSettings { get; set; }
 		
 		/// <summary>
 		/// Settings for job queue pool sizing strategy (if used) 
 		/// </summary>
-		public JobQueueSettings? JobQueueSettings { get; set; }
+		public JobQueueSettingsMessage? JobQueueSettings { get; set; }
 
 		/// <summary>
 		/// The minimum nunmber of agents to retain in this pool
@@ -282,12 +281,12 @@ namespace Horde.Build.Api
 		/// <summary>
 		/// Settings for lease utilization pool sizing strategy (if used)
 		/// </summary>
-		public LeaseUtilizationSettings? LeaseUtilizationSettings { get; set; }
+		public LeaseUtilizationSettingsMessage? LeaseUtilizationSettings { get; set; }
 		
 		/// <summary>
 		/// Settings for job queue pool sizing strategy (if used) 
 		/// </summary>
-		public JobQueueSettings? JobQueueSettings { get; set; }
+		public JobQueueSettingsMessage? JobQueueSettings { get; set; }
 
 		/// <summary>
 		/// The minimum nunmber of agents to retain in this pool
@@ -322,8 +321,8 @@ namespace Horde.Build.Api
 			ScaleOutCooldown = pool.ScaleOutCooldown == null ? null : (int)pool.ScaleOutCooldown.Value.TotalSeconds;
 			ScaleInCooldown = pool.ScaleInCooldown == null ? null : (int)pool.ScaleInCooldown.Value.TotalSeconds;
 			SizeStrategy = pool.SizeStrategy;
-			LeaseUtilizationSettings = pool.LeaseUtilizationSettings == null ? null : new LeaseUtilizationSettings(pool.LeaseUtilizationSettings);
-			JobQueueSettings = pool.JobQueueSettings == null ? null : new JobQueueSettings(pool.JobQueueSettings);
+			LeaseUtilizationSettings = pool.LeaseUtilizationSettings == null ? null : new LeaseUtilizationSettingsMessage(pool.LeaseUtilizationSettings);
+			JobQueueSettings = pool.JobQueueSettings == null ? null : new JobQueueSettingsMessage(pool.JobQueueSettings);
 			MinAgents = pool.MinAgents;
 			NumReserveAgents = pool.NumReserveAgents;
 			Workspaces = pool.Workspaces.Select(x => new GetAgentWorkspaceResponse(x)).ToList();
diff --git a/Engine/Source/Programs/Horde/Horde.Build/Agents/Pools/PoolsController.cs b/Engine/Source/Programs/Horde/Horde.Build/Agents/Pools/PoolsController.cs
--- a/Engine/Source/Programs/Horde/Horde.Build/Agents/Pools/PoolsController.cs
+++ b/Engine/Source/Programs/Horde/Horde.Build/Agents/Pools/PoolsController.cs
@@ -4,14 +4,12 @@ using System;
 using System.Collections.Generic;
 using System.Threading.Tasks;
 using Horde.Build.Acls;
-using Horde.Build.Api;
-using Horde.Build.Models;
-using Horde.Build.Services;
+using Horde.Build.Agents.Fleet;
 using Horde.Build.Utilities;
 using Microsoft.AspNetCore.Authorization;
 using Microsoft.AspNetCore.Mvc;
 
-namespace Horde.Build.Controllers
+namespace Horde.Build.Agents.Pools
 {
 	using PoolId = StringId<IPool>;
 
@@ -58,8 +56,8 @@ namespace Horde.Build.Controllers
 				return Forbid(AclAction.CreatePool);
 			}
 
-			Horde.Build.Fleet.Autoscale.LeaseUtilizationSettings? luSettings = create.LeaseUtilizationSettings?.Convert();
-			Horde.Build.Fleet.Autoscale.JobQueueSettings? jqSettings = create.JobQueueSettings?.Convert();
+			LeaseUtilizationSettings? luSettings = create.LeaseUtilizationSettings?.Convert();
+			JobQueueSettings? jqSettings = create.JobQueueSettings?.Convert();
 			
 			TimeSpan? scaleOutCooldown = create.ScaleOutCooldown == null ? null : TimeSpan.FromSeconds(create.ScaleOutCooldown.Value);
 			TimeSpan? scaleInCooldown = create.ScaleInCooldown == null ? null : TimeSpan.FromSeconds(create.ScaleInCooldown.Value);
diff --git a/Engine/Source/Programs/Horde/Horde.Build/Agents/Sessions/ISession.cs b/Engine/Source/Programs/Horde/Horde.Build/Agents/Sessions/ISession.cs
--- a/Engine/Source/Programs/Horde/Horde.Build/Agents/Sessions/ISession.cs
+++ b/Engine/Source/Programs/Horde/Horde.Build/Agents/Sessions/ISession.cs
@@ -4,7 +4,7 @@ using System;
 using System.Collections.Generic;
 using Horde.Build.Utilities;
 
-namespace Horde.Build.Models
+namespace Horde.Build.Agents.Sessions
 {
 	using SessionId = ObjectId<ISession>;
 
diff --git a/Engine/Source/Programs/Horde/Horde.Build/Agents/Sessions/ISessionCollection.cs b/Engine/Source/Programs/Horde/Horde.Build/Agents/Sessions/ISessionCollection.cs
--- a/Engine/Source/Programs/Horde/Horde.Build/Agents/Sessions/ISessionCollection.cs
+++ b/Engine/Source/Programs/Horde/Horde.Build/Agents/Sessions/ISessionCollection.cs
@@ -3,10 +3,9 @@
 using System;
 using System.Collections.Generic;
 using System.Threading.Tasks;
-using Horde.Build.Models;
 using Horde.Build.Utilities;
 
-namespace Horde.Build.Collections
+namespace Horde.Build.Agents.Sessions
 {
 	using SessionId = ObjectId<ISession>;
 
diff --git a/Engine/Source/Programs/Horde/Horde.Build/Agents/Sessions/SessionCollection.cs b/Engine/Source/Programs/Horde/Horde.Build/Agents/Sessions/SessionCollection.cs
--- a/Engine/Source/Programs/Horde/Horde.Build/Agents/Sessions/SessionCollection.cs
+++ b/Engine/Source/Programs/Horde/Horde.Build/Agents/Sessions/SessionCollection.cs
@@ -4,13 +4,12 @@ using System;
 using System.Collections.Generic;
 using System.Linq;
 using System.Threading.Tasks;
-using Horde.Build.Models;
 using Horde.Build.Server;
 using Horde.Build.Utilities;
 using MongoDB.Bson.Serialization.Attributes;
 using MongoDB.Driver;
 
-namespace Horde.Build.Collections.Impl
+namespace Horde.Build.Agents.Sessions
 {
 	using SessionId = ObjectId<ISession>;
 
diff --git a/Engine/Source/Programs/Horde/Horde.Build/Agents/Software/AgentSoftwareCollection.cs b/Engine/Source/Programs/Horde/Horde.Build/Agents/Software/AgentSoftwareCollection.cs
--- a/Engine/Source/Programs/Horde/Horde.Build/Agents/Software/AgentSoftwareCollection.cs
+++ b/Engine/Source/Programs/Horde/Horde.Build/Agents/Software/AgentSoftwareCollection.cs
@@ -10,7 +10,7 @@ using MongoDB.Bson;
 using MongoDB.Bson.Serialization.Attributes;
 using MongoDB.Driver;
 
-namespace Horde.Build.Collections.Impl
+namespace Horde.Build.Agents.Software
 {
 	/// <summary>
 	/// Collection of agent software documents
diff --git a/Engine/Source/Programs/Horde/Horde.Build/Agents/Software/AgentSoftwareController.Messages.cs b/Engine/Source/Programs/Horde/Horde.Build/Agents/Software/AgentSoftwareController.Messages.cs
--- a/Engine/Source/Programs/Horde/Horde.Build/Agents/Software/AgentSoftwareController.Messages.cs
+++ b/Engine/Source/Programs/Horde/Horde.Build/Agents/Software/AgentSoftwareController.Messages.cs
@@ -1,9 +1,8 @@
 // Copyright Epic Games, Inc. All Rights Reserved.
 
 using System;
-using Horde.Build.Models;
 
-namespace Horde.Build.Api
+namespace Horde.Build.Agents.Software
 {
 	/// <summary>
 	/// Information about an agent software channel
diff --git a/Engine/Source/Programs/Horde/Horde.Build/Agents/Software/AgentSoftwareController.cs b/Engine/Source/Programs/Horde/Horde.Build/Agents/Software/AgentSoftwareController.cs
--- a/Engine/Source/Programs/Horde/Horde.Build/Agents/Software/AgentSoftwareController.cs
+++ b/Engine/Source/Programs/Horde/Horde.Build/Agents/Software/AgentSoftwareController.cs
@@ -1,19 +1,16 @@
-﻿// Copyright Epic Games, Inc. All Rights Reserved.
+// Copyright Epic Games, Inc. All Rights Reserved.
 
 using System.Collections.Generic;
 using System.IO;
 using System.Threading.Tasks;
 using Horde.Build.Acls;
-using Horde.Build.Api;
-using Horde.Build.Models;
-using Horde.Build.Services;
 using Horde.Build.Utilities;
 using Microsoft.AspNetCore.Authorization;
 using Microsoft.AspNetCore.Http;
 using Microsoft.AspNetCore.Mvc;
 using Microsoft.Net.Http.Headers;
 
-namespace Horde.Build.Controllers
+namespace Horde.Build.Agents.Software
 {
 	using AgentSoftwareChannelName = StringId<AgentSoftwareChannels>;
 
diff --git a/Engine/Source/Programs/Horde/Horde.Build/Agents/Software/AgentSoftwareService.cs b/Engine/Source/Programs/Horde/Horde.Build/Agents/Software/AgentSoftwareService.cs
--- a/Engine/Source/Programs/Horde/Horde.Build/Agents/Software/AgentSoftwareService.cs
+++ b/Engine/Source/Programs/Horde/Horde.Build/Agents/Software/AgentSoftwareService.cs
@@ -7,16 +7,15 @@ using System.Linq;
 using System.Security.Cryptography.X509Certificates;
 using System.Threading.Tasks;
 using EpicGames.Core;
-using Horde.Build.Collections;
 using Horde.Build.Commands;
-using Horde.Build.Models;
+using Horde.Build.Server;
 using Horde.Build.Utilities;
 using Microsoft.Extensions.Logging;
 using Microsoft.Extensions.Options;
 using MongoDB.Bson.Serialization.Attributes;
 using MongoDB.Driver;
 
-namespace Horde.Build.Services
+namespace Horde.Build.Agents.Software
 {
 	using AgentSoftwareChannelName = StringId<AgentSoftwareChannels>;
 
diff --git a/Engine/Source/Programs/Horde/Horde.Build/Agents/Software/IAgentSoftwareChannel.cs b/Engine/Source/Programs/Horde/Horde.Build/Agents/Software/IAgentSoftwareChannel.cs
--- a/Engine/Source/Programs/Horde/Horde.Build/Agents/Software/IAgentSoftwareChannel.cs
+++ b/Engine/Source/Programs/Horde/Horde.Build/Agents/Software/IAgentSoftwareChannel.cs
@@ -1,10 +1,9 @@
 // Copyright Epic Games, Inc. All Rights Reserved.
 
 using System;
-using Horde.Build.Services;
 using Horde.Build.Utilities;
 
-namespace Horde.Build.Models
+namespace Horde.Build.Agents.Software
 {
 	using AgentSoftwareChannelName = StringId<AgentSoftwareChannels>;
 
diff --git a/Engine/Source/Programs/Horde/Horde.Build/Agents/Software/IAgentSoftwareCollection.cs b/Engine/Source/Programs/Horde/Horde.Build/Agents/Software/IAgentSoftwareCollection.cs
--- a/Engine/Source/Programs/Horde/Horde.Build/Agents/Software/IAgentSoftwareCollection.cs
+++ b/Engine/Source/Programs/Horde/Horde.Build/Agents/Software/IAgentSoftwareCollection.cs
@@ -2,7 +2,7 @@
 
 using System.Threading.Tasks;
 
-namespace Horde.Build.Collections
+namespace Horde.Build.Agents.Software
 {
 	/// <summary>
 	/// Interface for a collection of template documents
diff --git a/Engine/Source/Programs/Horde/Horde.Build/Agents/Telemetry/ITelemetryCollection.cs b/Engine/Source/Programs/Horde/Horde.Build/Agents/Telemetry/ITelemetryCollection.cs
--- a/Engine/Source/Programs/Horde/Horde.Build/Agents/Telemetry/ITelemetryCollection.cs
+++ b/Engine/Source/Programs/Horde/Horde.Build/Agents/Telemetry/ITelemetryCollection.cs
@@ -3,9 +3,8 @@
 using System;
 using System.Collections.Generic;
 using System.Threading.Tasks;
-using Horde.Build.Models;
 
-namespace Horde.Build.Collections
+namespace Horde.Build.Agents.Telemetry
 {
 	/// <summary>
 	/// Collection of utilization collection
diff --git a/Engine/Source/Programs/Horde/Horde.Build/Agents/Telemetry/Telemetry.cs b/Engine/Source/Programs/Horde/Horde.Build/Agents/Telemetry/Telemetry.cs
--- a/Engine/Source/Programs/Horde/Horde.Build/Agents/Telemetry/Telemetry.cs
+++ b/Engine/Source/Programs/Horde/Horde.Build/Agents/Telemetry/Telemetry.cs
@@ -2,9 +2,11 @@
 
 using System;
 using System.Collections.Generic;
+using Horde.Build.Agents.Pools;
+using Horde.Build.Streams;
 using Horde.Build.Utilities;
 
-namespace Horde.Build.Models
+namespace Horde.Build.Agents.Telemetry
 {
 	using PoolId = StringId<IPool>;
 	using StreamId = StringId<IStream>;
diff --git a/Engine/Source/Programs/Horde/Horde.Build/Agents/Telemetry/TelemetryCollection.cs b/Engine/Source/Programs/Horde/Horde.Build/Agents/Telemetry/TelemetryCollection.cs
--- a/Engine/Source/Programs/Horde/Horde.Build/Agents/Telemetry/TelemetryCollection.cs
+++ b/Engine/Source/Programs/Horde/Horde.Build/Agents/Telemetry/TelemetryCollection.cs
@@ -5,13 +5,14 @@ using System.Collections.Generic;
 using System.Linq;
 using System.Threading.Tasks;
 using EpicGames.Core;
-using Horde.Build.Models;
+using Horde.Build.Agents.Pools;
 using Horde.Build.Server;
+using Horde.Build.Streams;
 using Horde.Build.Utilities;
 using MongoDB.Bson.Serialization.Attributes;
 using MongoDB.Driver;
 
-namespace Horde.Build.Collections.Impl
+namespace Horde.Build.Agents.Telemetry
 {
 	using PoolId = StringId<IPool>;
 	using StreamId = StringId<IStream>;
diff --git a/Engine/Source/Programs/Horde/Horde.Build/Agents/Telemetry/TelemetryController.Messages.cs b/Engine/Source/Programs/Horde/Horde.Build/Agents/Telemetry/TelemetryController.Messages.cs
--- a/Engine/Source/Programs/Horde/Horde.Build/Agents/Telemetry/TelemetryController.Messages.cs
+++ b/Engine/Source/Programs/Horde/Horde.Build/Agents/Telemetry/TelemetryController.Messages.cs
@@ -3,9 +3,8 @@
 using System;
 using System.Collections.Generic;
 using EpicGames.Core;
-using Horde.Build.Models;
 
-namespace Horde.Build.Api
+namespace Horde.Build.Agents.Telemetry
 {
 
 	/// <summary>
diff --git a/Engine/Source/Programs/Horde/Horde.Build/Agents/Telemetry/TelemetryController.cs b/Engine/Source/Programs/Horde/Horde.Build/Agents/Telemetry/TelemetryController.cs
--- a/Engine/Source/Programs/Horde/Horde.Build/Agents/Telemetry/TelemetryController.cs
+++ b/Engine/Source/Programs/Horde/Horde.Build/Agents/Telemetry/TelemetryController.cs
@@ -3,12 +3,9 @@
 using System;
 using System.Collections.Generic;
 using System.Threading.Tasks;
-using Horde.Build.Api;
-using Horde.Build.Collections;
-using Horde.Build.Models;
 using Microsoft.AspNetCore.Mvc;
 
-namespace Horde.Build.Controllers
+namespace Horde.Build.Agents.Telemetry
 {
 	/// <summary>
 	/// Controller for the /api/v1/reports endpoint, used for the reports pages
diff --git a/Engine/Source/Programs/Horde/Horde.Build/Agents/Telemetry/TelemetryService.cs b/Engine/Source/Programs/Horde/Horde.Build/Agents/Telemetry/TelemetryService.cs
--- a/Engine/Source/Programs/Horde/Horde.Build/Agents/Telemetry/TelemetryService.cs
+++ b/Engine/Source/Programs/Horde/Horde.Build/Agents/Telemetry/TelemetryService.cs
@@ -5,14 +5,15 @@ using System.Collections.Generic;
 using System.Linq;
 using System.Threading;
 using System.Threading.Tasks;
-using Horde.Build.Collections;
-using Horde.Build.Models;
+using Horde.Build.Agents.Fleet;
+using Horde.Build.Agents.Leases;
+using Horde.Build.Agents.Pools;
 using Horde.Build.Utilities;
 using HordeCommon;
 using Microsoft.Extensions.Hosting;
 using Microsoft.Extensions.Logging;
 
-namespace Horde.Build.Services
+namespace Horde.Build.Agents.Telemetry
 {
 	using PoolId = StringId<IPool>;
 
diff --git a/Engine/Source/Programs/Horde/Horde.Build/Api/Changes.cs b/Engine/Source/Programs/Horde/Horde.Build/Api/Changes.cs
--- a/Engine/Source/Programs/Horde/Horde.Build/Api/Changes.cs
+++ b/Engine/Source/Programs/Horde/Horde.Build/Api/Changes.cs
@@ -1,9 +1,10 @@
 // Copyright Epic Games, Inc. All Rights Reserved.
 
 using System.Collections.Generic;
-using Horde.Build.Models;
+using Horde.Build.Perforce;
+using Horde.Build.Users;
 
-namespace Horde.Build.Api
+namespace Horde.Build.Perforce
 {
 	/// <summary>
 	/// Information about a submitted changelist
diff --git a/Engine/Source/Programs/Horde/Horde.Build/Api/Config.cs b/Engine/Source/Programs/Horde/Horde.Build/Api/Config.cs
--- a/Engine/Source/Programs/Horde/Horde.Build/Api/Config.cs
+++ b/Engine/Source/Programs/Horde/Horde.Build/Api/Config.cs
@@ -3,7 +3,7 @@
 using System;
 using System.Collections.Generic;
 
-namespace Horde.Build.Api
+namespace Horde.Build.Configuration
 {
 	/// <summary>
 	/// Server settings 
diff --git a/Engine/Source/Programs/Horde/Horde.Build/Auditing/AuditLog.cs b/Engine/Source/Programs/Horde/Horde.Build/Auditing/AuditLog.cs
--- a/Engine/Source/Programs/Horde/Horde.Build/Auditing/AuditLog.cs
+++ b/Engine/Source/Programs/Horde/Horde.Build/Auditing/AuditLog.cs
@@ -12,7 +12,7 @@ using MongoDB.Bson;
 using MongoDB.Bson.Serialization.Attributes;
 using MongoDB.Driver;
 
-namespace Horde.Build.Collections.Impl
+namespace Horde.Build.Auditing
 {
 	class AuditLog<TSubject> : IAuditLog<TSubject>, IAsyncDisposable, IDisposable
 	{
diff --git a/Engine/Source/Programs/Horde/Horde.Build/Auditing/IAuditLog.cs b/Engine/Source/Programs/Horde/Horde.Build/Auditing/IAuditLog.cs
--- a/Engine/Source/Programs/Horde/Horde.Build/Auditing/IAuditLog.cs
+++ b/Engine/Source/Programs/Horde/Horde.Build/Auditing/IAuditLog.cs
@@ -7,7 +7,7 @@ using System.Text;
 using System.Threading.Tasks;
 using Microsoft.Extensions.Logging;
 
-namespace Horde.Build.Collections
+namespace Horde.Build.Auditing
 {
 	/// <summary>
 	/// Message from an audit log
diff --git a/Engine/Source/Programs/Horde/Horde.Build/Authentication/AnonymousHandler.cs b/Engine/Source/Programs/Horde/Horde.Build/Authentication/AnonymousHandler.cs
--- a/Engine/Source/Programs/Horde/Horde.Build/Authentication/AnonymousHandler.cs
+++ b/Engine/Source/Programs/Horde/Horde.Build/Authentication/AnonymousHandler.cs
@@ -5,8 +5,7 @@ using System.Collections.Generic;
 using System.Security.Claims;
 using System.Text.Encodings.Web;
 using System.Threading.Tasks;
-using Horde.Build.Collections;
-using Horde.Build.Models;
+using Horde.Build.Users;
 using Horde.Build.Utilities;
 using Microsoft.AspNetCore.Authentication;
 using Microsoft.Extensions.Logging;
diff --git a/Engine/Source/Programs/Horde/Horde.Build/Authentication/HordeOpenIdConnectHandler.cs b/Engine/Source/Programs/Horde/Horde.Build/Authentication/HordeOpenIdConnectHandler.cs
--- a/Engine/Source/Programs/Horde/Horde.Build/Authentication/HordeOpenIdConnectHandler.cs
+++ b/Engine/Source/Programs/Horde/Horde.Build/Authentication/HordeOpenIdConnectHandler.cs
@@ -6,8 +6,7 @@ using System.Security.Claims;
 using System.Text.Encodings.Web;
 using System.Text.Json;
 using System.Threading.Tasks;
-using Horde.Build.Collections;
-using Horde.Build.Models;
+using Horde.Build.Users;
 using Horde.Build.Utilities;
 using Microsoft.AspNetCore.Authentication;
 using Microsoft.AspNetCore.Authentication.OAuth.Claims;
diff --git a/Engine/Source/Programs/Horde/Horde.Build/Commands/Config/SchemasCommand.cs b/Engine/Source/Programs/Horde/Horde.Build/Commands/Config/SchemasCommand.cs
--- a/Engine/Source/Programs/Horde/Horde.Build/Commands/Config/SchemasCommand.cs
+++ b/Engine/Source/Programs/Horde/Horde.Build/Commands/Config/SchemasCommand.cs
@@ -4,6 +4,7 @@ using System;
 using System.Threading.Tasks;
 using EpicGames.Core;
 using Horde.Build.Controllers;
+using Horde.Build.Server;
 using Horde.Build.Utilities;
 using Microsoft.Extensions.Logging;
 
diff --git a/Engine/Source/Programs/Horde/Horde.Build/Commands/Perforce/ReplicateCommand.cs b/Engine/Source/Programs/Horde/Horde.Build/Commands/Perforce/ReplicateCommand.cs
--- a/Engine/Source/Programs/Horde/Horde.Build/Commands/Perforce/ReplicateCommand.cs
+++ b/Engine/Source/Programs/Horde/Horde.Build/Commands/Perforce/ReplicateCommand.cs
@@ -7,10 +7,8 @@ using System.Threading.Tasks;
 using EpicGames.Core;
 using EpicGames.Horde.Storage;
 using EpicGames.Horde.Storage.Impl;
-using Horde.Build.Collections;
-using Horde.Build.Commits;
-using Horde.Build.Commits.Impl;
-using Horde.Build.Models;
+using Horde.Build.Perforce;
+using Horde.Build.Streams;
 using Horde.Build.Utilities;
 using Microsoft.Extensions.Configuration;
 using Microsoft.Extensions.DependencyInjection;
diff --git a/Engine/Source/Programs/Horde/Horde.Build/Compute/ComputeController.cs b/Engine/Source/Programs/Horde/Horde.Build/Compute/ComputeController.cs
--- a/Engine/Source/Programs/Horde/Horde.Build/Compute/ComputeController.cs
+++ b/Engine/Source/Programs/Horde/Horde.Build/Compute/ComputeController.cs
@@ -10,7 +10,7 @@ using Horde.Build.Utilities;
 using Microsoft.AspNetCore.Authorization;
 using Microsoft.AspNetCore.Mvc;
 
-namespace Horde.Build.Compute.Impl
+namespace Horde.Build.Compute
 {
 	/// <summary>
 	/// Controller for the /api/v1/compute endpoint
diff --git a/Engine/Source/Programs/Horde/Horde.Build/Compute/ComputeService.cs b/Engine/Source/Programs/Horde/Horde.Build/Compute/ComputeService.cs
--- a/Engine/Source/Programs/Horde/Horde.Build/Compute/ComputeService.cs
+++ b/Engine/Source/Programs/Horde/Horde.Build/Compute/ComputeService.cs
@@ -12,7 +12,8 @@ using EpicGames.Redis;
 using EpicGames.Serialization;
 using Google.Protobuf;
 using Google.Protobuf.WellKnownTypes;
-using Horde.Build.Models;
+using Horde.Build.Agents;
+using Horde.Build.Agents.Leases;
 using Horde.Build.Server;
 using Horde.Build.Tasks;
 using Horde.Build.Utilities;
@@ -23,7 +24,7 @@ using Microsoft.Extensions.Hosting;
 using Microsoft.Extensions.Logging;
 using StackExchange.Redis;
 
-namespace Horde.Build.Compute.Impl
+namespace Horde.Build.Compute
 {
 	using LeaseId = ObjectId<ILease>;
 
diff --git a/Engine/Source/Programs/Horde/Horde.Build/Compute/RedisMessageQueue.cs b/Engine/Source/Programs/Horde/Horde.Build/Compute/RedisMessageQueue.cs
--- a/Engine/Source/Programs/Horde/Horde.Build/Compute/RedisMessageQueue.cs
+++ b/Engine/Source/Programs/Horde/Horde.Build/Compute/RedisMessageQueue.cs
@@ -7,7 +7,7 @@ using System.Threading.Tasks;
 using EpicGames.Redis;
 using StackExchange.Redis;
 
-namespace Horde.Build.Compute.Impl
+namespace Horde.Build.Compute
 {
 	/// <summary>
 	/// Interface for a distributed message queue
diff --git a/Engine/Source/Programs/Horde/Horde.Build/Configuration/ConfigCollection.cs b/Engine/Source/Programs/Horde/Horde.Build/Configuration/ConfigCollection.cs
--- a/Engine/Source/Programs/Horde/Horde.Build/Configuration/ConfigCollection.cs
+++ b/Engine/Source/Programs/Horde/Horde.Build/Configuration/ConfigCollection.cs
@@ -10,7 +10,7 @@ using System.Linq;
 using System.Text.Json;
 using System.Threading.Tasks;
 
-namespace Horde.Build.Config
+namespace Horde.Build.Configuration
 {
 	/// <summary>
 	/// Collection of config documents
diff --git a/Engine/Source/Programs/Horde/Horde.Build/Configuration/ConfigUpdateService.cs b/Engine/Source/Programs/Horde/Horde.Build/Configuration/ConfigUpdateService.cs
--- a/Engine/Source/Programs/Horde/Horde.Build/Configuration/ConfigUpdateService.cs
+++ b/Engine/Source/Programs/Horde/Horde.Build/Configuration/ConfigUpdateService.cs
@@ -12,13 +12,15 @@ using System.Threading;
 using System.Threading.Tasks;
 using EpicGames.Core;
 using Horde.Build.Acls;
-using Horde.Build.Api;
-using Horde.Build.Collections;
-using Horde.Build.Models;
+using Horde.Build.Agents;
+using Horde.Build.Agents.Pools;
 using Horde.Build.Notifications;
+using Horde.Build.Perforce;
+using Horde.Build.Projects;
 using Horde.Build.Server;
-using Horde.Build.Services;
+using Horde.Build.Streams;
 using Horde.Build.Tools;
+using Horde.Build.Users;
 using Horde.Build.Utilities;
 using HordeCommon;
 using Microsoft.AspNetCore.StaticFiles;
@@ -27,7 +29,7 @@ using Microsoft.Extensions.Hosting;
 using Microsoft.Extensions.Logging;
 using Microsoft.Extensions.Options;
 
-namespace Horde.Build.Config
+namespace Horde.Build.Configuration
 {
 	using PoolId = StringId<IPool>;
 	using ProjectId = StringId<IProject>;
diff --git a/Engine/Source/Programs/Horde/Horde.Build/Devices/DeviceCollection.cs b/Engine/Source/Programs/Horde/Horde.Build/Devices/DeviceCollection.cs
--- a/Engine/Source/Programs/Horde/Horde.Build/Devices/DeviceCollection.cs
+++ b/Engine/Source/Programs/Horde/Horde.Build/Devices/DeviceCollection.cs
@@ -5,14 +5,15 @@ using System.Collections.Generic;
 using System.Linq;
 using System.Threading.Tasks;
 using Horde.Build.Acls;
-using Horde.Build.Models;
+using Horde.Build.Projects;
 using Horde.Build.Server;
+using Horde.Build.Users;
 using Horde.Build.Utilities;
 using MongoDB.Bson;
 using MongoDB.Bson.Serialization.Attributes;
 using MongoDB.Driver;
 
-namespace Horde.Build.Collections.Impl
+namespace Horde.Build.Devices
 {
 	using DeviceId = StringId<IDevice>;
 	using DevicePlatformId = StringId<IDevicePlatform>;
diff --git a/Engine/Source/Programs/Horde/Horde.Build/Devices/DeviceService.cs b/Engine/Source/Programs/Horde/Horde.Build/Devices/DeviceService.cs
--- a/Engine/Source/Programs/Horde/Horde.Build/Devices/DeviceService.cs
+++ b/Engine/Source/Programs/Horde/Horde.Build/Devices/DeviceService.cs
@@ -7,17 +7,20 @@ using System.Security.Claims;
 using System.Threading;
 using System.Threading.Tasks;
 using Horde.Build.Acls;
-using Horde.Build.Api;
-using Horde.Build.Collections;
-using Horde.Build.Models;
+using Horde.Build.Jobs;
+using Horde.Build.Jobs.Graphs;
 using Horde.Build.Notifications;
+using Horde.Build.Projects;
+using Horde.Build.Server;
+using Horde.Build.Streams;
+using Horde.Build.Users;
 using Horde.Build.Utilities;
 using HordeCommon;
 using Microsoft.Extensions.Hosting;
 using Microsoft.Extensions.Logging;
 using MongoDB.Bson;
 
-namespace Horde.Build.Services
+namespace Horde.Build.Devices
 {
 	using DeviceId = StringId<IDevice>;
 	using DevicePlatformId = StringId<IDevicePlatform>;
diff --git a/Engine/Source/Programs/Horde/Horde.Build/Devices/DevicesController.Messages.cs b/Engine/Source/Programs/Horde/Horde.Build/Devices/DevicesController.Messages.cs
--- a/Engine/Source/Programs/Horde/Horde.Build/Devices/DevicesController.Messages.cs
+++ b/Engine/Source/Programs/Horde/Horde.Build/Devices/DevicesController.Messages.cs
@@ -4,9 +4,8 @@ using System;
 using System.Collections.Generic;
 using System.ComponentModel.DataAnnotations;
 using System.Linq;
-using Horde.Build.Models;
 
-namespace Horde.Build.Api
+namespace Horde.Build.Devices
 {
 
 	/// <summary>
diff --git a/Engine/Source/Programs/Horde/Horde.Build/Devices/DevicesController.cs b/Engine/Source/Programs/Horde/Horde.Build/Devices/DevicesController.cs
--- a/Engine/Source/Programs/Horde/Horde.Build/Devices/DevicesController.cs
+++ b/Engine/Source/Programs/Horde/Horde.Build/Devices/DevicesController.cs
@@ -6,10 +6,8 @@ using System.Linq;
 using System.Text.Json;
 using System.Threading.Tasks;
 using Horde.Build.Acls;
-using Horde.Build.Api;
-using Horde.Build.Collections;
-using Horde.Build.Models;
-using Horde.Build.Services;
+using Horde.Build.Projects;
+using Horde.Build.Users;
 using Horde.Build.Utilities;
 using Microsoft.AspNetCore.Authorization;
 using Microsoft.AspNetCore.Mvc;
@@ -17,7 +15,7 @@ using Microsoft.Extensions.Logging;
 using MongoDB.Bson;
 using MongoDB.Driver;
 
-namespace Horde.Build.Controllers
+namespace Horde.Build.Devices
 {
 	using DeviceId = StringId<IDevice>;
 	using DevicePlatformId = StringId<IDevicePlatform>;
diff --git a/Engine/Source/Programs/Horde/Horde.Build/Devices/IDevice.cs b/Engine/Source/Programs/Horde/Horde.Build/Devices/IDevice.cs
--- a/Engine/Source/Programs/Horde/Horde.Build/Devices/IDevice.cs
+++ b/Engine/Source/Programs/Horde/Horde.Build/Devices/IDevice.cs
@@ -3,11 +3,13 @@
 using System;
 using System.Collections.Generic;
 using Horde.Build.Acls;
+using Horde.Build.Projects;
+using Horde.Build.Users;
 using Horde.Build.Utilities;
 using MongoDB.Bson;
 using MongoDB.Bson.Serialization.Attributes;
 
-namespace Horde.Build.Models
+namespace Horde.Build.Devices
 {
 	using DeviceId = StringId<IDevice>;
 	using DevicePlatformId = StringId<IDevicePlatform>;
diff --git a/Engine/Source/Programs/Horde/Horde.Build/Devices/IDeviceCollection.cs b/Engine/Source/Programs/Horde/Horde.Build/Devices/IDeviceCollection.cs
--- a/Engine/Source/Programs/Horde/Horde.Build/Devices/IDeviceCollection.cs
+++ b/Engine/Source/Programs/Horde/Horde.Build/Devices/IDeviceCollection.cs
@@ -2,11 +2,12 @@
 
 using System.Collections.Generic;
 using System.Threading.Tasks;
-using Horde.Build.Models;
+using Horde.Build.Projects;
+using Horde.Build.Users;
 using Horde.Build.Utilities;
 using MongoDB.Bson;
 
-namespace Horde.Build.Collections
+namespace Horde.Build.Devices
 {
 	using DeviceId = StringId<IDevice>;
 	using DevicePlatformId = StringId<IDevicePlatform>;
diff --git a/Engine/Source/Programs/Horde/Horde.Build/Issues/External/IExternalIssue.cs b/Engine/Source/Programs/Horde/Horde.Build/Issues/External/IExternalIssue.cs
--- a/Engine/Source/Programs/Horde/Horde.Build/Issues/External/IExternalIssue.cs
+++ b/Engine/Source/Programs/Horde/Horde.Build/Issues/External/IExternalIssue.cs
@@ -13,7 +13,7 @@ using Microsoft.Extensions.Options;
 using Polly;
 using Polly.Extensions.Http;
 
-namespace Horde.Build.Models
+namespace Horde.Build.Issues.External
 {
 	/// <summary>
 	/// 
diff --git a/Engine/Source/Programs/Horde/Horde.Build/Issues/External/IExternalIssueService.cs b/Engine/Source/Programs/Horde/Horde.Build/Issues/External/IExternalIssueService.cs
--- a/Engine/Source/Programs/Horde/Horde.Build/Issues/External/IExternalIssueService.cs
+++ b/Engine/Source/Programs/Horde/Horde.Build/Issues/External/IExternalIssueService.cs
@@ -4,11 +4,12 @@ using System;
 using System.Collections.Generic;
 using System.Threading;
 using System.Threading.Tasks;
-using Horde.Build.Models;
+using Horde.Build.Streams;
+using Horde.Build.Users;
 using Horde.Build.Utilities;
 using Microsoft.Extensions.Hosting;
 
-namespace Horde.Build.Services
+namespace Horde.Build.Issues.External
 {
 	using StreamId = StringId<IStream>;
 
diff --git a/Engine/Source/Programs/Horde/Horde.Build/Issues/External/JiraService.cs b/Engine/Source/Programs/Horde/Horde.Build/Issues/External/JiraService.cs
--- a/Engine/Source/Programs/Horde/Horde.Build/Issues/External/JiraService.cs
+++ b/Engine/Source/Programs/Horde/Horde.Build/Issues/External/JiraService.cs
@@ -10,19 +10,20 @@ using System.Text.Json;
 using System.Threading.Tasks;
 using Microsoft.Extensions.Logging;
 using Microsoft.Extensions.Options;
-using Horde.Build.Models;
 using Polly;
 using Polly.Extensions.Http;
 using Horde.Build.Utilities;
-using Horde.Build.Config;
+using Horde.Build.Configuration;
 using Microsoft.Extensions.Hosting;
 using HordeCommon;
 using System.Threading;
 using System.Linq;
 using Horde.Build.Server;
+using Horde.Build.Streams;
+using Horde.Build.Users;
 
-namespace Horde.Build.Services
-{	
+namespace Horde.Build.Issues.External
+{
 	/// <summary>
 	/// Jira issue information
 	/// </summary>
diff --git a/Engine/Source/Programs/Horde/Horde.Build/Issues/Handlers/CompileIssueHandler.cs b/Engine/Source/Programs/Horde/Horde.Build/Issues/Handlers/CompileIssueHandler.cs
--- a/Engine/Source/Programs/Horde/Horde.Build/Issues/Handlers/CompileIssueHandler.cs
+++ b/Engine/Source/Programs/Horde/Horde.Build/Issues/Handlers/CompileIssueHandler.cs
@@ -5,12 +5,13 @@ using System.Collections.Generic;
 using System.Diagnostics.CodeAnalysis;
 using System.Linq;
 using EpicGames.Core;
-using Horde.Build.Collections;
-using Horde.Build.Models;
+using Horde.Build.Jobs;
+using Horde.Build.Jobs.Graphs;
+using Horde.Build.Logs;
 using Microsoft.Extensions.Logging;
 using MongoDB.Driver;
 
-namespace Horde.Build.IssueHandlers.Impl
+namespace Horde.Build.Issues.Handlers
 {
 	/// <summary>
 	/// Instance of a particular compile error
diff --git a/Engine/Source/Programs/Horde/Horde.Build/Issues/Handlers/ContentIssueHandler.cs b/Engine/Source/Programs/Horde/Horde.Build/Issues/Handlers/ContentIssueHandler.cs
--- a/Engine/Source/Programs/Horde/Horde.Build/Issues/Handlers/ContentIssueHandler.cs
+++ b/Engine/Source/Programs/Horde/Horde.Build/Issues/Handlers/ContentIssueHandler.cs
@@ -5,13 +5,14 @@ using System.Collections.Generic;
 using System.Diagnostics.CodeAnalysis;
 using System.Linq;
 using EpicGames.Core;
-using Horde.Build.Collections;
-using Horde.Build.Models;
+using Horde.Build.Jobs;
+using Horde.Build.Jobs.Graphs;
+using Horde.Build.Logs;
 using Horde.Build.Utilities;
 using Microsoft.Extensions.Logging;
 using MongoDB.Driver;
 
-namespace Horde.Build.IssueHandlers.Impl
+namespace Horde.Build.Issues.Handlers
 {
 	/// <summary>
 	/// Instance of a particular compile error
diff --git a/Engine/Source/Programs/Horde/Horde.Build/Issues/IIssue.cs b/Engine/Source/Programs/Horde/Horde.Build/Issues/IIssue.cs
--- a/Engine/Source/Programs/Horde/Horde.Build/Issues/IIssue.cs
+++ b/Engine/Source/Programs/Horde/Horde.Build/Issues/IIssue.cs
@@ -3,10 +3,12 @@
 using System;
 using System.Collections.Generic;
 using System.Linq;
+using Horde.Build.Streams;
+using Horde.Build.Users;
 using Horde.Build.Utilities;
 using MongoDB.Bson;
 
-namespace Horde.Build.Models
+namespace Horde.Build.Issues
 {
 	using StreamId = StringId<IStream>;
 	using UserId = ObjectId<IUser>;
diff --git a/Engine/Source/Programs/Horde/Horde.Build/Issues/IIssueCollection.cs b/Engine/Source/Programs/Horde/Horde.Build/Issues/IIssueCollection.cs
--- a/Engine/Source/Programs/Horde/Horde.Build/Issues/IIssueCollection.cs
+++ b/Engine/Source/Programs/Horde/Horde.Build/Issues/IIssueCollection.cs
@@ -4,13 +4,17 @@ using System;
 using System.Collections.Generic;
 using System.Linq;
 using System.Threading.Tasks;
-using Horde.Build.Models;
-using Horde.Build.Server;
+using Horde.Build.Auditing;
+using Horde.Build.Jobs;
+using Horde.Build.Jobs.Graphs;
+using Horde.Build.Logs;
+using Horde.Build.Streams;
+using Horde.Build.Users;
 using Horde.Build.Utilities;
 using MongoDB.Bson;
 using MongoDB.Driver;
 
-namespace Horde.Build.Collections
+namespace Horde.Build.Issues
 {
 	using JobId = ObjectId<IJob>;
 	using LogId = ObjectId<ILogFile>;
diff --git a/Engine/Source/Programs/Horde/Horde.Build/Issues/IIssueDetails.cs b/Engine/Source/Programs/Horde/Horde.Build/Issues/IIssueDetails.cs
--- a/Engine/Source/Programs/Horde/Horde.Build/Issues/IIssueDetails.cs
+++ b/Engine/Source/Programs/Horde/Horde.Build/Issues/IIssueDetails.cs
@@ -2,10 +2,10 @@
 using System;
 using System.Collections.Generic;
 using System.Threading.Tasks;
-using Horde.Build.Services;
+using Horde.Build.Users;
 using Horde.Build.Utilities;
 
-namespace Horde.Build.Models
+namespace Horde.Build.Issues
 {
 	using UserId = ObjectId<IUser>;
 
diff --git a/Engine/Source/Programs/Horde/Horde.Build/Issues/IIssueFingerprint.cs b/Engine/Source/Programs/Horde/Horde.Build/Issues/IIssueFingerprint.cs
--- a/Engine/Source/Programs/Horde/Horde.Build/Issues/IIssueFingerprint.cs
+++ b/Engine/Source/Programs/Horde/Horde.Build/Issues/IIssueFingerprint.cs
@@ -5,7 +5,7 @@ using System.Collections.Generic;
 using System.Linq;
 using Horde.Build.Utilities;
 
-namespace Horde.Build.Models
+namespace Horde.Build.Issues
 {
 	/// <summary>
 	/// Fingerprint for an issue
diff --git a/Engine/Source/Programs/Horde/Horde.Build/Issues/IIssueHandler.cs b/Engine/Source/Programs/Horde/Horde.Build/Issues/IIssueHandler.cs
--- a/Engine/Source/Programs/Horde/Horde.Build/Issues/IIssueHandler.cs
+++ b/Engine/Source/Programs/Horde/Horde.Build/Issues/IIssueHandler.cs
@@ -4,10 +4,12 @@ using System;
 using System.Collections.Generic;
 using System.Diagnostics.CodeAnalysis;
 using EpicGames.Core;
-using Horde.Build.Collections;
-using Horde.Build.Models;
+using Horde.Build.Jobs;
+using Horde.Build.Jobs.Graphs;
+using Horde.Build.Logs;
+using Horde.Build.Perforce;
 
-namespace Horde.Build.IssueHandlers
+namespace Horde.Build.Issues
 {
 	/// <summary>
 	/// Information about a changelist and a value ranking the likelihood that it caused an issue
diff --git a/Engine/Source/Programs/Horde/Horde.Build/Issues/IIssueSpan.cs b/Engine/Source/Programs/Horde/Horde.Build/Issues/IIssueSpan.cs
--- a/Engine/Source/Programs/Horde/Horde.Build/Issues/IIssueSpan.cs
+++ b/Engine/Source/Programs/Horde/Horde.Build/Issues/IIssueSpan.cs
@@ -1,10 +1,12 @@
 // Copyright Epic Games, Inc. All Rights Reserved.
 
 using System.Collections.Generic;
+using Horde.Build.Streams;
+using Horde.Build.Users;
 using Horde.Build.Utilities;
 using MongoDB.Bson;
 
-namespace Horde.Build.Models
+namespace Horde.Build.Issues
 {
 	using StreamId = StringId<IStream>;
 	using TemplateRefId = StringId<TemplateRef>;
diff --git a/Engine/Source/Programs/Horde/Horde.Build/Issues/IIssueStep.cs b/Engine/Source/Programs/Horde/Horde.Build/Issues/IIssueStep.cs
--- a/Engine/Source/Programs/Horde/Horde.Build/Issues/IIssueStep.cs
+++ b/Engine/Source/Programs/Horde/Horde.Build/Issues/IIssueStep.cs
@@ -1,11 +1,14 @@
 // Copyright Epic Games, Inc. All Rights Reserved.
 
 using System;
+using Horde.Build.Jobs;
+using Horde.Build.Jobs.Graphs;
+using Horde.Build.Logs;
 using Horde.Build.Server;
 using Horde.Build.Utilities;
 using MongoDB.Bson;
 
-namespace Horde.Build.Models
+namespace Horde.Build.Issues
 {
 	using JobId = ObjectId<IJob>;
 	using LogId = ObjectId<ILogFile>;
diff --git a/Engine/Source/Programs/Horde/Horde.Build/Issues/IssueCollection.cs b/Engine/Source/Programs/Horde/Horde.Build/Issues/IssueCollection.cs
--- a/Engine/Source/Programs/Horde/Horde.Build/Issues/IssueCollection.cs
+++ b/Engine/Source/Programs/Horde/Horde.Build/Issues/IssueCollection.cs
@@ -6,8 +6,14 @@ using System.Diagnostics;
 using System.Linq;
 using System.Threading.Tasks;
 using EpicGames.Redis.Utility;
-using Horde.Build.Models;
+using Horde.Build.Auditing;
+using Horde.Build.Jobs;
+using Horde.Build.Jobs.Graphs;
+using Horde.Build.Issues;
+using Horde.Build.Logs;
 using Horde.Build.Server;
+using Horde.Build.Streams;
+using Horde.Build.Users;
 using Horde.Build.Utilities;
 using Microsoft.Extensions.Caching.Memory;
 using Microsoft.Extensions.Logging;
@@ -15,7 +21,7 @@ using MongoDB.Bson;
 using MongoDB.Bson.Serialization.Attributes;
 using MongoDB.Driver;
 
-namespace Horde.Build.Collections.Impl
+namespace Horde.Build.Issues
 {
 	using JobId = ObjectId<IJob>;
 	using LogId = ObjectId<ILogFile>;
diff --git a/Engine/Source/Programs/Horde/Horde.Build/Issues/IssueReportService.cs b/Engine/Source/Programs/Horde/Horde.Build/Issues/IssueReportService.cs
--- a/Engine/Source/Programs/Horde/Horde.Build/Issues/IssueReportService.cs
+++ b/Engine/Source/Programs/Horde/Horde.Build/Issues/IssueReportService.cs
@@ -6,17 +6,18 @@ using System.Linq;
 using System.Threading;
 using System.Threading.Tasks;
 using EpicGames.Core;
-using Horde.Build.Collections;
-using Horde.Build.Config;
-using Horde.Build.Models;
+using Horde.Build.Configuration;
+using Horde.Build.Jobs;
+using Horde.Build.Jobs.Graphs;
 using Horde.Build.Notifications;
 using Horde.Build.Server;
+using Horde.Build.Streams;
 using Horde.Build.Utilities;
 using HordeCommon;
 using Microsoft.Extensions.Hosting;
 using Microsoft.Extensions.Logging;
 
-namespace Horde.Build.Issues.Impl
+namespace Horde.Build.Issues
 {
 	using TemplateRefId = StringId<TemplateRef>;
 	using WorkflowId = StringId<WorkflowConfig>;
@@ -185,7 +186,7 @@ namespace Horde.Build.Issues.Impl
 			foreach (IGrouping<TemplateRefId, IJob> templateGroup in jobs.GroupBy(x => x.TemplateId))
 			{
 				WorkflowId? templateWorkflowId = null;
-				if (stream.Config.TryGetTemplate(templateGroup.Key, out Api.TemplateRefConfig? templateRefConfig))
+				if (stream.Config.TryGetTemplate(templateGroup.Key, out TemplateRefConfig? templateRefConfig))
 				{
 					templateWorkflowId = templateRefConfig.Annotations.WorkflowId;
 				}
diff --git a/Engine/Source/Programs/Horde/Horde.Build/Issues/IssueService.cs b/Engine/Source/Programs/Horde/Horde.Build/Issues/IssueService.cs
--- a/Engine/Source/Programs/Horde/Horde.Build/Issues/IssueService.cs
+++ b/Engine/Source/Programs/Horde/Horde.Build/Issues/IssueService.cs
@@ -9,13 +9,15 @@ using System.Text.RegularExpressions;
 using System.Threading;
 using System.Threading.Tasks;
 using EpicGames.Core;
-using Horde.Build.Api;
-using Horde.Build.Collections;
-using Horde.Build.Config;
-using Horde.Build.IssueHandlers;
-using Horde.Build.IssueHandlers.Impl;
-using Horde.Build.Models;
+using Horde.Build.Configuration;
+using Horde.Build.Issues.Handlers;
+using Horde.Build.Jobs;
+using Horde.Build.Jobs.Graphs;
+using Horde.Build.Logs;
+using Horde.Build.Perforce;
 using Horde.Build.Server;
+using Horde.Build.Streams;
+using Horde.Build.Users;
 using Horde.Build.Utilities;
 using HordeCommon;
 using Microsoft.Extensions.Hosting;
@@ -26,7 +28,7 @@ using MongoDB.Driver;
 using OpenTracing;
 using OpenTracing.Util;
 
-namespace Horde.Build.Services
+namespace Horde.Build.Issues
 {
 	using LogId = ObjectId<ILogFile>;
 	using StreamId = StringId<IStream>;
diff --git a/Engine/Source/Programs/Horde/Horde.Build/Issues/IssuesController.Messages.cs b/Engine/Source/Programs/Horde/Horde.Build/Issues/IssuesController.Messages.cs
--- a/Engine/Source/Programs/Horde/Horde.Build/Issues/IssuesController.Messages.cs
+++ b/Engine/Source/Programs/Horde/Horde.Build/Issues/IssuesController.Messages.cs
@@ -4,12 +4,15 @@ using System;
 using System.Collections.Generic;
 using System.Linq;
 using EpicGames.Core;
-using Horde.Build.Models;
+using Horde.Build.Issues;
+using Horde.Build.Issues.External;
 using Horde.Build.Server;
+using Horde.Build.Streams;
+using Horde.Build.Users;
 using Horde.Build.Utilities;
 using MongoDB.Bson;
 
-namespace Horde.Build.Api
+namespace Horde.Build.Issues
 {
 	using StreamId = StringId<IStream>;
 	using TemplateRefId = StringId<TemplateRef>;
diff --git a/Engine/Source/Programs/Horde/Horde.Build/Issues/IssuesController.cs b/Engine/Source/Programs/Horde/Horde.Build/Issues/IssuesController.cs
--- a/Engine/Source/Programs/Horde/Horde.Build/Issues/IssuesController.cs
+++ b/Engine/Source/Programs/Horde/Horde.Build/Issues/IssuesController.cs
@@ -5,19 +5,22 @@ using System.Collections.Generic;
 using System.Linq;
 using System.Threading.Tasks;
 using Horde.Build.Acls;
-using Horde.Build.Api;
-using Horde.Build.Collections;
-using Horde.Build.Models;
+using Horde.Build.Auditing;
 using Horde.Build.Server;
-using Horde.Build.Services;
 using Horde.Build.Utilities;
 using HordeCommon;
 using Microsoft.AspNetCore.Authorization;
 using Microsoft.AspNetCore.Mvc;
 using MongoDB.Bson;
 using Microsoft.Extensions.Logging;
-
-namespace Horde.Build.Controllers
+using Horde.Build.Jobs;
+using Horde.Build.Logs;
+using Horde.Build.Streams;
+using Horde.Build.Issues.External;
+using Horde.Build.Users;
+using Horde.Build.Jobs.Graphs;
+
+namespace Horde.Build.Issues
 {
 	using JobId = ObjectId<IJob>;
 	using LogId = ObjectId<ILogFile>;
