Horde: Move files into namespaces corresponding to their location on disk.

diff --git a/Engine/Source/Programs/Horde/Horde.Build/Jobs/Artifacts/ArtifactCollection.cs b/Engine/Source/Programs/Horde/Horde.Build/Jobs/Artifacts/ArtifactCollection.cs
new file mode 100644
--- /dev/null
+++ b/Engine/Source/Programs/Horde/Horde.Build/Jobs/Artifacts/ArtifactCollection.cs
@@ -0,0 +1,273 @@
+// Copyright Epic Games, Inc. All Rights Reserved.
+
+using System;
+using System.Collections.Generic;
+using System.Linq;
+using System.Threading.Tasks;
+using Horde.Build.Jobs;
+using Horde.Build.Server;
+using Horde.Build.Storage;
+using Horde.Build.Utilities;
+using MongoDB.Bson;
+using MongoDB.Bson.Serialization.Attributes;
+using MongoDB.Driver;
+
+namespace Horde.Build.Jobs.Artifacts
+{
+	using JobId = ObjectId<IJob>;
+
+	/// <summary>
+	/// Wraps functionality for manipulating artifacts
+	/// </summary>
+	public class ArtifactCollection : IArtifactCollection
+	{
+		class Artifact : IArtifact
+		{
+			[BsonRequired, BsonId]
+			public ObjectId Id { get; set; }
+
+			[BsonRequired]
+			public JobId JobId { get; set; }
+
+			[BsonRequired]
+			public string Name { get; set; }
+
+			public SubResourceId? StepId { get; set; }
+
+			[BsonRequired]
+			public long Length { get; set; }
+
+			public string MimeType { get; set; }
+
+			[BsonRequired]
+			public int UpdateIndex { get; set; }
+
+			[BsonConstructor]
+			private Artifact()
+			{
+				Name = null!;
+				MimeType = null!;
+			}
+
+			public Artifact(JobId jobId, SubResourceId? stepId, string name, long length, string mimeType)
+			{
+				Id = ObjectId.GenerateNewId();
+				JobId = jobId;
+				StepId = stepId;
+				Name = name;
+				Length = length;
+				MimeType = mimeType;
+			}
+		}
+
+		private readonly IStorageBackend _storageBackend;
+		private readonly IMongoCollection<Artifact> _artifacts;
+
+		/// <summary>
+		/// Constructor
+		/// </summary>
+		/// <param name="mongoService">The database service</param>
+		/// <param name="storageBackend">The storage backend</param>
+		public ArtifactCollection(MongoService mongoService, IStorageBackend<ArtifactCollection> storageBackend)
+		{
+			_storageBackend = storageBackend;
+
+			// Initialize Artifacts table
+			_artifacts = mongoService.GetCollection<Artifact>("Artifacts", keys => keys.Ascending(x => x.JobId));
+		}
+
+		/// <summary>
+		/// Creates a new artifact
+		/// </summary>
+		/// <param name="jobId">Unique id of the job that owns this artifact</param>
+		/// <param name="stepId">Optional Step id</param>
+		/// <param name="name">Name of artifact</param>
+		/// <param name="mimeType">Type of artifact</param>
+		/// <param name="data">The data to write</param>
+		/// <returns>The new log file document</returns>
+		public async Task<IArtifact> CreateArtifactAsync(JobId jobId, SubResourceId? stepId, string name, string mimeType, System.IO.Stream data)
+		{
+			// upload first
+			string artifactName = ValidateName(name);
+			await _storageBackend.WriteAsync(GetPath(jobId, stepId, artifactName), data);
+
+			// then create entry
+			Artifact newArtifact = new Artifact(jobId, stepId, artifactName, data.Length, mimeType);
+			await _artifacts.InsertOneAsync(newArtifact);
+			return newArtifact;
+		}
+
+		/// <summary>
+		/// Gets all the available artifacts for a job
+		/// </summary>
+		/// <param name="jobId">Unique id of the job to query</param>
+		/// <param name="stepId">Unique id of the Step to query</param>
+		/// <param name="name">Name of the artifact</param>
+		/// <returns>List of artifact documents</returns>
+		public async Task<List<IArtifact>> GetArtifactsAsync(JobId? jobId, SubResourceId? stepId, string? name)
+		{
+			FilterDefinitionBuilder<Artifact> builder = Builders<Artifact>.Filter;
+
+			FilterDefinition<Artifact> filter = FilterDefinition<Artifact>.Empty;
+			if (jobId != null)
+			{
+				filter &= builder.Eq(x => x.JobId, jobId.Value);
+			}
+			if (stepId != null)
+			{
+				filter &= builder.Eq(x => x.StepId, stepId.Value);
+			}
+			if (name != null)
+			{
+				filter &= builder.Eq(x => x.Name, name);
+			}
+
+			return await _artifacts.Find(filter).ToListAsync<Artifact, IArtifact>();
+		}
+
+		/// <summary>
+		/// Gets a specific list of artifacts based on id
+		/// </summary>
+		/// <param name="artifactIds">The list of artifact Ids</param>
+		/// <returns>List of artifact documents</returns>
+		public async Task<List<IArtifact>> GetArtifactsAsync(IEnumerable<ObjectId> artifactIds)
+		{
+			FilterDefinitionBuilder<Artifact> builder = Builders<Artifact>.Filter;
+
+			FilterDefinition<Artifact> filter = FilterDefinition<Artifact>.Empty;
+			filter &= builder.In(x => x.Id, artifactIds);
+			
+			return await _artifacts.Find(filter).ToListAsync<Artifact, IArtifact>();
+		}
+
+		/// <summary>
+		/// Gets an artifact by ID
+		/// </summary>
+		/// <param name="artifactId">Unique id of the artifact</param>
+		/// <returns>The artifact document</returns>
+		public async Task<IArtifact?> GetArtifactAsync(ObjectId artifactId)
+		{
+			return await _artifacts.Find<Artifact>(x => x.Id == artifactId).FirstOrDefaultAsync();
+		}
+
+		/// <summary>
+		/// Attempts to update a single artifact document, if the update counter is valid
+		/// </summary>
+		/// <param name="current">The current artifact document</param>
+		/// <param name="update">The update definition</param>
+		/// <returns>True if the document was updated, false if another writer updated the document first</returns>
+		private async Task<bool> TryUpdateArtifactAsync(Artifact current, UpdateDefinition<Artifact> update)
+		{
+			UpdateResult result = await _artifacts.UpdateOneAsync<Artifact>(x => x.Id == current.Id && x.UpdateIndex == current.UpdateIndex, update.Set(x => x.UpdateIndex, current.UpdateIndex + 1));
+			return result.ModifiedCount == 1;
+		}
+
+		/// <summary>
+		/// Updates an artifact
+		/// </summary>
+		/// <param name="artifact">The artifact</param>
+		/// <param name="newMimeType">New mime type</param>
+		/// <param name="newData">New data</param>
+		/// <returns>Async task</returns>
+		public async Task<bool> UpdateArtifactAsync(IArtifact? artifact, string newMimeType, System.IO.Stream newData)
+		{
+			while (artifact != null)
+			{
+				UpdateDefinitionBuilder<Artifact> updateBuilder = Builders<Artifact>.Update;
+
+				List<UpdateDefinition<Artifact>> updates = new List<UpdateDefinition<Artifact>>();
+				updates.Add(updateBuilder.Set(x => x.MimeType, newMimeType));
+				updates.Add(updateBuilder.Set(x => x.Length, newData.Length));
+
+				// re-upload the data to external
+				string artifactName = ValidateName(artifact.Name);
+				await _storageBackend.WriteAsync(GetPath(artifact.JobId, artifact.StepId, artifactName), newData);
+
+				if (await TryUpdateArtifactAsync((Artifact)artifact, updateBuilder.Combine(updates)))
+				{
+					return true;
+				}
+
+				artifact = await GetArtifactAsync(artifact.Id);
+			}
+			return false;
+		}
+
+		/// <summary>
+		/// gets artifact data
+		/// </summary>
+		/// <param name="artifact">The artifact</param>
+		/// <returns>The chunk data</returns>
+		public async Task<System.IO.Stream> OpenArtifactReadStreamAsync(IArtifact artifact)
+		{
+			System.IO.Stream? stream = await _storageBackend.ReadAsync(GetPath(artifact.JobId, artifact.StepId, artifact.Name));
+			if (stream == null)
+			{
+				throw new Exception($"Unable to get artifact {artifact.Id}");
+			}
+			return stream;
+		}
+
+		/// <summary>
+		/// Get the path for an artifact
+		/// </summary>
+		/// <param name="jobId"></param>
+		/// <param name="stepId"></param>
+		/// <param name="name"></param>
+		/// <returns></returns>
+		private static string GetPath(JobId jobId, SubResourceId? stepId, string name)
+		{
+			if (stepId == null)
+			{
+				return $"{jobId}/{name}";
+			}
+			else
+			{
+				return $"{jobId}/{stepId.Value}/{name}";
+			}
+		}
+
+		/// <summary>
+		/// Checks that the name given is valid for an artifact
+		/// </summary>
+		/// <param name="name"></param>
+		private static string ValidateName(string name)
+		{
+			string newName = name.Replace('\\', '/');
+			if (newName.Length == 0)
+			{
+				throw new ArgumentException("Artifact name is empty");
+			}
+			else if (newName.StartsWith('/'))
+			{
+				throw new ArgumentException($"Artifact has an absolute path ({newName})");
+			}
+			else if (newName.EndsWith('/'))
+			{
+				throw new ArgumentException($"Artifact does not have a file name ({newName})");
+			}
+			else if (newName.Contains("//", StringComparison.Ordinal))
+			{
+				throw new ArgumentException($"Artifact name contains an invalid directory ({newName})");
+			}
+
+			string invalidChars = ":<>|\"";
+			for (int idx = 0; idx < newName.Length; idx++)
+			{
+				int minIdx = idx;
+				for (; idx < newName.Length && newName[idx] != '/'; idx++)
+				{
+					if (invalidChars.Contains(newName[idx], StringComparison.Ordinal))
+					{
+						throw new ArgumentException($"Invalid character in artifact name ({newName})");
+					}
+				}
+				if ((idx == minIdx + 1 && newName[minIdx] == '.') || (idx == minIdx + 2 && newName[minIdx] == '.' && newName[minIdx + 1] == '.'))
+				{
+					throw new ArgumentException($"Artifact may not contain symbolic directory names ({newName})");
+				}
+			}
+			return newName;
+		}
+	}
+}
\ No newline at end of file
diff --git a/Engine/Source/Programs/Horde/Horde.Build/Jobs/Artifacts/ArtifactsController.Messages.cs b/Engine/Source/Programs/Horde/Horde.Build/Jobs/Artifacts/ArtifactsController.Messages.cs
new file mode 100644
--- /dev/null
+++ b/Engine/Source/Programs/Horde/Horde.Build/Jobs/Artifacts/ArtifactsController.Messages.cs
@@ -0,0 +1,118 @@
+// Copyright Epic Games, Inc. All Rights Reserved.
+
+using System.Collections.Generic;
+using Horde.Build.Jobs.Artifacts;
+
+namespace Horde.Build.Jobs.Artifacts
+{
+	/// <summary>
+	/// Response from creating an artifact
+	/// </summary>
+	public class CreateArtifactResponse
+	{
+		/// <summary>
+		/// Unique id for this artifact
+		/// </summary>
+		public string Id { get; set; }
+
+		/// <summary>
+		/// Constructor
+		/// </summary>
+		/// <param name="id">The artifact file id</param>
+		public CreateArtifactResponse(string id)
+		{
+			Id = id;
+		}
+	}
+
+	/// <summary>
+	/// Request to get a zip of artifacts
+	/// </summary>
+	public class GetArtifactZipRequest
+	{
+		/// <summary>
+		/// Job id to get all artifacts for
+		/// </summary>
+		public string? JobId { get; set; }
+
+		/// <summary>
+		/// Step id to filter by
+		/// </summary>
+		public string? StepId { get; set; }
+
+		/// <summary>
+		/// Further filter by a list of artifact ids
+		/// </summary>
+		public List<string>? ArtifactIds { get; set; }
+
+		/// <summary>
+		/// Constructor
+		/// </summary>
+		/// <param name="jobId">Job id to get all artifacts for</param>
+		/// <param name="stepId">step to filter by</param>
+		/// <param name="artifactIds">The artifact ids.  Returns all artifacts for a job </param>
+		public GetArtifactZipRequest(string? jobId, string? stepId, List<string>? artifactIds)
+		{
+			JobId = jobId;
+			StepId = stepId;
+			ArtifactIds = artifactIds;
+		}
+	}
+
+	/// <summary>
+	/// Response describing an artifact
+	/// </summary>
+	public class GetArtifactResponse
+	{
+		/// <summary>
+		/// Unique id of the artifact
+		/// </summary>
+		public string Id { get; set; }
+
+		/// <summary>
+		/// Name of the artifact
+		/// </summary>
+		public string Name { get; set; }
+
+		/// <summary>
+		/// Unique id of the job for this artifact
+		/// </summary>
+		public string JobId { get; set; }
+
+		/// <summary>
+		/// Optional unique id of the step for this artifact
+		/// </summary>
+		public string? StepId { get; set; }
+
+		/// <summary>
+		/// Type of artifact
+		/// </summary>
+		public string MimeType { get; set; }
+
+		/// <summary>
+		/// Length of artifact
+		/// </summary>
+		public long Length { get; set; }
+
+		/// <summary>
+		/// Short-lived code that can be used to download this artifact through direct download links in the browser
+		/// </summary>
+		public string? Code { get; set; }
+
+		/// <summary>
+		/// Constructor
+		/// </summary>
+		/// <param name="artifact">The artifact to construct from</param>
+		/// <param name="code">The direct download code</param>
+		public GetArtifactResponse(IArtifact artifact, string? code)
+		{
+			Id = artifact.Id.ToString();
+			Name = artifact.Name;
+			JobId = artifact.JobId.ToString();
+			StepId = artifact.StepId.ToString();
+			MimeType = artifact.MimeType;
+			Length = artifact.Length;
+			Code = code;
+		}
+	}
+}
diff --git a/Engine/Source/Programs/Horde/Horde.Build/Jobs/Artifacts/ArtifactsController.cs b/Engine/Source/Programs/Horde/Horde.Build/Jobs/Artifacts/ArtifactsController.cs
--- a/Engine/Source/Programs/Horde/Horde.Build/Jobs/Artifacts/ArtifactsController.cs
+++ b/Engine/Source/Programs/Horde/Horde.Build/Jobs/Artifacts/ArtifactsController.cs
@@ -10,11 +10,9 @@ using System.Net.Mime;
 using System.Security.Claims;
 using System.Threading.Tasks;
 using Horde.Build.Acls;
-using Horde.Build.Api;
 using Horde.Build.Jobs;
-using Horde.Build.Models;
+using Horde.Build.Jobs.Graphs;
 using Horde.Build.Server;
-using Horde.Build.Services;
 using Horde.Build.Utilities;
 using Microsoft.AspNetCore.Authorization;
 using Microsoft.AspNetCore.Http;
@@ -22,7 +20,7 @@ using Microsoft.AspNetCore.Mvc;
 using Microsoft.IdentityModel.Tokens;
 using MongoDB.Bson;
 
-namespace Horde.Build.Controllers
+namespace Horde.Build.Jobs.Artifacts
 {
 	using JobId = ObjectId<IJob>;
 
diff --git a/Engine/Source/Programs/Horde/Horde.Build/Jobs/Artifacts/IArtifact.cs b/Engine/Source/Programs/Horde/Horde.Build/Jobs/Artifacts/IArtifact.cs
--- a/Engine/Source/Programs/Horde/Horde.Build/Jobs/Artifacts/IArtifact.cs
+++ b/Engine/Source/Programs/Horde/Horde.Build/Jobs/Artifacts/IArtifact.cs
@@ -3,7 +3,7 @@
 using Horde.Build.Utilities;
 using MongoDB.Bson;
 
-namespace Horde.Build.Models
+namespace Horde.Build.Jobs.Artifacts
 {
 	using JobId = ObjectId<IJob>;
 
diff --git a/Engine/Source/Programs/Horde/Horde.Build/Jobs/Artifacts/IArtifactCollection.cs b/Engine/Source/Programs/Horde/Horde.Build/Jobs/Artifacts/IArtifactCollection.cs
--- a/Engine/Source/Programs/Horde/Horde.Build/Jobs/Artifacts/IArtifactCollection.cs
+++ b/Engine/Source/Programs/Horde/Horde.Build/Jobs/Artifacts/IArtifactCollection.cs
@@ -2,11 +2,10 @@
 
 using System.Collections.Generic;
 using System.Threading.Tasks;
-using Horde.Build.Models;
 using Horde.Build.Utilities;
 using MongoDB.Bson;
 
-namespace Horde.Build.Jobs
+namespace Horde.Build.Jobs.Artifacts
 {
 	using JobId = ObjectId<IJob>;
 
diff --git a/Engine/Source/Programs/Horde/Horde.Build/Jobs/Graphs/GraphCollection.cs b/Engine/Source/Programs/Horde/Horde.Build/Jobs/Graphs/GraphCollection.cs
new file mode 100644
--- /dev/null
+++ b/Engine/Source/Programs/Horde/Horde.Build/Jobs/Graphs/GraphCollection.cs
@@ -0,0 +1,476 @@
+// Copyright Epic Games, Inc. All Rights Reserved.
+
+using System;
+using System.Collections.Concurrent;
+using System.Collections.Generic;
+using System.Diagnostics;
+using System.Linq;
+using System.Threading;
+using System.Threading.Tasks;
+using EpicGames.Core;
+using Horde.Build.Jobs.Templates;
+using Horde.Build.Server;
+using Horde.Build.Utilities;
+using HordeCommon;
+using MongoDB.Bson;
+using MongoDB.Bson.Serialization.Attributes;
+using MongoDB.Bson.Serialization.Options;
+using MongoDB.Driver;
+
+namespace Horde.Build.Jobs.Graphs
+{
+	/// <summary>
+	/// Collection of graph documents
+	/// </summary>
+	public class GraphCollection : IGraphCollection
+	{
+		/// <summary>
+		/// Represents a node in the graph
+		/// </summary>
+		[DebuggerDisplay("{Name}")]
+		class Node : INode
+		{
+			[BsonRequired]
+			public string Name { get; set; }
+
+			[BsonIgnoreIfNull]
+			public NodeRef[] InputDependencies { get; set; }
+
+			[BsonIgnoreIfNull]
+			public NodeRef[] OrderDependencies { get; set; }
+
+			public Priority Priority { get; set; }
+			public bool AllowRetry { get; set; } = true;
+			public bool RunEarly { get; set; }
+			public bool Warnings { get; set; } = true;
+
+			[BsonIgnoreIfNull]
+			public Dictionary<string, string>? Credentials { get; set; }
+
+			[BsonIgnoreIfNull]
+			public Dictionary<string, string>? Properties { get; set; }
+
+			[BsonIgnoreIfNull]
+			[BsonDictionaryOptions(DictionaryRepresentation.ArrayOfDocuments)]
+			public NodeAnnotations? Annotations { get; set; }
+
+			IReadOnlyDictionary<string, string>? INode.Credentials => Credentials;
+			IReadOnlyDictionary<string, string>? INode.Properties => Properties;
+			IReadOnlyNodeAnnotations INode.Annotations => Annotations ?? NodeAnnotations.Empty;
+
+			[BsonConstructor]
+			private Node()
+			{
+				Name = null!;
+				InputDependencies = null!;
+				OrderDependencies = null!;
+			}
+
+			public Node(string name, NodeRef[] inputDependencies, NodeRef[] orderDependencies, Priority priority, bool allowRetry, bool runEarly, bool warnings, Dictionary<string, string>? credentials, Dictionary<string, string>? properties, IReadOnlyNodeAnnotations? annotations)
+			{
+				Name = name;
+				InputDependencies = inputDependencies;
+				OrderDependencies = orderDependencies;
+				Priority = priority;
+				AllowRetry = allowRetry;
+				RunEarly = runEarly;
+				Warnings = warnings;
+				Credentials = credentials;
+				Properties = properties;
+				if (annotations != null && annotations.Count > 0)
+				{
+					Annotations = new NodeAnnotations(annotations);
+				}
+			}
+		}
+
+		class NodeGroup : INodeGroup
+		{
+			public string AgentType { get; set; }
+			public List<Node> Nodes { get; set; }
+
+			IReadOnlyList<INode> INodeGroup.Nodes => Nodes;
+
+			/// <summary>
+			/// Private constructor for BSON serializer
+			/// </summary>
+			[BsonConstructor]
+			private NodeGroup()
+			{
+				AgentType = null!;
+				Nodes = null!;
+			}
+
+			/// <summary>
+			/// Constructor
+			/// </summary>
+			/// <param name="agentType">The type of agent to execute this group</param>
+			/// <param name="nodes">Nodes to execute</param>
+			public NodeGroup(string agentType, List<Node> nodes)
+			{
+				AgentType = agentType;
+				Nodes = nodes;
+			}
+		}
+
+		class Aggregate : IAggregate
+		{
+			public string Name { get; set; }
+			public List<NodeRef> Nodes { get; set; }
+
+			IReadOnlyList<NodeRef> IAggregate.Nodes => Nodes;
+
+			private Aggregate()
+			{
+				Name = null!;
+				Nodes = new List<NodeRef>();
+			}
+
+			public Aggregate(string name, List<NodeRef> nodes)
+			{
+				Name = name;
+				Nodes = nodes;
+			}
+		}
+
+		class Label : ILabel
+		{
+			[BsonIgnoreIfNull]
+			public string? Name { get; set; }
+
+			[BsonIgnoreIfNull]
+			public string? Category { get; set; }
+
+			[BsonIgnoreIfNull]
+			public string? DashboardName { get; set; }
+
+			[BsonIgnoreIfNull]
+			public string? DashboardCategory { get; set; }
+
+			[BsonIgnoreIfNull]
+			public string? UgsName { get; set; }
+
+			[BsonIgnoreIfNull]
+			public string? UgsProject { get; set; }
+
+			[BsonIgnoreIfNull]
+			public LabelChange Change { get; set; }
+
+			public List<NodeRef> RequiredNodes { get; set; }
+			public List<NodeRef> IncludedNodes { get; set; }
+
+			string? ILabel.DashboardName => DashboardName ?? Name;
+			string? ILabel.DashboardCategory => DashboardCategory ?? Category;
+
+			private Label()
+			{
+				RequiredNodes = new List<NodeRef>();
+				IncludedNodes = new List<NodeRef>();
+			}
+
+			public Label(string? dashboardName, string? dashboardCategory, string? ugsName, string? ugsProject, LabelChange change, List<NodeRef> requiredNodes, List<NodeRef> includedNodes)
+			{
+				DashboardName = dashboardName;
+				DashboardCategory = dashboardCategory;
+				UgsName = ugsName;
+				UgsProject = ugsProject;
+				Change = change;
+				RequiredNodes = requiredNodes;
+				IncludedNodes = includedNodes;
+			}
+		}
+
+		class GraphDocument : IGraph
+		{
+			public static GraphDocument Empty { get; } = new GraphDocument(new List<NodeGroup>(), new List<Aggregate>(), new List<Label>());
+
+			[BsonRequired, BsonId]
+			public ContentHash Id { get; private set; } = ContentHash.Empty;
+
+			public int Schema { get; set; }
+			public List<NodeGroup> Groups { get; private set; } = new List<NodeGroup>();
+			public List<Aggregate> Aggregates { get; private set; } = new List<Aggregate>();
+			public List<Label> Labels { get; private set; } = new List<Label>();
+
+			[BsonIgnore]
+			IReadOnlyDictionary<string, NodeRef>? _cachedNodeNameToRef;
+
+			IReadOnlyList<INodeGroup> IGraph.Groups => Groups;
+			IReadOnlyList<IAggregate> IGraph.Aggregates => Aggregates;
+			IReadOnlyList<ILabel> IGraph.Labels => Labels;
+
+			[BsonConstructor]
+			private GraphDocument()
+			{
+			}
+
+			public GraphDocument(List<NodeGroup> groups, List<Aggregate> aggregates, List<Label> labels)
+			{
+				Groups = groups;
+				Aggregates = aggregates;
+				Labels = labels;
+				Id = ContentHash.SHA1(BsonExtensionMethods.ToBson(this));
+			}
+
+			public GraphDocument(GraphDocument baseGraph, List<NewGroup>? newGroupRequests, List<NewAggregate>? newAggregateRequests, List<NewLabel>? newLabelRequests)
+			{
+				Dictionary<string, NodeRef> nodeNameToRef = new Dictionary<string, NodeRef>(baseGraph.GetNodeNameToRef(), StringComparer.OrdinalIgnoreCase);
+
+				// Update the new list of groups
+				List<NodeGroup> newGroups = new List<NodeGroup>(baseGraph.Groups);
+				if (newGroupRequests != null)
+				{
+					foreach (NewGroup newGroupRequest in newGroupRequests)
+					{
+						List<Node> nodes = new List<Node>();
+						foreach (NewNode newNodeRequest in newGroupRequest.Nodes)
+						{
+							int nodeIdx = nodes.Count;
+
+							Priority priority = newNodeRequest.Priority ?? Priority.Normal;
+							bool bAllowRetry = newNodeRequest.AllowRetry ?? true;
+							bool bRunEarly = newNodeRequest.RunEarly ?? false;
+							bool bWarnings = newNodeRequest.Warnings ?? true;
+
+							NodeRef[] inputDependencies = (newNodeRequest.InputDependencies == null) ? Array.Empty<NodeRef>() : newNodeRequest.InputDependencies.Select(x => nodeNameToRef[x]).ToArray();
+							NodeRef[] orderDependencies = (newNodeRequest.OrderDependencies == null) ? Array.Empty<NodeRef>() : newNodeRequest.OrderDependencies.Select(x => nodeNameToRef[x]).ToArray();
+							orderDependencies = orderDependencies.Union(inputDependencies).ToArray();
+							nodes.Add(new Node(newNodeRequest.Name, inputDependencies, orderDependencies, priority, bAllowRetry, bRunEarly, bWarnings, newNodeRequest.Credentials, newNodeRequest.Properties, newNodeRequest.Annotations));
+
+							nodeNameToRef.Add(newNodeRequest.Name, new NodeRef(newGroups.Count, nodeIdx));
+						}
+						newGroups.Add(new NodeGroup(newGroupRequest.AgentType, nodes));
+					}
+				}
+
+				// Update the list of aggregates
+				List<Aggregate> newAggregates = new List<Aggregate>(baseGraph.Aggregates);
+				if (newAggregateRequests != null)
+				{
+					foreach (NewAggregate newAggregateRequest in newAggregateRequests)
+					{
+						List<NodeRef> nodes = newAggregateRequest.Nodes.ConvertAll(x => nodeNameToRef[x]);
+						newAggregates.Add(new Aggregate(newAggregateRequest.Name, nodes));
+					}
+				}
+
+				// Update the list of labels
+				List<Label> newLabels = new List<Label>(baseGraph.Labels);
+				if (newLabelRequests != null)
+				{
+					foreach (NewLabel newLabelRequest in newLabelRequests)
+					{
+						List<NodeRef> requiredNodes = newLabelRequest.RequiredNodes.ConvertAll(x => nodeNameToRef[x]);
+						List<NodeRef> includedNodes = newLabelRequest.IncludedNodes.ConvertAll(x => nodeNameToRef[x]);
+						newLabels.Add(new Label(newLabelRequest.DashboardName, newLabelRequest.DashboardCategory, newLabelRequest.UgsName, newLabelRequest.UgsProject, newLabelRequest.Change, requiredNodes, includedNodes));
+					}
+				}
+
+				// Create the new arrays
+				Groups = newGroups;
+				Aggregates = newAggregates;
+				Labels = newLabels;
+
+				// Create the new graph, and save the generated node lookup into it
+				_cachedNodeNameToRef = nodeNameToRef;
+
+				// Compute the hash
+				Id = ContentHash.SHA1(BsonExtensionMethods.ToBson(this));
+			}
+
+			public IReadOnlyDictionary<string, NodeRef> GetNodeNameToRef()
+			{
+				if (_cachedNodeNameToRef == null)
+				{
+					Dictionary<string, NodeRef> nodeNameToRef = new Dictionary<string, NodeRef>(StringComparer.OrdinalIgnoreCase);
+					for (int groupIdx = 0; groupIdx < Groups.Count; groupIdx++)
+					{
+						List<Node> nodes = Groups[groupIdx].Nodes;
+						for (int nodeIdx = 0; nodeIdx < nodes.Count; nodeIdx++)
+						{
+							Node node = nodes[nodeIdx];
+							nodeNameToRef[node.Name] = new NodeRef(groupIdx, nodeIdx);
+						}
+					}
+					_cachedNodeNameToRef = nodeNameToRef;
+				}
+				return _cachedNodeNameToRef;
+			}
+		}
+
+		/// <summary>
+		/// Stores information about a cached graph
+		/// </summary>
+		class CachedGraph
+		{
+			/// <summary>
+			/// Time at which the graph was last accessed
+			/// </summary>
+			public long LastAccessTime => _lastAccessTimePrivate;
+
+			/// <summary>
+			/// Backing value for <see cref="LastAccessTime"/>
+			/// </summary>
+			private long _lastAccessTimePrivate;
+
+			/// <summary>
+			/// The graph instance
+			/// </summary>
+			public IGraph Graph { get; }
+
+			/// <summary>
+			/// Constructor
+			/// </summary>
+			/// <param name="graph">The graph to store</param>
+			public CachedGraph(IGraph graph)
+			{
+				_lastAccessTimePrivate = Stopwatch.GetTimestamp();
+				Graph = graph;
+			}
+
+			/// <summary>
+			/// Update the last access time
+			/// </summary>
+			public void Touch()
+			{
+				for (; ; )
+				{
+					long time = Stopwatch.GetTimestamp();
+					long lastAccessTimeCopy = _lastAccessTimePrivate;
+					if (time < lastAccessTimeCopy || Interlocked.CompareExchange(ref _lastAccessTimePrivate, time, lastAccessTimeCopy) == lastAccessTimeCopy)
+					{
+						break;
+					}
+				}
+			}
+		}
+
+		/// <summary>
+		/// The jobs collection
+		/// </summary>
+		readonly IMongoCollection<GraphDocument> _graphs;
+
+		/// <summary>
+		/// Maximum number of graphs to keep in the cache
+		/// </summary>
+		const int MaxGraphs = 1000;
+
+		/// <summary>
+		/// Cache of recently accessed graphs
+		/// </summary>
+		readonly ConcurrentDictionary<ContentHash, CachedGraph> _cachedGraphs = new ConcurrentDictionary<ContentHash, CachedGraph>();
+
+		/// <summary>
+		/// Constructor
+		/// </summary>
+		/// <param name="mongoService">The database service singleton</param>
+		public GraphCollection(MongoService mongoService)
+		{
+			_graphs = mongoService.GetCollection<GraphDocument>("Graphs");
+		}
+
+		/// <summary>
+		/// Adds a new graph document
+		/// </summary>
+		/// <param name="graph">The graph to add</param>
+		/// <returns>Async task</returns>
+		async Task AddAsync(GraphDocument graph)
+		{
+			if (!await _graphs.Find(x => x.Id == graph.Id).AnyAsync())
+			{
+				try
+				{
+					await _graphs.InsertOneAsync(graph);
+				}
+				catch (MongoWriteException ex)
+				{
+					if (ex.WriteError.Category != ServerErrorCategory.DuplicateKey)
+					{
+						throw;
+					}
+				}
+			}
+		}
+
+		/// <inheritdoc/>
+		public async Task<IGraph> AddAsync(ITemplate template)
+		{
+			Node node = new Node(IJob.SetupNodeName, Array.Empty<NodeRef>(), Array.Empty<NodeRef>(), Priority.High, true, false, true, null, null, null);
+			NodeGroup group = new NodeGroup(template.InitialAgentType ?? "Win64", new List<Node> { node });
+
+			GraphDocument graph = new GraphDocument(new List<NodeGroup> { group }, new List<Aggregate>(), new List<Label>());
+			await AddAsync(graph);
+			return graph;
+		}
+
+		/// <inheritdoc/>
+		public async Task<IGraph> AppendAsync(IGraph? baseGraph, List<NewGroup>? newGroupRequests, List<NewAggregate>? newAggregateRequests, List<NewLabel>? newLabelRequests)
+		{
+			GraphDocument graph = new GraphDocument((GraphDocument?)baseGraph ?? GraphDocument.Empty, newGroupRequests, newAggregateRequests, newLabelRequests);
+			await AddAsync(graph);
+			return graph;
+		}
+
+		/// <inheritdoc/>
+		public async Task<IGraph> GetAsync(ContentHash? hash)
+		{
+			// Special case for an empty graph request
+			if (hash == null || hash == ContentHash.Empty || hash == GraphDocument.Empty.Id)
+			{
+				return GraphDocument.Empty;
+			}
+
+			// Try to read the graph from the cache
+			CachedGraph? cachedGraph;
+			if (_cachedGraphs.TryGetValue(hash, out cachedGraph))
+			{
+				// Update the last access time
+				cachedGraph.Touch();
+			}
+			else
+			{
+				// Trim the cache
+				while (_cachedGraphs.Count > MaxGraphs)
+				{
+					ContentHash? removeHash = _cachedGraphs.OrderBy(x => x.Value.LastAccessTime).Select(x => x.Key).FirstOrDefault();
+					if (removeHash == null || removeHash == ContentHash.Empty)
+					{
+						break;
+					}
+					_cachedGraphs.TryRemove(removeHash, out _);
+				}
+
+				// Create the new entry
+				cachedGraph = new CachedGraph(await _graphs.Find<GraphDocument>(x => x.Id == hash).FirstAsync());
+				_cachedGraphs.TryAdd(hash, cachedGraph);
+			}
+			return cachedGraph.Graph;
+		}
+
+		/// <inheritdoc/>
+		public async Task<List<IGraph>> FindAllAsync(ContentHash[]? hashes, int? index, int? count)
+		{
+			FilterDefinitionBuilder<GraphDocument> filterBuilder = Builders<GraphDocument>.Filter;
+
+			FilterDefinition<GraphDocument> filter = filterBuilder.Empty;
+			if (hashes != null)
+			{
+				filter &= filterBuilder.In(x => x.Id, hashes);
+			}
+
+			List<GraphDocument> results;
+			IFindFluent<GraphDocument, GraphDocument> search = _graphs.Find(filter);
+			if (index != null)
+			{
+				search = search.Skip(index.Value);
+			}
+			if (count != null)
+			{
+				search = search.Limit(count.Value);
+			}
+
+			results = await search.ToListAsync();
+			return results.ConvertAll<IGraph>(x => x);
+		}
+	}
+}
diff --git a/Engine/Source/Programs/Horde/Horde.Build/Jobs/Graphs/GraphMessages.cs b/Engine/Source/Programs/Horde/Horde.Build/Jobs/Graphs/GraphMessages.cs
--- a/Engine/Source/Programs/Horde/Horde.Build/Jobs/Graphs/GraphMessages.cs
+++ b/Engine/Source/Programs/Horde/Horde.Build/Jobs/Graphs/GraphMessages.cs
@@ -3,10 +3,9 @@
 using System.Collections.Generic;
 using System.Linq;
 using EpicGames.Core;
-using Horde.Build.Models;
 using HordeCommon;
 
-namespace Horde.Build.Api
+namespace Horde.Build.Jobs.Graphs
 {
 	/// <summary>
 	/// Information required to create a node
diff --git a/Engine/Source/Programs/Horde/Horde.Build/Jobs/Graphs/IGraph.cs b/Engine/Source/Programs/Horde/Horde.Build/Jobs/Graphs/IGraph.cs
--- a/Engine/Source/Programs/Horde/Horde.Build/Jobs/Graphs/IGraph.cs
+++ b/Engine/Source/Programs/Horde/Horde.Build/Jobs/Graphs/IGraph.cs
@@ -9,7 +9,7 @@ using EpicGames.Core;
 using HordeCommon;
 using MongoDB.Bson.Serialization.Attributes;
 
-namespace Horde.Build.Models
+namespace Horde.Build.Jobs.Graphs
 {
 	/// <summary>
 	/// Represents a node in the graph
diff --git a/Engine/Source/Programs/Horde/Horde.Build/Jobs/Graphs/IGraphCollection.cs b/Engine/Source/Programs/Horde/Horde.Build/Jobs/Graphs/IGraphCollection.cs
--- a/Engine/Source/Programs/Horde/Horde.Build/Jobs/Graphs/IGraphCollection.cs
+++ b/Engine/Source/Programs/Horde/Horde.Build/Jobs/Graphs/IGraphCollection.cs
@@ -3,9 +3,9 @@
 using System.Collections.Generic;
 using System.Threading.Tasks;
 using EpicGames.Core;
-using Horde.Build.Models;
+using Horde.Build.Jobs.Templates;
 
-namespace Horde.Build.Collections
+namespace Horde.Build.Jobs.Graphs
 {
 	/// <summary>
 	/// Interface for a collection of graph documents
diff --git a/Engine/Source/Programs/Horde/Horde.Build/Jobs/IJob.cs b/Engine/Source/Programs/Horde/Horde.Build/Jobs/IJob.cs
--- a/Engine/Source/Programs/Horde/Horde.Build/Jobs/IJob.cs
+++ b/Engine/Source/Programs/Horde/Horde.Build/Jobs/IJob.cs
@@ -6,13 +6,22 @@ using System.Diagnostics.CodeAnalysis;
 using System.Linq;
 using EpicGames.Core;
 using Horde.Build.Acls;
-using Horde.Build.Api;
+using Horde.Build.Agents;
+using Horde.Build.Agents.Leases;
+using Horde.Build.Agents.Pools;
+using Horde.Build.Agents.Sessions;
+using Horde.Build.Jobs.Graphs;
+using Horde.Build.Jobs.Timing;
+using Horde.Build.Logs;
+using Horde.Build.Streams;
+using Horde.Build.Ugs;
+using Horde.Build.Users;
 using Horde.Build.Utilities;
 using HordeCommon;
 using MongoDB.Bson;
 using MongoDB.Driver;
 
-namespace Horde.Build.Models
+namespace Horde.Build.Jobs
 {
 	using JobId = ObjectId<IJob>;
 	using LeaseId = ObjectId<ILease>;
diff --git a/Engine/Source/Programs/Horde/Horde.Build/Jobs/IJobCollection.cs b/Engine/Source/Programs/Horde/Horde.Build/Jobs/IJobCollection.cs
--- a/Engine/Source/Programs/Horde/Horde.Build/Jobs/IJobCollection.cs
+++ b/Engine/Source/Programs/Horde/Horde.Build/Jobs/IJobCollection.cs
@@ -4,13 +4,19 @@ using System;
 using System.Collections.Generic;
 using System.Threading.Tasks;
 using EpicGames.Core;
-using Horde.Build.Api;
-using Horde.Build.Models;
+using Horde.Build.Agents;
+using Horde.Build.Agents.Leases;
+using Horde.Build.Agents.Pools;
+using Horde.Build.Agents.Sessions;
+using Horde.Build.Jobs.Graphs;
+using Horde.Build.Logs;
+using Horde.Build.Streams;
+using Horde.Build.Users;
 using Horde.Build.Utilities;
 using HordeCommon;
 using MongoDB.Bson;
 
-namespace Horde.Build.Collections
+namespace Horde.Build.Jobs
 {
 	using JobId = ObjectId<IJob>;
 	using LeaseId = ObjectId<ILease>;
diff --git a/Engine/Source/Programs/Horde/Horde.Build/Jobs/IJobStepRef.cs b/Engine/Source/Programs/Horde/Horde.Build/Jobs/IJobStepRef.cs
--- a/Engine/Source/Programs/Horde/Horde.Build/Jobs/IJobStepRef.cs
+++ b/Engine/Source/Programs/Horde/Horde.Build/Jobs/IJobStepRef.cs
@@ -1,12 +1,16 @@
 // Copyright Epic Games, Inc. All Rights Reserved.
 
 using System;
+using Horde.Build.Agents;
+using Horde.Build.Agents.Pools;
+using Horde.Build.Logs;
+using Horde.Build.Streams;
 using Horde.Build.Utilities;
 using HordeCommon;
 using MongoDB.Bson.Serialization;
 using MongoDB.Bson.Serialization.Attributes;
 
-namespace Horde.Build.Models
+namespace Horde.Build.Jobs
 {
 	using JobId = ObjectId<IJob>;
 	using LogId = ObjectId<ILogFile>;
diff --git a/Engine/Source/Programs/Horde/Horde.Build/Jobs/IJobStepRefCollection.cs b/Engine/Source/Programs/Horde/Horde.Build/Jobs/IJobStepRefCollection.cs
--- a/Engine/Source/Programs/Horde/Horde.Build/Jobs/IJobStepRefCollection.cs
+++ b/Engine/Source/Programs/Horde/Horde.Build/Jobs/IJobStepRefCollection.cs
@@ -3,11 +3,15 @@
 using System;
 using System.Collections.Generic;
 using System.Threading.Tasks;
-using Horde.Build.Models;
+using Horde.Build.Agents;
+using Horde.Build.Agents.Pools;
+using Horde.Build.Jobs.Graphs;
+using Horde.Build.Logs;
+using Horde.Build.Streams;
 using Horde.Build.Utilities;
 using HordeCommon;
 
-namespace Horde.Build.Collections
+namespace Horde.Build.Jobs
 {
 	using LogId = ObjectId<ILogFile>;
 	using PoolId = StringId<IPool>;
diff --git a/Engine/Source/Programs/Horde/Horde.Build/Jobs/JobCollection.cs b/Engine/Source/Programs/Horde/Horde.Build/Jobs/JobCollection.cs
--- a/Engine/Source/Programs/Horde/Horde.Build/Jobs/JobCollection.cs
+++ b/Engine/Source/Programs/Horde/Horde.Build/Jobs/JobCollection.cs
@@ -8,10 +8,15 @@ using System.Linq;
 using System.Threading.Tasks;
 using EpicGames.Core;
 using Horde.Build.Acls;
-using Horde.Build.Api;
-using Horde.Build.Models;
+using Horde.Build.Agents;
+using Horde.Build.Agents.Leases;
+using Horde.Build.Agents.Pools;
+using Horde.Build.Agents.Sessions;
+using Horde.Build.Jobs.Graphs;
+using Horde.Build.Logs;
 using Horde.Build.Server;
-using Horde.Build.Services;
+using Horde.Build.Streams;
+using Horde.Build.Users;
 using Horde.Build.Utilities;
 using HordeCommon;
 using Microsoft.Extensions.Logging;
@@ -23,7 +28,7 @@ using MongoDB.Driver;
 using OpenTracing;
 using OpenTracing.Util;
 
-namespace Horde.Build.Collections.Impl
+namespace Horde.Build.Jobs
 {
 	using JobId = ObjectId<IJob>;
 	using LeaseId = ObjectId<ILease>;
diff --git a/Engine/Source/Programs/Horde/Horde.Build/Jobs/JobService.cs b/Engine/Source/Programs/Horde/Horde.Build/Jobs/JobService.cs
--- a/Engine/Source/Programs/Horde/Horde.Build/Jobs/JobService.cs
+++ b/Engine/Source/Programs/Horde/Horde.Build/Jobs/JobService.cs
@@ -9,10 +9,18 @@ using System.Text.RegularExpressions;
 using System.Threading.Tasks;
 using EpicGames.Core;
 using Horde.Build.Acls;
-using Horde.Build.Api;
-using Horde.Build.Collections;
-using Horde.Build.Models;
-using Horde.Build.Tasks.Impl;
+using Horde.Build.Agents;
+using Horde.Build.Agents.Leases;
+using Horde.Build.Agents.Sessions;
+using Horde.Build.Issues;
+using Horde.Build.Jobs.Graphs;
+using Horde.Build.Jobs.Templates;
+using Horde.Build.Jobs.Timing;
+using Horde.Build.Logs;
+using Horde.Build.Notifications;
+using Horde.Build.Perforce;
+using Horde.Build.Streams;
+using Horde.Build.Users;
 using Horde.Build.Utilities;
 using HordeCommon;
 using Microsoft.Extensions.Logging;
@@ -21,7 +29,7 @@ using MongoDB.Driver;
 using OpenTracing;
 using OpenTracing.Util;
 
-namespace Horde.Build.Services
+namespace Horde.Build.Jobs
 {
 	using JobId = ObjectId<IJob>;
 	using LeaseId = ObjectId<ILease>;
diff --git a/Engine/Source/Programs/Horde/Horde.Build/Jobs/JobStepRefCollection.cs b/Engine/Source/Programs/Horde/Horde.Build/Jobs/JobStepRefCollection.cs
--- a/Engine/Source/Programs/Horde/Horde.Build/Jobs/JobStepRefCollection.cs
+++ b/Engine/Source/Programs/Horde/Horde.Build/Jobs/JobStepRefCollection.cs
@@ -4,14 +4,17 @@ using System;
 using System.Collections.Generic;
 using System.Linq;
 using System.Threading.Tasks;
-using Horde.Build.Models;
+using Horde.Build.Agents;
+using Horde.Build.Agents.Pools;
+using Horde.Build.Logs;
 using Horde.Build.Server;
+using Horde.Build.Streams;
 using Horde.Build.Utilities;
 using HordeCommon;
 using MongoDB.Bson.Serialization.Attributes;
 using MongoDB.Driver;
 
-namespace Horde.Build.Collections.Impl
+namespace Horde.Build.Jobs
 {
 	using LogId = ObjectId<ILogFile>;
 	using PoolId = StringId<IPool>;
diff --git a/Engine/Source/Programs/Horde/Horde.Build/Jobs/JobTaskSource.cs b/Engine/Source/Programs/Horde/Horde.Build/Jobs/JobTaskSource.cs
--- a/Engine/Source/Programs/Horde/Horde.Build/Jobs/JobTaskSource.cs
+++ b/Engine/Source/Programs/Horde/Horde.Build/Jobs/JobTaskSource.cs
@@ -12,12 +12,16 @@ using System.Threading.Tasks;
 using EpicGames.Core;
 using Google.Protobuf;
 using Google.Protobuf.WellKnownTypes;
-using Horde.Build.Api;
-using Horde.Build.Collections;
-using Horde.Build.Commits;
-using Horde.Build.Models;
+using Horde.Build.Agents;
+using Horde.Build.Agents.Leases;
+using Horde.Build.Agents.Pools;
+using Horde.Build.Jobs.Graphs;
+using Horde.Build.Logs;
+using Horde.Build.Perforce;
 using Horde.Build.Server;
-using Horde.Build.Services;
+using Horde.Build.Streams;
+using Horde.Build.Tasks;
+using Horde.Build.Ugs;
 using Horde.Build.Utilities;
 using HordeCommon;
 using HordeCommon.Rpc.Tasks;
@@ -26,7 +30,7 @@ using Microsoft.Extensions.Logging;
 using Microsoft.Extensions.Options;
 using MongoDB.Driver;
 
-namespace Horde.Build.Tasks.Impl
+namespace Horde.Build.Jobs
 {
 	using JobId = ObjectId<IJob>;
 	using LeaseId = ObjectId<ILease>;
diff --git a/Engine/Source/Programs/Horde/Horde.Build/Jobs/JobsController.Messages.cs b/Engine/Source/Programs/Horde/Horde.Build/Jobs/JobsController.Messages.cs
--- a/Engine/Source/Programs/Horde/Horde.Build/Jobs/JobsController.Messages.cs
+++ b/Engine/Source/Programs/Horde/Horde.Build/Jobs/JobsController.Messages.cs
@@ -6,10 +6,14 @@ using System.ComponentModel.DataAnnotations;
 using System.Linq;
 using System.Text.Json.Serialization;
 using EpicGames.Core;
-using Horde.Build.Models;
+using Horde.Build.Acls;
+using Horde.Build.Jobs;
+using Horde.Build.Jobs.Graphs;
+using Horde.Build.Streams;
+using Horde.Build.Users;
 using HordeCommon;
 
-namespace Horde.Build.Api
+namespace Horde.Build.Jobs
 {
 	/// <summary>
 	/// State of the job
diff --git a/Engine/Source/Programs/Horde/Horde.Build/Jobs/JobsController.cs b/Engine/Source/Programs/Horde/Horde.Build/Jobs/JobsController.cs
--- a/Engine/Source/Programs/Horde/Horde.Build/Jobs/JobsController.cs
+++ b/Engine/Source/Programs/Horde/Horde.Build/Jobs/JobsController.cs
@@ -7,12 +7,16 @@ using System.Threading.Tasks;
 using EpicGames.Core;
 using EpicGames.Perforce;
 using Horde.Build.Acls;
-using Horde.Build.Api;
-using Horde.Build.Collections;
-using Horde.Build.Jobs;
-using Horde.Build.Models;
+using Horde.Build.Agents;
+using Horde.Build.Jobs.Artifacts;
+using Horde.Build.Jobs.Graphs;
+using Horde.Build.Jobs.Templates;
+using Horde.Build.Jobs.Timing;
+using Horde.Build.Logs;
 using Horde.Build.Notifications;
-using Horde.Build.Services;
+using Horde.Build.Perforce;
+using Horde.Build.Streams;
+using Horde.Build.Users;
 using Horde.Build.Utilities;
 using HordeCommon;
 using Microsoft.AspNetCore.Authorization;
@@ -22,7 +26,7 @@ using MongoDB.Bson;
 using OpenTracing;
 using OpenTracing.Util;
 
-namespace Horde.Build.Controllers
+namespace Horde.Build.Jobs
 {
 	using JobId = ObjectId<IJob>;
 	using StreamId = StringId<IStream>;
diff --git a/Engine/Source/Programs/Horde/Horde.Build/Jobs/Schedules/ISchedule.cs b/Engine/Source/Programs/Horde/Horde.Build/Jobs/Schedules/ISchedule.cs
--- a/Engine/Source/Programs/Horde/Horde.Build/Jobs/Schedules/ISchedule.cs
+++ b/Engine/Source/Programs/Horde/Horde.Build/Jobs/Schedules/ISchedule.cs
@@ -2,10 +2,12 @@
 
 using System;
 using System.Collections.Generic;
+using Horde.Build.Perforce;
+using Horde.Build.Streams;
 using Horde.Build.Utilities;
 using MongoDB.Bson.Serialization.Attributes;
 
-namespace Horde.Build.Models
+namespace Horde.Build.Jobs.Schedules
 {
 	using JobId = ObjectId<IJob>;
 	using TemplateRefId = StringId<TemplateRef>;
diff --git a/Engine/Source/Programs/Horde/Horde.Build/Jobs/Schedules/ScheduleService.cs b/Engine/Source/Programs/Horde/Horde.Build/Jobs/Schedules/ScheduleService.cs
--- a/Engine/Source/Programs/Horde/Horde.Build/Jobs/Schedules/ScheduleService.cs
+++ b/Engine/Source/Programs/Horde/Horde.Build/Jobs/Schedules/ScheduleService.cs
@@ -11,10 +11,11 @@ using EpicGames.Core;
 using EpicGames.Redis;
 using EpicGames.Redis.Utility;
 using EpicGames.Serialization;
-using Horde.Build.Api;
-using Horde.Build.Collections;
-using Horde.Build.Models;
+using Horde.Build.Jobs.Graphs;
+using Horde.Build.Jobs.Templates;
+using Horde.Build.Perforce;
 using Horde.Build.Server;
+using Horde.Build.Streams;
 using Horde.Build.Utilities;
 using HordeCommon;
 using Microsoft.Extensions.Hosting;
@@ -26,7 +27,7 @@ using OpenTracing.Util;
 using StackExchange.Redis;
 using TimeZoneConverter;
 
-namespace Horde.Build.Services
+namespace Horde.Build.Jobs.Schedules
 {
 	using JobId = ObjectId<IJob>;
 	using StreamId = StringId<IStream>;
diff --git a/Engine/Source/Programs/Horde/Horde.Build/Jobs/Schedules/Schedules.cs b/Engine/Source/Programs/Horde/Horde.Build/Jobs/Schedules/Schedules.cs
new file mode 100644
--- /dev/null
+++ b/Engine/Source/Programs/Horde/Horde.Build/Jobs/Schedules/Schedules.cs
@@ -0,0 +1,310 @@
+// Copyright Epic Games, Inc. All Rights Reserved.
+
+using System;
+using System.Collections.Generic;
+using System.ComponentModel.DataAnnotations;
+using Horde.Build.Acls;
+using Horde.Build.Jobs.Schedules;
+using Horde.Build.Perforce;
+using Horde.Build.Streams;
+using Horde.Build.Utilities;
+
+namespace Horde.Build.Jobs.Schedules
+{
+	using TemplateRefId = StringId<TemplateRef>;
+
+	/// <summary>
+	/// Parameters to create a new schedule
+	/// </summary>
+	public class CreateSchedulePatternRequest
+	{
+		/// <summary>
+		/// Days of the week to run this schedule on. If null, the schedule will run every day.
+		/// </summary>
+		public List<string>? DaysOfWeek { get; set; }
+
+		/// <summary>
+		/// Time during the day for the first schedule to trigger. Measured in minutes from midnight.
+		/// </summary>
+		public int MinTime { get; set; }
+
+		/// <summary>
+		/// Time during the day for the last schedule to trigger. Measured in minutes from midnight.
+		/// </summary>
+		public int? MaxTime { get; set; }
+
+		/// <summary>
+		/// Interval between each schedule triggering
+		/// </summary>
+		public int? Interval { get; set; }
+
+		/// <summary>
+		/// Constructs a model object from this request
+		/// </summary>
+		/// <returns>Model object</returns>
+		public SchedulePattern ToModel()
+		{
+			return new SchedulePattern(DaysOfWeek?.ConvertAll(x => Enum.Parse<DayOfWeek>(x)), MinTime, MaxTime, Interval);
+		}
+	}
+
+	/// <summary>
+	/// Gate allowing a schedule to trigger.
+	/// </summary>
+	public class CreateScheduleGateRequest
+	{
+		/// <summary>
+		/// The template containing the dependency
+		/// </summary>
+		[Required]
+		public string TemplateId { get; set; } = String.Empty;
+
+		/// <summary>
+		/// Target to wait for
+		/// </summary>
+		[Required]
+		public string Target { get; set; } = String.Empty;
+
+		/// <summary>
+		/// Constructs a model object
+		/// </summary>
+		/// <returns>New model object.</returns>
+		public ScheduleGate ToModel()
+		{
+			return new ScheduleGate(new TemplateRefId(TemplateId), Target);
+		}
+	}
+
+	/// <summary>
+	/// Parameters to create a new schedule
+	/// </summary>
+	public class CreateScheduleRequest
+	{
+		/// <summary>
+		/// Roles to impersonate for this schedule
+		/// </summary>
+		public List<CreateAclClaimRequest>? Claims { get; set; }
+
+		/// <summary>
+		/// Whether the schedule should be enabled
+		/// </summary>
+		public bool Enabled { get; set; }
+
+		/// <summary>
+		/// Maximum number of builds that can be active at once
+		/// </summary>
+		public int MaxActive { get; set; }
+
+		/// <summary>
+		/// Maximum number of changes the schedule can fall behind head revision. If greater than zero, builds will be triggered for every submitted changelist until the backlog is this size.
+		/// </summary>
+		public int MaxChanges { get; set; }
+
+		/// <summary>
+		/// Whether the build requires a change to be submitted
+		/// </summary>
+		public bool RequireSubmittedChange { get; set; } = true;
+
+		/// <summary>
+		/// Gate allowing the schedule to trigger
+		/// </summary>
+		public CreateScheduleGateRequest? Gate { get; set; }
+
+		/// <summary>
+		/// The types of changes to run for
+		/// </summary>
+		public List<ChangeContentFlags>? Filter { get; set; }
+
+		/// <summary>
+		/// Files that should cause the job to trigger
+		/// </summary>
+		public List<string>? Files { get; set; }
+
+		/// <summary>
+		/// Parameters for the template
+		/// </summary>
+		public Dictionary<string, string> TemplateParameters { get; set; } = new Dictionary<string, string>();
+
+		/// <summary>
+		/// New patterns for the schedule
+		/// </summary>
+		public List<CreateSchedulePatternRequest> Patterns { get; set; } = new List<CreateSchedulePatternRequest>();
+
+		/// <summary>
+		/// Constructs a model object
+		/// </summary>
+		/// <param name="currentTimeUtc">The current time</param>
+		/// <returns>New model object</returns>
+		public Schedule ToModel(DateTime currentTimeUtc)
+		{
+			return new Schedule(currentTimeUtc, Enabled, MaxActive, MaxChanges, RequireSubmittedChange, Gate?.ToModel(), Filter, Files, TemplateParameters, Patterns.ConvertAll(x => x.ToModel()));
+		}
+	}
+
+	/// <summary>
+	/// Information about a schedule pattern
+	/// </summary>
+	public class GetSchedulePatternResponse
+	{
+		/// <summary>
+		/// Days of the week to run this schedule on. If null, the schedule will run every day.
+		/// </summary>
+		public List<string>? DaysOfWeek { get; set; }
+
+		/// <summary>
+		/// Time during the day for the first schedule to trigger. Measured in minutes from midnight.
+		/// </summary>
+		public int MinTime { get; set; }
+
+		/// <summary>
+		/// Time during the day for the last schedule to trigger. Measured in minutes from midnight.
+		/// </summary>
+		public int? MaxTime { get; set; }
+
+		/// <summary>
+		/// Interval between each schedule triggering
+		/// </summary>
+		public int? Interval { get; set; }
+
+		/// <summary>
+		/// Constructor
+		/// </summary>
+		/// <param name="pattern">The pattern to construct from</param>
+		public GetSchedulePatternResponse(SchedulePattern pattern)
+		{
+			DaysOfWeek = pattern.DaysOfWeek?.ConvertAll(x => x.ToString());
+			MinTime = pattern.MinTime;
+			MaxTime = pattern.MaxTime;
+			Interval = pattern.Interval;
+		}
+	}
+
+	/// <summary>
+	/// Gate allowing a schedule to trigger.
+	/// </summary>
+	public class GetScheduleGateResponse
+	{
+		/// <summary>
+		/// The template containing the dependency
+		/// </summary>
+		public string TemplateId { get; set; }
+
+		/// <summary>
+		/// Target to wait for
+		/// </summary>
+		public string Target { get; set; }
+
+		/// <summary>
+		/// Constructor
+		/// </summary>
+		/// <param name="gate">Gate to construct from</param>
+		public GetScheduleGateResponse(ScheduleGate gate)
+		{
+			TemplateId = gate.TemplateRefId.ToString();
+			Target = gate.Target;
+		}
+	}
+
+	/// <summary>
+	/// Response describing a schedule
+	/// </summary>
+	public class GetScheduleResponse
+	{
+		/// <summary>
+		/// Whether the schedule is currently enabled
+		/// </summary>
+		public bool Enabled { get; set; }
+
+		/// <summary>
+		/// Maximum number of scheduled jobs at once
+		/// </summary>
+		public int MaxActive { get; set; }
+
+		/// <summary>
+		/// Maximum number of changes the schedule can fall behind head revision. If greater than zero, builds will be triggered for every submitted changelist until the backlog is this size.
+		/// </summary>
+		public int MaxChanges { get; set; }
+
+		/// <summary>
+		/// Whether the build requires a change to be submitted
+		/// </summary>
+		public bool RequireSubmittedChange { get; set; }
+
+		/// <summary>
+		/// Gate for this schedule to trigger
+		/// </summary>
+		public GetScheduleGateResponse? Gate { get; set; }
+
+		/// <summary>
+		/// The types of changes to run for
+		/// </summary>
+		public List<ChangeContentFlags>? Filter { get; set; }
+
+		/// <summary>
+		/// Parameters for the template
+		/// </summary>
+		public Dictionary<string, string> TemplateParameters { get; set; }
+
+		/// <summary>
+		/// New patterns for the schedule
+		/// </summary>
+		public List<GetSchedulePatternResponse> Patterns { get; set; }
+
+		/// <summary>
+		/// Last changelist number that this was triggered for
+		/// </summary>
+		public int LastTriggerChange { get; set; }
+
+		/// <summary>
+		/// Last time that the schedule was triggered
+		/// </summary>
+		public DateTimeOffset LastTriggerTime { get; set; }
+
+		/// <summary>
+		/// List of active jobs
+		/// </summary>
+		public List<string> ActiveJobs { get; set; }
+
+		/// <summary>
+		/// Constructor
+		/// </summary>
+		/// <param name="schedule">The schedule to construct from</param>
+		public GetScheduleResponse(Schedule schedule)
+		{
+			Enabled = schedule.Enabled;
+			MaxActive = schedule.MaxActive;
+			MaxChanges = schedule.MaxChanges;
+			RequireSubmittedChange = schedule.RequireSubmittedChange;
+			if (schedule.Gate != null)
+			{
+				Gate = new GetScheduleGateResponse(schedule.Gate);
+			}
+			Filter = schedule.Filter;
+			TemplateParameters = schedule.TemplateParameters;
+			Patterns = schedule.Patterns.ConvertAll(x => new GetSchedulePatternResponse(x));
+			LastTriggerChange = schedule.LastTriggerChange;
+			LastTriggerTime = schedule.GetLastTriggerTimeUtc();
+			ActiveJobs = schedule.ActiveJobs.ConvertAll(x => x.ToString());
+		}
+	}
+
+	/// <summary>
+	/// Response describing when a schedule is expected to trigger
+	/// </summary>
+	public class GetScheduleForecastResponse
+	{
+		/// <summary>
+		/// Next trigger times
+		/// </summary>
+		public List<DateTime> Times { get; set; }
+
+		/// <summary>
+		/// Constructor
+		/// </summary>
+		/// <param name="times">List of trigger times</param>
+		public GetScheduleForecastResponse(List<DateTime> times)
+		{
+			Times = times;
+		}
+	}
+}
diff --git a/Engine/Source/Programs/Horde/Horde.Build/Jobs/SubResourceId.cs b/Engine/Source/Programs/Horde/Horde.Build/Jobs/SubResourceId.cs
--- a/Engine/Source/Programs/Horde/Horde.Build/Jobs/SubResourceId.cs
+++ b/Engine/Source/Programs/Horde/Horde.Build/Jobs/SubResourceId.cs
@@ -7,7 +7,7 @@ using System.Globalization;
 using MongoDB.Bson.Serialization;
 using MongoDB.Bson.Serialization.Attributes;
 
-namespace Horde.Build.Models
+namespace Horde.Build.Jobs
 {
 	/// <summary>
 	/// Identifier for subresources. Assigning unique ids to subresources prevents against race conditions using indices when subresources are added and removed.
diff --git a/Engine/Source/Programs/Horde/Horde.Build/Jobs/Templates/ITemplate.cs b/Engine/Source/Programs/Horde/Horde.Build/Jobs/Templates/ITemplate.cs
--- a/Engine/Source/Programs/Horde/Horde.Build/Jobs/Templates/ITemplate.cs
+++ b/Engine/Source/Programs/Horde/Horde.Build/Jobs/Templates/ITemplate.cs
@@ -3,11 +3,10 @@
 using System;
 using System.Collections.Generic;
 using EpicGames.Core;
-using Horde.Build.Api;
 using HordeCommon;
 using MongoDB.Bson.Serialization.Attributes;
 
-namespace Horde.Build.Models
+namespace Horde.Build.Jobs.Templates
 {
 	/// <summary>
 	/// Base class for parameters used to configure templates via the new build dialog
diff --git a/Engine/Source/Programs/Horde/Horde.Build/Jobs/Templates/ITemplateCollection.cs b/Engine/Source/Programs/Horde/Horde.Build/Jobs/Templates/ITemplateCollection.cs
--- a/Engine/Source/Programs/Horde/Horde.Build/Jobs/Templates/ITemplateCollection.cs
+++ b/Engine/Source/Programs/Horde/Horde.Build/Jobs/Templates/ITemplateCollection.cs
@@ -3,10 +3,9 @@
 using System.Collections.Generic;
 using System.Threading.Tasks;
 using EpicGames.Core;
-using Horde.Build.Models;
 using HordeCommon;
 
-namespace Horde.Build.Collections
+namespace Horde.Build.Jobs.Templates
 {
 	/// <summary>
 	/// Interface for a collection of template documents
diff --git a/Engine/Source/Programs/Horde/Horde.Build/Jobs/Templates/TemplateCollection.cs b/Engine/Source/Programs/Horde/Horde.Build/Jobs/Templates/TemplateCollection.cs
new file mode 100644
--- /dev/null
+++ b/Engine/Source/Programs/Horde/Horde.Build/Jobs/Templates/TemplateCollection.cs
@@ -0,0 +1,152 @@
+// Copyright Epic Games, Inc. All Rights Reserved.
+
+using System;
+using System.Collections.Generic;
+using System.Linq;
+using System.Threading.Tasks;
+using EpicGames.Core;
+using Horde.Build.Server;
+using Horde.Build.Streams;
+using Horde.Build.Utilities;
+using HordeCommon;
+using Microsoft.Extensions.Caching.Memory;
+using MongoDB.Bson;
+using MongoDB.Bson.Serialization;
+using MongoDB.Bson.Serialization.Attributes;
+using MongoDB.Driver;
+
+namespace Horde.Build.Jobs.Templates
+{
+	/// <summary>
+	/// Collection of template documents
+	/// </summary>
+	public sealed class TemplateCollection : ITemplateCollection, IDisposable
+	{
+		/// <summary>
+		/// Document describing a job template. These objects are considered immutable once created and uniquely referenced by hash, in order to de-duplicate across all job runs.
+		/// </summary>
+		class TemplateDocument : ITemplate
+		{
+			[BsonRequired, BsonId]
+			public ContentHash Id { get; private set; } = ContentHash.Empty;
+
+			[BsonRequired]
+			public string Name { get; private set; }
+
+			public Priority? Priority { get; private set; }
+			public bool AllowPreflights { get; set; } = true;
+			public bool UpdateIssues { get; set; } = false;
+			public bool PromoteIssuesByDefault { get; set; } = false;
+			public string? InitialAgentType { get; set; }
+
+			[BsonIgnoreIfNull]
+			public string? SubmitNewChange { get; set; }
+
+			[BsonIgnoreIfNull]
+			public string? SubmitDescription { get; set; }
+
+			public List<string> Arguments { get; private set; } = new List<string>();
+			public List<Parameter> Parameters { get; private set; } = new List<Parameter>();
+
+			IReadOnlyList<string> ITemplate.Arguments => Arguments;
+			IReadOnlyList<Parameter> ITemplate.Parameters => Parameters;
+
+			[BsonConstructor]
+			private TemplateDocument()
+			{
+				Name = null!;
+			}
+
+			public TemplateDocument(string name, Priority? priority, bool bAllowPreflights, bool updateIssues, bool promoteIssuesByDefault, string? initialAgentType, string? submitNewChange, string? submitDescription, List<string>? arguments, List<Parameter>? parameters)
+			{
+				Name = name;
+				Priority = priority;
+				AllowPreflights = bAllowPreflights;
+				UpdateIssues = updateIssues;
+				PromoteIssuesByDefault = promoteIssuesByDefault;
+				InitialAgentType = initialAgentType;
+				SubmitNewChange = submitNewChange;
+				SubmitDescription = submitDescription;
+				Arguments = arguments ?? new List<string>();
+				Parameters = parameters ?? new List<Parameter>();
+
+				// Compute the hash once all other fields have been set
+				Id = ContentHash.SHA1(BsonExtensionMethods.ToBson(this));
+			}
+		}
+
+		/// <summary>
+		/// Template documents
+		/// </summary>
+		readonly IMongoCollection<TemplateDocument> _templates;
+
+		/// <summary>
+		/// Cache of template documents
+		/// </summary>
+		readonly MemoryCache _templateCache;
+
+		/// <summary>
+		/// Constructor
+		/// </summary>
+		/// <param name="mongoService">The database service singleton</param>
+		public TemplateCollection(MongoService mongoService)
+		{
+			// Ensure discriminator cannot be registered twice (throws exception). Can otherwise happen during unit tests.
+			if (BsonSerializer.LookupDiscriminatorConvention(typeof(JobsTabColumn)) == null)
+			{
+				BsonSerializer.RegisterDiscriminatorConvention(typeof(JobsTabColumn), new DefaultDiscriminatorConvention(typeof(JobsTabColumn), typeof(JobsTabLabelColumn)));	
+			}
+			
+			_templates = mongoService.GetCollection<TemplateDocument>("Templates");
+
+			MemoryCacheOptions options = new MemoryCacheOptions();
+			_templateCache = new MemoryCache(options);
+		}
+
+		/// <inheritdoc/>
+		public void Dispose()
+		{
+			_templateCache.Dispose();
+		}
+
+		/// <inheritdoc/>
+		public async Task<ITemplate> AddAsync(string name, Priority? priority, bool bAllowPreflights, bool bUpdateIssues, bool bPromoteIssuesByDefault, string? initialAgentType, string? submitNewChange, string? submitDescription, List<string>? arguments, List<Parameter>? parameters)
+		{
+			TemplateDocument template = new TemplateDocument(name, priority, bAllowPreflights, bUpdateIssues, bPromoteIssuesByDefault, initialAgentType, submitNewChange, submitDescription, arguments, parameters);
+			if (await GetAsync(template.Id) == null)
+			{
+				await _templates.ReplaceOneAsync(x => x.Id == template.Id, template, new ReplaceOptions { IsUpsert = true });
+			}
+			return template;
+		}
+
+		/// <inheritdoc/>
+		public async Task<List<ITemplate>> FindAllAsync()
+		{
+			List<TemplateDocument> results = await _templates.Find(FilterDefinition<TemplateDocument>.Empty).ToListAsync();
+			return results.ConvertAll<ITemplate>(x => x);
+		}
+
+		/// <inheritdoc/>
+		public async Task<ITemplate?> GetAsync(ContentHash templateId)
+		{
+			object? result;
+			if (_templateCache.TryGetValue(templateId, out result))
+			{
+				return (ITemplate?)result;
+			}
+
+			ITemplate? template = await _templates.Find<TemplateDocument>(x => x.Id == templateId).FirstOrDefaultAsync();
+			if (template != null)
+			{
+				using (ICacheEntry entry = _templateCache.CreateEntry(templateId))
+				{
+					entry.SetSlidingExpiration(TimeSpan.FromMinutes(30.0));
+					entry.SetValue(template);
+				}
+			}
+
+			return template;
+		}
+	}
+}
diff --git a/Engine/Source/Programs/Horde/Horde.Build/Jobs/Templates/TemplateMessages.cs b/Engine/Source/Programs/Horde/Horde.Build/Jobs/Templates/TemplateMessages.cs
--- a/Engine/Source/Programs/Horde/Horde.Build/Jobs/Templates/TemplateMessages.cs
+++ b/Engine/Source/Programs/Horde/Horde.Build/Jobs/Templates/TemplateMessages.cs
@@ -4,10 +4,9 @@ using System;
 using System.Collections.Generic;
 using System.ComponentModel.DataAnnotations;
 using EpicGames.Core;
-using Horde.Build.Models;
 using HordeCommon;
 
-namespace Horde.Build.Api
+namespace Horde.Build.Jobs.Templates
 {
 	/// <summary>
 	/// Base class for template parameters
diff --git a/Engine/Source/Programs/Horde/Horde.Build/Jobs/TestData/ITestData.cs b/Engine/Source/Programs/Horde/Horde.Build/Jobs/TestData/ITestData.cs
--- a/Engine/Source/Programs/Horde/Horde.Build/Jobs/TestData/ITestData.cs
+++ b/Engine/Source/Programs/Horde/Horde.Build/Jobs/TestData/ITestData.cs
@@ -1,9 +1,10 @@
 // Copyright Epic Games, Inc. All Rights Reserved.
 
+using Horde.Build.Streams;
 using Horde.Build.Utilities;
 using MongoDB.Bson;
 
-namespace Horde.Build.Models
+namespace Horde.Build.Jobs.TestData
 {
 	using JobId = ObjectId<IJob>;
 	using StreamId = StringId<IStream>;
diff --git a/Engine/Source/Programs/Horde/Horde.Build/Jobs/TestData/ITestDataCollection.cs b/Engine/Source/Programs/Horde/Horde.Build/Jobs/TestData/ITestDataCollection.cs
--- a/Engine/Source/Programs/Horde/Horde.Build/Jobs/TestData/ITestDataCollection.cs
+++ b/Engine/Source/Programs/Horde/Horde.Build/Jobs/TestData/ITestDataCollection.cs
@@ -2,11 +2,11 @@
 
 using System.Collections.Generic;
 using System.Threading.Tasks;
-using Horde.Build.Models;
+using Horde.Build.Streams;
 using Horde.Build.Utilities;
 using MongoDB.Bson;
 
-namespace Horde.Build.Collections
+namespace Horde.Build.Jobs.TestData
 {
 	using JobId = ObjectId<IJob>;
 	using StreamId = StringId<IStream>;
diff --git a/Engine/Source/Programs/Horde/Horde.Build/Jobs/TestData/TestDataCollection.cs b/Engine/Source/Programs/Horde/Horde.Build/Jobs/TestData/TestDataCollection.cs
new file mode 100644
--- /dev/null
+++ b/Engine/Source/Programs/Horde/Horde.Build/Jobs/TestData/TestDataCollection.cs
@@ -0,0 +1,129 @@
+// Copyright Epic Games, Inc. All Rights Reserved.
+
+using System;
+using System.Collections.Generic;
+using System.Linq;
+using System.Threading.Tasks;
+using Horde.Build.Server;
+using Horde.Build.Streams;
+using Horde.Build.Utilities;
+using MongoDB.Bson;
+using MongoDB.Driver;
+
+namespace Horde.Build.Jobs.TestData
+{
+	using JobId = ObjectId<IJob>;
+	using StreamId = StringId<IStream>;
+	using TemplateRefId = StringId<TemplateRef>;
+
+	/// <summary>
+	/// Collection of test data documents
+	/// </summary>
+	public class TestDataCollection : ITestDataCollection
+	{
+		/// <summary>
+		/// Information about a test data document
+		/// </summary>
+		class TestDataDocument : ITestData
+		{
+			public ObjectId Id { get; set; }
+			public StreamId StreamId { get; set; }
+			public TemplateRefId TemplateRefId { get; set; }
+			public JobId JobId { get; set; }
+			public SubResourceId StepId { get; set; }
+			public int Change { get; set; }
+			public string Key { get; set; }
+			public BsonDocument Data { get; set; }
+
+			private TestDataDocument()
+			{
+				Key = String.Empty;
+				Data = new BsonDocument();
+			}
+
+			public TestDataDocument(IJob job, IJobStep jobStep, string key, BsonDocument value)
+			{
+				Id = ObjectId.GenerateNewId();
+				StreamId = job.StreamId;
+				TemplateRefId = job.TemplateId;
+				JobId = job.Id;
+				StepId = jobStep.Id;
+				Change = job.Change;
+				Key = key;
+				Data = value;
+			}
+		}
+
+		/// <summary>
+		/// The stream collection
+		/// </summary>
+		readonly IMongoCollection<TestDataDocument> _testDataDocuments;
+
+		/// <summary>
+		/// Constructor
+		/// </summary>
+		/// <param name="mongoService">The database service instance</param>
+		public TestDataCollection(MongoService mongoService)
+		{
+			List<MongoIndex<TestDataDocument>> indexes = new List<MongoIndex<TestDataDocument>>();
+			indexes.Add(keys => keys.Ascending(x => x.StreamId).Ascending(x => x.Change).Ascending(x => x.Key));
+			indexes.Add(keys => keys.Ascending(x => x.JobId).Ascending(x => x.StepId).Ascending(x => x.Key), unique: true);
+			_testDataDocuments = mongoService.GetCollection<TestDataDocument>("TestData", indexes);
+		}
+
+		/// <inheritdoc/>
+		public async Task<ITestData> AddAsync(IJob job, IJobStep step, string key, BsonDocument value)
+		{
+			TestDataDocument newTestData = new TestDataDocument(job, step, key, value);
+			await _testDataDocuments.InsertOneAsync(newTestData);
+			return newTestData;
+		}
+
+		/// <inheritdoc/>
+		public async Task<ITestData?> GetAsync(ObjectId id)
+		{
+			return await _testDataDocuments.Find<TestDataDocument>(x => x.Id == id).FirstOrDefaultAsync();
+		}
+
+		/// <inheritdoc/>
+		public async Task<List<ITestData>> FindAsync(StreamId? streamId, int? minChange, int? maxChange, JobId? jobId, SubResourceId? stepId, string? key = null, int index = 0, int count = 10)
+		{
+			FilterDefinition<TestDataDocument> filter = FilterDefinition<TestDataDocument>.Empty;
+			if (streamId != null)
+			{
+				filter &= Builders<TestDataDocument>.Filter.Eq(x => x.StreamId, streamId.Value);
+				if (minChange != null)
+				{
+					filter &= Builders<TestDataDocument>.Filter.Gte(x => x.Change, minChange.Value);
+				}
+				if (maxChange != null)
+				{
+					filter &= Builders<TestDataDocument>.Filter.Lte(x => x.Change, maxChange.Value);
+				}
+			}
+			if (jobId != null)
+			{
+				filter &= Builders<TestDataDocument>.Filter.Eq(x => x.JobId, jobId.Value);
+				if (stepId != null)
+				{
+					filter &= Builders<TestDataDocument>.Filter.Eq(x => x.StepId, stepId.Value);
+				}
+			}
+			if (key != null)
+			{
+				filter &= Builders<TestDataDocument>.Filter.Eq(x => x.Key, key);
+			}
+
+			SortDefinition<TestDataDocument> sort = Builders<TestDataDocument>.Sort.Ascending(x => x.StreamId).Descending(x => x.Change);
+
+			List<TestDataDocument> results = await _testDataDocuments.Find(filter).Sort(sort).Skip(index).Limit(count).ToListAsync();
+			return results.ConvertAll<ITestData>(x => x);
+		}
+
+		/// <inheritdoc/>
+		public async Task DeleteAsync(ObjectId id)
+		{
+			await _testDataDocuments.DeleteOneAsync<TestDataDocument>(x => x.Id == id);
+		}
+	}
+}
diff --git a/Engine/Source/Programs/Horde/Horde.Build/Jobs/TestData/TestDataController.Messages.cs b/Engine/Source/Programs/Horde/Horde.Build/Jobs/TestData/TestDataController.Messages.cs
--- a/Engine/Source/Programs/Horde/Horde.Build/Jobs/TestData/TestDataController.Messages.cs
+++ b/Engine/Source/Programs/Horde/Horde.Build/Jobs/TestData/TestDataController.Messages.cs
@@ -3,10 +3,10 @@
 using System;
 using System.Collections.Generic;
 using System.ComponentModel.DataAnnotations;
-using Horde.Build.Models;
+using Horde.Build.Jobs.TestData;
 using MongoDB.Bson.Serialization;
 
-namespace Horde.Build.Api
+namespace Horde.Build.Jobs.TestData
 {
 	/// <summary>
 	/// Response object describing test data to store
diff --git a/Engine/Source/Programs/Horde/Horde.Build/Jobs/TestData/TestDataController.cs b/Engine/Source/Programs/Horde/Horde.Build/Jobs/TestData/TestDataController.cs
--- a/Engine/Source/Programs/Horde/Horde.Build/Jobs/TestData/TestDataController.cs
+++ b/Engine/Source/Programs/Horde/Horde.Build/Jobs/TestData/TestDataController.cs
@@ -1,18 +1,15 @@
-﻿// Copyright Epic Games, Inc. All Rights Reserved.
+// Copyright Epic Games, Inc. All Rights Reserved.
 
 using System.Collections.Generic;
 using System.Threading.Tasks;
 using Horde.Build.Acls;
-using Horde.Build.Api;
-using Horde.Build.Collections;
-using Horde.Build.Models;
-using Horde.Build.Services;
+using Horde.Build.Streams;
 using Horde.Build.Utilities;
 using Microsoft.AspNetCore.Authorization;
 using Microsoft.AspNetCore.Mvc;
 using MongoDB.Bson;
 
-namespace Horde.Build.Controllers
+namespace Horde.Build.Jobs.TestData
 {
 	using JobId = ObjectId<IJob>;
 	using StreamId = StringId<IStream>;
diff --git a/Engine/Source/Programs/Horde/Horde.Build/Logs/Builder/LocalLogBuilder.cs b/Engine/Source/Programs/Horde/Horde.Build/Logs/Builder/LocalLogBuilder.cs
--- a/Engine/Source/Programs/Horde/Horde.Build/Logs/Builder/LocalLogBuilder.cs
+++ b/Engine/Source/Programs/Horde/Horde.Build/Logs/Builder/LocalLogBuilder.cs
@@ -5,8 +5,7 @@ using System.Collections.Concurrent;
 using System.Collections.Generic;
 using System.IO;
 using System.Threading.Tasks;
-using Horde.Build.Api;
-using Horde.Build.Models;
+using Horde.Build.Logs.Data;
 using Horde.Build.Utilities;
 
 namespace Horde.Build.Logs.Builder
diff --git a/Engine/Source/Programs/Horde/Horde.Build/Logs/Data/LogChunkData.cs b/Engine/Source/Programs/Horde/Horde.Build/Logs/Data/LogChunkData.cs
--- a/Engine/Source/Programs/Horde/Horde.Build/Logs/Data/LogChunkData.cs
+++ b/Engine/Source/Programs/Horde/Horde.Build/Logs/Data/LogChunkData.cs
@@ -4,10 +4,9 @@ using System;
 using System.Collections.Generic;
 using System.Linq;
 using EpicGames.Core;
-using Horde.Build.Api;
 using Microsoft.Extensions.Logging;
 
-namespace Horde.Build.Logs
+namespace Horde.Build.Logs.Data
 {
 	/// <summary>
 	/// Data for a log chunk
diff --git a/Engine/Source/Programs/Horde/Horde.Build/Logs/Data/LogIndexData.cs b/Engine/Source/Programs/Horde/Horde.Build/Logs/Data/LogIndexData.cs
--- a/Engine/Source/Programs/Horde/Horde.Build/Logs/Data/LogIndexData.cs
+++ b/Engine/Source/Programs/Horde/Horde.Build/Logs/Data/LogIndexData.cs
@@ -9,7 +9,7 @@ using Horde.Build.Utilities;
 using OpenTracing;
 using OpenTracing.Util;
 
-namespace Horde.Build.Logs
+namespace Horde.Build.Logs.Data
 {
 	/// <summary>
 	/// Contains a source block of text to be indexed
diff --git a/Engine/Source/Programs/Horde/Horde.Build/Logs/Data/LogSubChunkData.cs b/Engine/Source/Programs/Horde/Horde.Build/Logs/Data/LogSubChunkData.cs
--- a/Engine/Source/Programs/Horde/Horde.Build/Logs/Data/LogSubChunkData.cs
+++ b/Engine/Source/Programs/Horde/Horde.Build/Logs/Data/LogSubChunkData.cs
@@ -2,11 +2,10 @@
 
 using System;
 using EpicGames.Core;
-using Horde.Build.Api;
 using Horde.Build.Utilities;
 using Microsoft.Extensions.Logging;
 
-namespace Horde.Build.Logs
+namespace Horde.Build.Logs.Data
 {
 	/// <summary>
 	/// Pending data for a sub-chunk
diff --git a/Engine/Source/Programs/Horde/Horde.Build/Logs/Data/LogText.cs b/Engine/Source/Programs/Horde/Horde.Build/Logs/Data/LogText.cs
--- a/Engine/Source/Programs/Horde/Horde.Build/Logs/Data/LogText.cs
+++ b/Engine/Source/Programs/Horde/Horde.Build/Logs/Data/LogText.cs
@@ -9,7 +9,7 @@ using System.Text.Json;
 using EpicGames.Core;
 using Microsoft.Extensions.Logging;
 
-namespace Horde.Build.Logs
+namespace Horde.Build.Logs.Data
 {
 	/// <summary>
 	/// Exception thrown trying to parse log output
diff --git a/Engine/Source/Programs/Horde/Horde.Build/Logs/ILogEvent.cs b/Engine/Source/Programs/Horde/Horde.Build/Logs/ILogEvent.cs
--- a/Engine/Source/Programs/Horde/Horde.Build/Logs/ILogEvent.cs
+++ b/Engine/Source/Programs/Horde/Horde.Build/Logs/ILogEvent.cs
@@ -7,7 +7,7 @@ using HordeCommon;
 using Microsoft.Extensions.Logging;
 using MongoDB.Bson;
 
-namespace Horde.Build.Models
+namespace Horde.Build.Logs
 {
 	using LogId = ObjectId<ILogFile>;
 
diff --git a/Engine/Source/Programs/Horde/Horde.Build/Logs/ILogEventCollection.cs b/Engine/Source/Programs/Horde/Horde.Build/Logs/ILogEventCollection.cs
--- a/Engine/Source/Programs/Horde/Horde.Build/Logs/ILogEventCollection.cs
+++ b/Engine/Source/Programs/Horde/Horde.Build/Logs/ILogEventCollection.cs
@@ -2,11 +2,10 @@
 
 using System.Collections.Generic;
 using System.Threading.Tasks;
-using Horde.Build.Models;
 using Horde.Build.Utilities;
 using MongoDB.Bson;
 
-namespace Horde.Build.Collections
+namespace Horde.Build.Logs
 {
 	using LogId = ObjectId<ILogFile>;
 
diff --git a/Engine/Source/Programs/Horde/Horde.Build/Logs/ILogFile.cs b/Engine/Source/Programs/Horde/Horde.Build/Logs/ILogFile.cs
--- a/Engine/Source/Programs/Horde/Horde.Build/Logs/ILogFile.cs
+++ b/Engine/Source/Programs/Horde/Horde.Build/Logs/ILogFile.cs
@@ -3,10 +3,11 @@
 using System;
 using System.Collections.Generic;
 using EpicGames.Core;
-using Horde.Build.Api;
+using Horde.Build.Agents.Sessions;
+using Horde.Build.Jobs;
 using Horde.Build.Utilities;
 
-namespace Horde.Build.Models
+namespace Horde.Build.Logs
 {
 	using JobId = ObjectId<IJob>;
 	using LogId = ObjectId<ILogFile>;
diff --git a/Engine/Source/Programs/Horde/Horde.Build/Logs/ILogFileCollection.cs b/Engine/Source/Programs/Horde/Horde.Build/Logs/ILogFileCollection.cs
--- a/Engine/Source/Programs/Horde/Horde.Build/Logs/ILogFileCollection.cs
+++ b/Engine/Source/Programs/Horde/Horde.Build/Logs/ILogFileCollection.cs
@@ -2,11 +2,11 @@
 
 using System.Collections.Generic;
 using System.Threading.Tasks;
-using Horde.Build.Api;
-using Horde.Build.Models;
+using Horde.Build.Agents.Sessions;
+using Horde.Build.Jobs;
 using Horde.Build.Utilities;
 
-namespace Horde.Build.Collections
+namespace Horde.Build.Logs
 {
 	using JobId = ObjectId<IJob>;
 	using LogId = ObjectId<ILogFile>;
diff --git a/Engine/Source/Programs/Horde/Horde.Build/Logs/ILogStorage.cs b/Engine/Source/Programs/Horde/Horde.Build/Logs/ILogStorage.cs
--- a/Engine/Source/Programs/Horde/Horde.Build/Logs/ILogStorage.cs
+++ b/Engine/Source/Programs/Horde/Horde.Build/Logs/ILogStorage.cs
@@ -2,7 +2,7 @@
 
 using System;
 using System.Threading.Tasks;
-using Horde.Build.Models;
+using Horde.Build.Logs.Data;
 using Horde.Build.Utilities;
 
 namespace Horde.Build.Logs
diff --git a/Engine/Source/Programs/Horde/Horde.Build/Logs/LogEventCollection.cs b/Engine/Source/Programs/Horde/Horde.Build/Logs/LogEventCollection.cs
--- a/Engine/Source/Programs/Horde/Horde.Build/Logs/LogEventCollection.cs
+++ b/Engine/Source/Programs/Horde/Horde.Build/Logs/LogEventCollection.cs
@@ -4,7 +4,6 @@ using System;
 using System.Collections.Generic;
 using System.Linq;
 using System.Threading.Tasks;
-using Horde.Build.Models;
 using Horde.Build.Server;
 using Horde.Build.Utilities;
 using HordeCommon;
@@ -13,7 +12,7 @@ using MongoDB.Bson;
 using MongoDB.Bson.Serialization.Attributes;
 using MongoDB.Driver;
 
-namespace Horde.Build.Collections.Impl
+namespace Horde.Build.Logs
 {
 	using LogId = ObjectId<ILogFile>;
 
diff --git a/Engine/Source/Programs/Horde/Horde.Build/Logs/LogEventMessages.cs b/Engine/Source/Programs/Horde/Horde.Build/Logs/LogEventMessages.cs
--- a/Engine/Source/Programs/Horde/Horde.Build/Logs/LogEventMessages.cs
+++ b/Engine/Source/Programs/Horde/Horde.Build/Logs/LogEventMessages.cs
@@ -3,10 +3,9 @@
 using System.Collections.Generic;
 using System.Text.Json;
 using EpicGames.Core;
-using Horde.Build.Models;
 using HordeCommon;
 
-namespace Horde.Build.Api
+namespace Horde.Build.Logs
 {
 	/// <summary>
 	/// Information about an uploaded event
diff --git a/Engine/Source/Programs/Horde/Horde.Build/Logs/LogFileCollection.cs b/Engine/Source/Programs/Horde/Horde.Build/Logs/LogFileCollection.cs
--- a/Engine/Source/Programs/Horde/Horde.Build/Logs/LogFileCollection.cs
+++ b/Engine/Source/Programs/Horde/Horde.Build/Logs/LogFileCollection.cs
@@ -4,14 +4,14 @@ using System.Collections.Generic;
 using System.Linq;
 using System.Net;
 using System.Threading.Tasks;
-using Horde.Build.Api;
-using Horde.Build.Models;
+using Horde.Build.Agents.Sessions;
+using Horde.Build.Jobs;
 using Horde.Build.Server;
 using Horde.Build.Utilities;
 using MongoDB.Bson.Serialization.Attributes;
 using MongoDB.Driver;
 
-namespace Horde.Build.Collections.Impl
+namespace Horde.Build.Logs
 {
 	using JobId = ObjectId<IJob>;
 	using LogId = ObjectId<ILogFile>;
diff --git a/Engine/Source/Programs/Horde/Horde.Build/Logs/LogFileMessages.cs b/Engine/Source/Programs/Horde/Horde.Build/Logs/LogFileMessages.cs
--- a/Engine/Source/Programs/Horde/Horde.Build/Logs/LogFileMessages.cs
+++ b/Engine/Source/Programs/Horde/Horde.Build/Logs/LogFileMessages.cs
@@ -2,10 +2,9 @@
 
 using System.Collections.Generic;
 using Horde.Build.Logs;
-using Horde.Build.Models;
-using Horde.Build.Services;
+using Horde.Build.Logs.Data;
 
-namespace Horde.Build.Api
+namespace Horde.Build.Logs
 {
 	/// <summary>
 	/// The type of data stored in this log file
diff --git a/Engine/Source/Programs/Horde/Horde.Build/Logs/LogFileService.cs b/Engine/Source/Programs/Horde/Horde.Build/Logs/LogFileService.cs
--- a/Engine/Source/Programs/Horde/Horde.Build/Logs/LogFileService.cs
+++ b/Engine/Source/Programs/Horde/Horde.Build/Logs/LogFileService.cs
@@ -13,10 +13,9 @@ using System.Text.Json;
 using System.Threading;
 using System.Threading.Tasks;
 using EpicGames.Core;
-using Horde.Build.Api;
-using Horde.Build.Collections;
-using Horde.Build.Logs;
-using Horde.Build.Models;
+using Horde.Build.Agents.Sessions;
+using Horde.Build.Jobs;
+using Horde.Build.Logs.Data;
 using Horde.Build.Utilities;
 using HordeCommon;
 using Microsoft.Extensions.Caching.Memory;
@@ -28,7 +27,7 @@ using OpenTracing;
 using OpenTracing.Util;
 using Stream = System.IO.Stream;
 
-namespace Horde.Build.Services
+namespace Horde.Build.Logs
 {
 	using JobId = ObjectId<IJob>;
 	using LogId = ObjectId<ILogFile>;
diff --git a/Engine/Source/Programs/Horde/Horde.Build/Logs/LogsController.cs b/Engine/Source/Programs/Horde/Horde.Build/Logs/LogsController.cs
--- a/Engine/Source/Programs/Horde/Horde.Build/Logs/LogsController.cs
+++ b/Engine/Source/Programs/Horde/Horde.Build/Logs/LogsController.cs
@@ -7,17 +7,15 @@ using System.Security.Claims;
 using System.Text;
 using System.Threading.Tasks;
 using Horde.Build.Acls;
-using Horde.Build.Api;
-using Horde.Build.Collections;
-using Horde.Build.Logs;
-using Horde.Build.Models;
-using Horde.Build.Services;
+using Horde.Build.Issues;
+using Horde.Build.Jobs;
+using Horde.Build.Logs.Data;
 using Horde.Build.Utilities;
 using Microsoft.AspNetCore.Authorization;
 using Microsoft.AspNetCore.Mvc;
 using MongoDB.Bson;
 
-namespace Horde.Build.Controllers
+namespace Horde.Build.Logs
 {
 	using JobId = ObjectId<IJob>;
 	using LogId = ObjectId<ILogFile>;
diff --git a/Engine/Source/Programs/Horde/Horde.Build/Logs/Storage/LocalLogStorage.cs b/Engine/Source/Programs/Horde/Horde.Build/Logs/Storage/LocalLogStorage.cs
--- a/Engine/Source/Programs/Horde/Horde.Build/Logs/Storage/LocalLogStorage.cs
+++ b/Engine/Source/Programs/Horde/Horde.Build/Logs/Storage/LocalLogStorage.cs
@@ -2,7 +2,7 @@
 
 using System;
 using System.Threading.Tasks;
-using Horde.Build.Models;
+using Horde.Build.Logs.Data;
 using Horde.Build.Utilities;
 using Microsoft.Extensions.Caching.Memory;
 
diff --git a/Engine/Source/Programs/Horde/Horde.Build/Logs/Storage/NullLogStorage.cs b/Engine/Source/Programs/Horde/Horde.Build/Logs/Storage/NullLogStorage.cs
--- a/Engine/Source/Programs/Horde/Horde.Build/Logs/Storage/NullLogStorage.cs
+++ b/Engine/Source/Programs/Horde/Horde.Build/Logs/Storage/NullLogStorage.cs
@@ -1,10 +1,10 @@
 // Copyright Epic Games, Inc. All Rights Reserved.
 
 using System.Threading.Tasks;
-using Horde.Build.Models;
+using Horde.Build.Logs.Data;
 using Horde.Build.Utilities;
 
-namespace Horde.Build.Logs.Storage.Impl
+namespace Horde.Build.Logs.Storage
 {
 	using LogId = ObjectId<ILogFile>;
 
diff --git a/Engine/Source/Programs/Horde/Horde.Build/Logs/Storage/PersistentLogStorage.cs b/Engine/Source/Programs/Horde/Horde.Build/Logs/Storage/PersistentLogStorage.cs
--- a/Engine/Source/Programs/Horde/Horde.Build/Logs/Storage/PersistentLogStorage.cs
+++ b/Engine/Source/Programs/Horde/Horde.Build/Logs/Storage/PersistentLogStorage.cs
@@ -3,12 +3,12 @@
 using System;
 using System.Threading.Tasks;
 using EpicGames.Core;
-using Horde.Build.Models;
+using Horde.Build.Logs.Data;
 using Horde.Build.Storage;
 using Horde.Build.Utilities;
 using Microsoft.Extensions.Logging;
 
-namespace Horde.Build.Logs.Readers
+namespace Horde.Build.Logs.Storage
 {
 	using LogId = ObjectId<ILogFile>;
 
diff --git a/Engine/Source/Programs/Horde/Horde.Build/Logs/Storage/RedisLogStorage.cs b/Engine/Source/Programs/Horde/Horde.Build/Logs/Storage/RedisLogStorage.cs
--- a/Engine/Source/Programs/Horde/Horde.Build/Logs/Storage/RedisLogStorage.cs
+++ b/Engine/Source/Programs/Horde/Horde.Build/Logs/Storage/RedisLogStorage.cs
@@ -2,12 +2,12 @@
 
 using System;
 using System.Threading.Tasks;
-using Horde.Build.Models;
+using Horde.Build.Logs.Data;
 using Horde.Build.Utilities;
 using Microsoft.Extensions.Logging;
 using StackExchange.Redis;
 
-namespace Horde.Build.Logs.Storage.Impl
+namespace Horde.Build.Logs.Storage
 {
 	using LogId = ObjectId<ILogFile>;
 
diff --git a/Engine/Source/Programs/Horde/Horde.Build/Notifications/IEvent.cs b/Engine/Source/Programs/Horde/Horde.Build/Notifications/IEvent.cs
--- a/Engine/Source/Programs/Horde/Horde.Build/Notifications/IEvent.cs
+++ b/Engine/Source/Programs/Horde/Horde.Build/Notifications/IEvent.cs
@@ -1,9 +1,9 @@
 // Copyright Epic Games, Inc. All Rights Reserved.
 
-using Horde.Build.Api;
+using Horde.Build.Streams;
 using Horde.Build.Utilities;
 
-namespace Horde.Build.Models
+namespace Horde.Build.Notifications
 {
 	using StreamId = StringId<IStream>;
 	using TemplateRefId = StringId<TemplateRef>;
diff --git a/Engine/Source/Programs/Horde/Horde.Build/Notifications/INotificationService.cs b/Engine/Source/Programs/Horde/Horde.Build/Notifications/INotificationService.cs
--- a/Engine/Source/Programs/Horde/Horde.Build/Notifications/INotificationService.cs
+++ b/Engine/Source/Programs/Horde/Horde.Build/Notifications/INotificationService.cs
@@ -4,11 +4,13 @@ using System.Collections.Generic;
 using System.Diagnostics.CodeAnalysis;
 using System.Security.Claims;
 using System.Threading.Tasks;
-using Horde.Build.Api;
-using Horde.Build.Issues.Impl;
-using Horde.Build.Models;
-using Horde.Build.Server;
+using Horde.Build.Devices;
+using Horde.Build.Jobs;
+using Horde.Build.Jobs.Graphs;
+using Horde.Build.Issues;
+using Horde.Build.Users;
 using MongoDB.Bson;
+using Horde.Build.Streams;
 
 namespace Horde.Build.Notifications
 {
diff --git a/Engine/Source/Programs/Horde/Horde.Build/Notifications/INotificationSink.cs b/Engine/Source/Programs/Horde/Horde.Build/Notifications/INotificationSink.cs
--- a/Engine/Source/Programs/Horde/Horde.Build/Notifications/INotificationSink.cs
+++ b/Engine/Source/Programs/Horde/Horde.Build/Notifications/INotificationSink.cs
@@ -3,10 +3,13 @@
 using System;
 using System.Collections.Generic;
 using System.Threading.Tasks;
-using Horde.Build.Api;
-using Horde.Build.Issues.Impl;
-using Horde.Build.Models;
-using Horde.Build.Server;
+using Horde.Build.Devices;
+using Horde.Build.Issues;
+using Horde.Build.Jobs;
+using Horde.Build.Jobs.Graphs;
+using Horde.Build.Logs;
+using Horde.Build.Streams;
+using Horde.Build.Users;
 using HordeCommon;
 
 namespace Horde.Build.Notifications
diff --git a/Engine/Source/Programs/Horde/Horde.Build/Notifications/INotificationTrigger.cs b/Engine/Source/Programs/Horde/Horde.Build/Notifications/INotificationTrigger.cs
--- a/Engine/Source/Programs/Horde/Horde.Build/Notifications/INotificationTrigger.cs
+++ b/Engine/Source/Programs/Horde/Horde.Build/Notifications/INotificationTrigger.cs
@@ -1,10 +1,11 @@
 // Copyright Epic Games, Inc. All Rights Reserved.
 
 using System.Collections.Generic;
+using Horde.Build.Users;
 using Horde.Build.Utilities;
 using MongoDB.Bson;
 
-namespace Horde.Build.Models
+namespace Horde.Build.Notifications
 {
 	using UserId = ObjectId<IUser>;
 
diff --git a/Engine/Source/Programs/Horde/Horde.Build/Notifications/INotificationTriggerCollection.cs b/Engine/Source/Programs/Horde/Horde.Build/Notifications/INotificationTriggerCollection.cs
--- a/Engine/Source/Programs/Horde/Horde.Build/Notifications/INotificationTriggerCollection.cs
+++ b/Engine/Source/Programs/Horde/Horde.Build/Notifications/INotificationTriggerCollection.cs
@@ -2,11 +2,12 @@
 
 using System.Collections.Generic;
 using System.Threading.Tasks;
-using Horde.Build.Models;
+using Horde.Build.Notifications;
+using Horde.Build.Users;
 using Horde.Build.Utilities;
 using MongoDB.Bson;
 
-namespace Horde.Build.Collections
+namespace Horde.Build.Notifications
 {
 	using UserId = ObjectId<IUser>;
 
diff --git a/Engine/Source/Programs/Horde/Horde.Build/Notifications/ISubscription.cs b/Engine/Source/Programs/Horde/Horde.Build/Notifications/ISubscription.cs
--- a/Engine/Source/Programs/Horde/Horde.Build/Notifications/ISubscription.cs
+++ b/Engine/Source/Programs/Horde/Horde.Build/Notifications/ISubscription.cs
@@ -1,8 +1,9 @@
 // Copyright Epic Games, Inc. All Rights Reserved.
 
+using Horde.Build.Users;
 using Horde.Build.Utilities;
 
-namespace Horde.Build.Models
+namespace Horde.Build.Notifications
 {
 	using UserId = ObjectId<IUser>;
 
diff --git a/Engine/Source/Programs/Horde/Horde.Build/Notifications/ISubscriptionCollection.cs b/Engine/Source/Programs/Horde/Horde.Build/Notifications/ISubscriptionCollection.cs
--- a/Engine/Source/Programs/Horde/Horde.Build/Notifications/ISubscriptionCollection.cs
+++ b/Engine/Source/Programs/Horde/Horde.Build/Notifications/ISubscriptionCollection.cs
@@ -2,11 +2,10 @@
 
 using System.Collections.Generic;
 using System.Threading.Tasks;
-using Horde.Build.Api;
-using Horde.Build.Models;
+using Horde.Build.Users;
 using Horde.Build.Utilities;
 
-namespace Horde.Build.Collections
+namespace Horde.Build.Notifications
 {
 	using UserId = ObjectId<IUser>;
 
diff --git a/Engine/Source/Programs/Horde/Horde.Build/Notifications/NotificationMessages.cs b/Engine/Source/Programs/Horde/Horde.Build/Notifications/NotificationMessages.cs
--- a/Engine/Source/Programs/Horde/Horde.Build/Notifications/NotificationMessages.cs
+++ b/Engine/Source/Programs/Horde/Horde.Build/Notifications/NotificationMessages.cs
@@ -1,8 +1,6 @@
 // Copyright Epic Games, Inc. All Rights Reserved.
 
-using Horde.Build.Models;
-
-namespace Horde.Build.Api
+namespace Horde.Build.Notifications
 {
 	/// <summary>
 	/// A request from a user to be notified when something happens.
diff --git a/Engine/Source/Programs/Horde/Horde.Build/Notifications/NotificationService.cs b/Engine/Source/Programs/Horde/Horde.Build/Notifications/NotificationService.cs
--- a/Engine/Source/Programs/Horde/Horde.Build/Notifications/NotificationService.cs
+++ b/Engine/Source/Programs/Horde/Horde.Build/Notifications/NotificationService.cs
@@ -8,12 +8,15 @@ using System.Security.Claims;
 using System.Threading;
 using System.Threading.Tasks;
 using EpicGames.Core;
-using Horde.Build.Api;
-using Horde.Build.Collections;
-using Horde.Build.Issues.Impl;
-using Horde.Build.Models;
+using Horde.Build.Agents.Pools;
+using Horde.Build.Devices;
+using Horde.Build.Issues;
+using Horde.Build.Jobs;
+using Horde.Build.Jobs.Graphs;
+using Horde.Build.Logs;
 using Horde.Build.Server;
-using Horde.Build.Services;
+using Horde.Build.Streams;
+using Horde.Build.Users;
 using Horde.Build.Utilities;
 using HordeCommon;
 using Microsoft.Extensions.Hosting;
@@ -25,7 +28,7 @@ using StackExchange.Redis;
 using StatsdClient;
 using JsonSerializer = System.Text.Json.JsonSerializer;
 
-namespace Horde.Build.Notifications.Impl
+namespace Horde.Build.Notifications
 {
 	using UserId = ObjectId<IUser>;
 
diff --git a/Engine/Source/Programs/Horde/Horde.Build/Notifications/NotificationTriggerCollection.cs b/Engine/Source/Programs/Horde/Horde.Build/Notifications/NotificationTriggerCollection.cs
--- a/Engine/Source/Programs/Horde/Horde.Build/Notifications/NotificationTriggerCollection.cs
+++ b/Engine/Source/Programs/Horde/Horde.Build/Notifications/NotificationTriggerCollection.cs
@@ -3,14 +3,14 @@
 using System.Collections.Generic;
 using System.Linq;
 using System.Threading.Tasks;
-using Horde.Build.Models;
 using Horde.Build.Server;
+using Horde.Build.Users;
 using Horde.Build.Utilities;
 using MongoDB.Bson;
 using MongoDB.Bson.Serialization.Attributes;
 using MongoDB.Driver;
 
-namespace Horde.Build.Collections.Impl
+namespace Horde.Build.Notifications
 {
 	using UserId = ObjectId<IUser>;
 
diff --git a/Engine/Source/Programs/Horde/Horde.Build/Notifications/Sinks/SlackNotificationSink.cs b/Engine/Source/Programs/Horde/Horde.Build/Notifications/Sinks/SlackNotificationSink.cs
--- a/Engine/Source/Programs/Horde/Horde.Build/Notifications/Sinks/SlackNotificationSink.cs
+++ b/Engine/Source/Programs/Horde/Horde.Build/Notifications/Sinks/SlackNotificationSink.cs
@@ -15,13 +15,17 @@ using System.Text.RegularExpressions;
 using System.Threading;
 using System.Threading.Tasks;
 using EpicGames.Core;
-using Horde.Build.Api;
-using Horde.Build.Collections;
-using Horde.Build.Config;
-using Horde.Build.Issues.Impl;
-using Horde.Build.Models;
+using Horde.Build.Configuration;
+using Horde.Build.Devices;
+using Horde.Build.Issues;
+using Horde.Build.Issues.External;
+using Horde.Build.Jobs;
+using Horde.Build.Jobs.Graphs;
+using Horde.Build.Logs;
+using Horde.Build.Perforce;
 using Horde.Build.Server;
-using Horde.Build.Services;
+using Horde.Build.Streams;
+using Horde.Build.Users;
 using Horde.Build.Utilities;
 using Horde.Build.Utilities.Slack.BlockKit;
 using HordeCommon;
@@ -34,7 +38,7 @@ using MongoDB.Bson;
 using MongoDB.Bson.Serialization.Attributes;
 using MongoDB.Driver;
 
-namespace Horde.Build.Notifications.Impl
+namespace Horde.Build.Notifications.Sinks
 {
 	using JobId = ObjectId<IJob>;
 	using LogId = ObjectId<ILogFile>;
diff --git a/Engine/Source/Programs/Horde/Horde.Build/Notifications/SubscriptionCollection.cs b/Engine/Source/Programs/Horde/Horde.Build/Notifications/SubscriptionCollection.cs
--- a/Engine/Source/Programs/Horde/Horde.Build/Notifications/SubscriptionCollection.cs
+++ b/Engine/Source/Programs/Horde/Horde.Build/Notifications/SubscriptionCollection.cs
@@ -7,15 +7,16 @@ using System.Linq;
 using System.Text;
 using System.Threading.Tasks;
 using EpicGames.Core;
-using Horde.Build.Api;
-using Horde.Build.Models;
+using Horde.Build.Jobs;
 using Horde.Build.Server;
+using Horde.Build.Streams;
+using Horde.Build.Users;
 using Horde.Build.Utilities;
 using HordeCommon;
 using MongoDB.Bson.Serialization.Attributes;
 using MongoDB.Driver;
 
-namespace Horde.Build.Collections.Impl
+namespace Horde.Build.Notifications
 {
 	using StreamId = StringId<IStream>;
 	using TemplateRefId = StringId<TemplateRef>;
diff --git a/Engine/Source/Programs/Horde/Horde.Build/Notifications/SubscriptionMessages.cs b/Engine/Source/Programs/Horde/Horde.Build/Notifications/SubscriptionMessages.cs
--- a/Engine/Source/Programs/Horde/Horde.Build/Notifications/SubscriptionMessages.cs
+++ b/Engine/Source/Programs/Horde/Horde.Build/Notifications/SubscriptionMessages.cs
@@ -3,11 +3,12 @@
 using System;
 using System.ComponentModel.DataAnnotations;
 using EpicGames.Core;
-using Horde.Build.Models;
+using Horde.Build.Jobs;
+using Horde.Build.Streams;
 using Horde.Build.Utilities;
 using HordeCommon;
 
-namespace Horde.Build.Api
+namespace Horde.Build.Notifications
 {
 	using StreamId = StringId<IStream>;
 	using TemplateRefId = StringId<TemplateRef>;
diff --git a/Engine/Source/Programs/Horde/Horde.Build/Notifications/SubscriptionsController.cs b/Engine/Source/Programs/Horde/Horde.Build/Notifications/SubscriptionsController.cs
--- a/Engine/Source/Programs/Horde/Horde.Build/Notifications/SubscriptionsController.cs
+++ b/Engine/Source/Programs/Horde/Horde.Build/Notifications/SubscriptionsController.cs
@@ -4,14 +4,12 @@ using System;
 using System.Collections.Generic;
 using System.Threading.Tasks;
 using Horde.Build.Acls;
-using Horde.Build.Api;
-using Horde.Build.Collections;
-using Horde.Build.Models;
+using Horde.Build.Users;
 using Horde.Build.Utilities;
 using Microsoft.AspNetCore.Authorization;
 using Microsoft.AspNetCore.Mvc;
 
-namespace Horde.Build.Controllers
+namespace Horde.Build.Notifications
 {
 	using UserId = ObjectId<IUser>;
 
diff --git a/Engine/Source/Programs/Horde/Horde.Build/Perforce/Change.cs b/Engine/Source/Programs/Horde/Horde.Build/Perforce/Change.cs
--- a/Engine/Source/Programs/Horde/Horde.Build/Perforce/Change.cs
+++ b/Engine/Source/Programs/Horde/Horde.Build/Perforce/Change.cs
@@ -8,8 +8,9 @@
 using System.Linq;
 using System.Text.RegularExpressions;
 using EpicGames.Core;
+using Horde.Build.Users;
 
-namespace Horde.Build.Models
+namespace Horde.Build.Perforce
 {
 	/// <summary>
 	/// Summary information for a change
diff --git a/Engine/Source/Programs/Horde/Horde.Build/Perforce/CommitCollection.cs b/Engine/Source/Programs/Horde/Horde.Build/Perforce/CommitCollection.cs
--- a/Engine/Source/Programs/Horde/Horde.Build/Perforce/CommitCollection.cs
+++ b/Engine/Source/Programs/Horde/Horde.Build/Perforce/CommitCollection.cs
@@ -4,14 +4,15 @@ using System;
 using System.Collections.Generic;
 using System.Linq;
 using System.Threading.Tasks;
-using Horde.Build.Models;
 using Horde.Build.Server;
+using Horde.Build.Streams;
+using Horde.Build.Users;
 using Horde.Build.Utilities;
 using MongoDB.Bson;
 using MongoDB.Bson.Serialization.Attributes;
 using MongoDB.Driver;
 
-namespace Horde.Build.Commits.Impl
+namespace Horde.Build.Perforce
 {
 	using CommitId = ObjectId<ICommit>;
 	using StreamId = StringId<IStream>;
diff --git a/Engine/Source/Programs/Horde/Horde.Build/Perforce/CommitService.cs b/Engine/Source/Programs/Horde/Horde.Build/Perforce/CommitService.cs
--- a/Engine/Source/Programs/Horde/Horde.Build/Perforce/CommitService.cs
+++ b/Engine/Source/Programs/Horde/Horde.Build/Perforce/CommitService.cs
@@ -8,7 +8,6 @@ using System.Linq;
 using System.Text;
 using System.Text.RegularExpressions;
 using System.Threading;
-using System.Threading.Channels;
 using System.Threading.Tasks;
 using EpicGames.Core;
 using EpicGames.Horde.Bundles;
@@ -17,10 +16,9 @@ using EpicGames.Horde.Storage;
 using EpicGames.Perforce;
 using EpicGames.Redis;
 using EpicGames.Serialization;
-using Horde.Build.Collections;
-using Horde.Build.Models;
 using Horde.Build.Server;
-using Horde.Build.Services;
+using Horde.Build.Streams;
+using Horde.Build.Users;
 using Horde.Build.Utilities;
 using HordeCommon;
 using Microsoft.Extensions.Caching.Memory;
@@ -29,7 +27,7 @@ using Microsoft.Extensions.Logging;
 using Microsoft.Extensions.Options;
 using StackExchange.Redis;
 
-namespace Horde.Build.Commits.Impl
+namespace Horde.Build.Perforce
 {
 	using CommitId = ObjectId<ICommit>;
 	using StreamId = StringId<IStream>;
diff --git a/Engine/Source/Programs/Horde/Horde.Build/Perforce/ICommit.cs b/Engine/Source/Programs/Horde/Horde.Build/Perforce/ICommit.cs
--- a/Engine/Source/Programs/Horde/Horde.Build/Perforce/ICommit.cs
+++ b/Engine/Source/Programs/Horde/Horde.Build/Perforce/ICommit.cs
@@ -1,10 +1,11 @@
 ﻿// Copyright Epic Games, Inc. All Rights Reserved.
 
 using System;
-using Horde.Build.Models;
+using Horde.Build.Streams;
+using Horde.Build.Users;
 using Horde.Build.Utilities;
 
-namespace Horde.Build.Commits
+namespace Horde.Build.Perforce
 {
 	using CommitId = ObjectId<ICommit>;
 	using StreamId = StringId<IStream>;
diff --git a/Engine/Source/Programs/Horde/Horde.Build/Perforce/ICommitCollection.cs b/Engine/Source/Programs/Horde/Horde.Build/Perforce/ICommitCollection.cs
--- a/Engine/Source/Programs/Horde/Horde.Build/Perforce/ICommitCollection.cs
+++ b/Engine/Source/Programs/Horde/Horde.Build/Perforce/ICommitCollection.cs
@@ -2,10 +2,10 @@
 
 using System.Collections.Generic;
 using System.Threading.Tasks;
-using Horde.Build.Models;
+using Horde.Build.Streams;
 using Horde.Build.Utilities;
 
-namespace Horde.Build.Commits
+namespace Horde.Build.Perforce
 {
 	using CommitId = ObjectId<ICommit>;
 	using StreamId = StringId<IStream>;
diff --git a/Engine/Source/Programs/Horde/Horde.Build/Perforce/ICommitService.cs b/Engine/Source/Programs/Horde/Horde.Build/Perforce/ICommitService.cs
--- a/Engine/Source/Programs/Horde/Horde.Build/Perforce/ICommitService.cs
+++ b/Engine/Source/Programs/Horde/Horde.Build/Perforce/ICommitService.cs
@@ -2,10 +2,10 @@
 
 using System;
 using EpicGames.Horde.Storage;
-using Horde.Build.Models;
+using Horde.Build.Streams;
 using Horde.Build.Utilities;
 
-namespace Horde.Build.Commits
+namespace Horde.Build.Perforce
 {
 	using StreamId = StringId<IStream>;
 
diff --git a/Engine/Source/Programs/Horde/Horde.Build/Perforce/IPerforceService.cs b/Engine/Source/Programs/Horde/Horde.Build/Perforce/IPerforceService.cs
--- a/Engine/Source/Programs/Horde/Horde.Build/Perforce/IPerforceService.cs
+++ b/Engine/Source/Programs/Horde/Horde.Build/Perforce/IPerforceService.cs
@@ -6,9 +6,11 @@ using System.Diagnostics.CodeAnalysis;
 using System.Text.RegularExpressions;
 using System.Threading.Tasks;
 using EpicGames.Perforce;
-using Horde.Build.Models;
+using Horde.Build.Jobs.Templates;
+using Horde.Build.Streams;
+using Horde.Build.Users;
 
-namespace Horde.Build.Services
+namespace Horde.Build.Perforce
 {
 	/// <summary>
 	/// Information about a Perforce user
diff --git a/Engine/Source/Programs/Horde/Horde.Build/Perforce/PerforceController.cs b/Engine/Source/Programs/Horde/Horde.Build/Perforce/PerforceController.cs
--- a/Engine/Source/Programs/Horde/Horde.Build/Perforce/PerforceController.cs
+++ b/Engine/Source/Programs/Horde/Horde.Build/Perforce/PerforceController.cs
@@ -4,14 +4,12 @@ using System.Collections.Generic;
 using System.Net;
 using System.Threading.Tasks;
 using Horde.Build.Acls;
-using Horde.Build.Models;
 using Horde.Build.Server;
-using Horde.Build.Services;
 using Microsoft.AspNetCore.Authorization;
 using Microsoft.AspNetCore.Mvc;
 using Microsoft.Extensions.Logging;
 
-namespace Horde.Build.Controllers
+namespace Horde.Build.Perforce
 {
 	/// <summary>
 	/// Controller for Perforce triggers and callbacks
diff --git a/Engine/Source/Programs/Horde/Horde.Build/Perforce/PerforceLoadBalancer.cs b/Engine/Source/Programs/Horde/Horde.Build/Perforce/PerforceLoadBalancer.cs
--- a/Engine/Source/Programs/Horde/Horde.Build/Perforce/PerforceLoadBalancer.cs
+++ b/Engine/Source/Programs/Horde/Horde.Build/Perforce/PerforceLoadBalancer.cs
@@ -10,8 +10,8 @@ using System.Text.Json;
 using System.Threading;
 using System.Threading.Tasks;
 using Google.Protobuf.WellKnownTypes;
-using Horde.Build.Collections;
-using Horde.Build.Models;
+using Horde.Build.Agents;
+using Horde.Build.Agents.Leases;
 using Horde.Build.Server;
 using Horde.Build.Utilities;
 using HordeCommon;
@@ -21,7 +21,7 @@ using Microsoft.Extensions.Logging;
 using MongoDB.Bson.Serialization.Attributes;
 using MongoDB.Driver;
 
-namespace Horde.Build.Services
+namespace Horde.Build.Perforce
 {
 	/// <summary>
 	/// Health of a particular server
diff --git a/Engine/Source/Programs/Horde/Horde.Build/Perforce/PerforceService.cs b/Engine/Source/Programs/Horde/Horde.Build/Perforce/PerforceService.cs
--- a/Engine/Source/Programs/Horde/Horde.Build/Perforce/PerforceService.cs
+++ b/Engine/Source/Programs/Horde/Horde.Build/Perforce/PerforceService.cs
@@ -14,9 +14,8 @@ using System.Text.RegularExpressions;
 using System.Threading.Tasks;
 using EpicGames.Core;
 using EpicGames.Perforce;
-using Horde.Build.Collections;
-using Horde.Build.Models;
 using Horde.Build.Server;
+using Horde.Build.Users;
 using Horde.Build.Utilities;
 using Microsoft.Extensions.Caching.Memory;
 using Microsoft.Extensions.Logging;
@@ -24,9 +23,9 @@ using Microsoft.Extensions.Options;
 using OpenTracing;
 using OpenTracing.Util;
 
-namespace Horde.Build.Services
+namespace Horde.Build.Perforce
 {
-	using P4 = Perforce.P4;
+	using P4 = global::Perforce.P4;
 
 	static class PerforceExtensions
 	{
diff --git a/Engine/Source/Programs/Horde/Horde.Build/Projects/IProject.cs b/Engine/Source/Programs/Horde/Horde.Build/Projects/IProject.cs
--- a/Engine/Source/Programs/Horde/Horde.Build/Projects/IProject.cs
+++ b/Engine/Source/Programs/Horde/Horde.Build/Projects/IProject.cs
@@ -1,16 +1,17 @@
-﻿// Copyright Epic Games, Inc. All Rights Reserved.
+// Copyright Epic Games, Inc. All Rights Reserved.
 
 using System.Collections.Generic;
 using System.Linq;
 using System.Text.RegularExpressions;
 using EpicGames.Core;
 using Horde.Build.Acls;
-using Horde.Build.Api;
+using Horde.Build.Projects;
+using Horde.Build.Streams;
 using Horde.Build.Utilities;
 using MongoDB.Bson.Serialization.Attributes;
 using MongoDB.Driver;
 
-namespace Horde.Build.Models
+namespace Horde.Build.Projects
 {
 	using ProjectId = StringId<IProject>;
 
