Horde: Move files into namespaces corresponding to their location on disk.

diff --git a/Engine/Source/Programs/Horde/Horde.Build/Projects/IProjectCollection.cs b/Engine/Source/Programs/Horde/Horde.Build/Projects/IProjectCollection.cs
--- a/Engine/Source/Programs/Horde/Horde.Build/Projects/IProjectCollection.cs
+++ b/Engine/Source/Programs/Horde/Horde.Build/Projects/IProjectCollection.cs
@@ -2,11 +2,10 @@
 
 using System.Collections.Generic;
 using System.Threading.Tasks;
-using Horde.Build.Models;
 using Horde.Build.Server;
 using Horde.Build.Utilities;
 
-namespace Horde.Build.Collections
+namespace Horde.Build.Projects
 {
 	using ProjectId = StringId<IProject>;
 
diff --git a/Engine/Source/Programs/Horde/Horde.Build/Projects/ProjectCollection.cs b/Engine/Source/Programs/Horde/Horde.Build/Projects/ProjectCollection.cs
--- a/Engine/Source/Programs/Horde/Horde.Build/Projects/ProjectCollection.cs
+++ b/Engine/Source/Programs/Horde/Horde.Build/Projects/ProjectCollection.cs
@@ -6,13 +6,12 @@ using System.Diagnostics.CodeAnalysis;
 using System.Linq;
 using System.Threading.Tasks;
 using Horde.Build.Acls;
-using Horde.Build.Models;
 using Horde.Build.Server;
 using Horde.Build.Utilities;
 using MongoDB.Bson.Serialization.Attributes;
 using MongoDB.Driver;
 
-namespace Horde.Build.Collections.Impl
+namespace Horde.Build.Projects
 {
 	using ProjectId = StringId<IProject>;
 
diff --git a/Engine/Source/Programs/Horde/Horde.Build/Projects/ProjectService.cs b/Engine/Source/Programs/Horde/Horde.Build/Projects/ProjectService.cs
--- a/Engine/Source/Programs/Horde/Horde.Build/Projects/ProjectService.cs
+++ b/Engine/Source/Programs/Horde/Horde.Build/Projects/ProjectService.cs
@@ -1,14 +1,12 @@
-﻿// Copyright Epic Games, Inc. All Rights Reserved.
+// Copyright Epic Games, Inc. All Rights Reserved.
 
 using System.Collections.Generic;
 using System.Security.Claims;
 using System.Threading.Tasks;
 using Horde.Build.Acls;
-using Horde.Build.Collections;
-using Horde.Build.Models;
 using Horde.Build.Utilities;
 
-namespace Horde.Build.Services
+namespace Horde.Build.Projects
 {
 	using ProjectId = StringId<IProject>;
 
diff --git a/Engine/Source/Programs/Horde/Horde.Build/Projects/ProjectsController.Messages.cs b/Engine/Source/Programs/Horde/Horde.Build/Projects/ProjectsController.Messages.cs
--- a/Engine/Source/Programs/Horde/Horde.Build/Projects/ProjectsController.Messages.cs
+++ b/Engine/Source/Programs/Horde/Horde.Build/Projects/ProjectsController.Messages.cs
@@ -2,9 +2,10 @@
 
 using System.Collections.Generic;
 using System.ComponentModel.DataAnnotations;
-using Horde.Build.Models;
+using Horde.Build.Acls;
+using Horde.Build.Projects;
 
-namespace Horde.Build.Api
+namespace Horde.Build.Projects
 {
 	/// <summary>
 	/// Information about a category to display for a stream
diff --git a/Engine/Source/Programs/Horde/Horde.Build/Projects/ProjectsController.cs b/Engine/Source/Programs/Horde/Horde.Build/Projects/ProjectsController.cs
--- a/Engine/Source/Programs/Horde/Horde.Build/Projects/ProjectsController.cs
+++ b/Engine/Source/Programs/Horde/Horde.Build/Projects/ProjectsController.cs
@@ -1,16 +1,14 @@
-﻿// Copyright Epic Games, Inc. All Rights Reserved.
+// Copyright Epic Games, Inc. All Rights Reserved.
 
 using System.Collections.Generic;
 using System.Threading.Tasks;
 using Horde.Build.Acls;
-using Horde.Build.Api;
-using Horde.Build.Models;
-using Horde.Build.Services;
+using Horde.Build.Streams;
 using Horde.Build.Utilities;
 using Microsoft.AspNetCore.Authorization;
 using Microsoft.AspNetCore.Mvc;
 
-namespace Horde.Build.Controllers
+namespace Horde.Build.Projects
 {
 	using ProjectId = StringId<IProject>;
 
diff --git a/Engine/Source/Programs/Horde/Horde.Build/Secrets/Credential.cs b/Engine/Source/Programs/Horde/Horde.Build/Secrets/Credential.cs
--- a/Engine/Source/Programs/Horde/Horde.Build/Secrets/Credential.cs
+++ b/Engine/Source/Programs/Horde/Horde.Build/Secrets/Credential.cs
@@ -1,11 +1,11 @@
-﻿// Copyright Epic Games, Inc. All Rights Reserved.
+// Copyright Epic Games, Inc. All Rights Reserved.
 
 using System.Collections.Generic;
 using Horde.Build.Acls;
 using MongoDB.Bson;
 using MongoDB.Bson.Serialization.Attributes;
 
-namespace Horde.Build.Models
+namespace Horde.Build.Secrets
 {
 	/// <summary>
 	/// Stores information about a credential
diff --git a/Engine/Source/Programs/Horde/Horde.Build/Secrets/CredentialService.cs b/Engine/Source/Programs/Horde/Horde.Build/Secrets/CredentialService.cs
--- a/Engine/Source/Programs/Horde/Horde.Build/Secrets/CredentialService.cs
+++ b/Engine/Source/Programs/Horde/Horde.Build/Secrets/CredentialService.cs
@@ -5,12 +5,11 @@ using System.Linq;
 using System.Security.Claims;
 using System.Threading.Tasks;
 using Horde.Build.Acls;
-using Horde.Build.Models;
 using Horde.Build.Server;
 using MongoDB.Bson;
 using MongoDB.Driver;
 
-namespace Horde.Build.Services
+namespace Horde.Build.Secrets
 {
 	/// <summary>
 	/// Service for accessing and modifying credentials
diff --git a/Engine/Source/Programs/Horde/Horde.Build/Secrets/CredentialsController.Messages.cs b/Engine/Source/Programs/Horde/Horde.Build/Secrets/CredentialsController.Messages.cs
--- a/Engine/Source/Programs/Horde/Horde.Build/Secrets/CredentialsController.Messages.cs
+++ b/Engine/Source/Programs/Horde/Horde.Build/Secrets/CredentialsController.Messages.cs
@@ -2,9 +2,10 @@
 
 using System.Collections.Generic;
 using System.ComponentModel.DataAnnotations;
-using Horde.Build.Models;
+using Horde.Build.Acls;
+using Horde.Build.Secrets;
 
-namespace Horde.Build.Api
+namespace Horde.Build.Credentials
 {
 	/// <summary>
 	/// Parameters to create a new credential
diff --git a/Engine/Source/Programs/Horde/Horde.Build/Secrets/CredentialsController.cs b/Engine/Source/Programs/Horde/Horde.Build/Secrets/CredentialsController.cs
--- a/Engine/Source/Programs/Horde/Horde.Build/Secrets/CredentialsController.cs
+++ b/Engine/Source/Programs/Horde/Horde.Build/Secrets/CredentialsController.cs
@@ -1,17 +1,15 @@
-﻿// Copyright Epic Games, Inc. All Rights Reserved.
+// Copyright Epic Games, Inc. All Rights Reserved.
 
 using System.Collections.Generic;
 using System.Threading.Tasks;
 using Horde.Build.Acls;
-using Horde.Build.Api;
-using Horde.Build.Models;
-using Horde.Build.Services;
+using Horde.Build.Credentials;
 using Horde.Build.Utilities;
 using Microsoft.AspNetCore.Authorization;
 using Microsoft.AspNetCore.Mvc;
 using MongoDB.Bson;
 
-namespace Horde.Build.Controllers
+namespace Horde.Build.Secrets
 {
 	/// <summary>
 	/// Controller for the /api/v1/credentials endpoint
diff --git a/Engine/Source/Programs/Horde/Horde.Build/Server/AccountController.cs b/Engine/Source/Programs/Horde/Horde.Build/Server/AccountController.cs
--- a/Engine/Source/Programs/Horde/Horde.Build/Server/AccountController.cs
+++ b/Engine/Source/Programs/Horde/Horde.Build/Server/AccountController.cs
@@ -1,4 +1,4 @@
-﻿// Copyright Epic Games, Inc. All Rights Reserved.
+// Copyright Epic Games, Inc. All Rights Reserved.
 
 using System;
 using System.Globalization;
@@ -13,7 +13,7 @@ using Microsoft.AspNetCore.Authentication.OpenIdConnect;
 using Microsoft.AspNetCore.Mvc;
 using Microsoft.Extensions.Options;
 
-namespace Horde.Build.Controllers
+namespace Horde.Build.Server
 {
 	/// <summary>
 	/// Controller managing account status
diff --git a/Engine/Source/Programs/Horde/Horde.Build/Server/AdminController.cs b/Engine/Source/Programs/Horde/Horde.Build/Server/AdminController.cs
--- a/Engine/Source/Programs/Horde/Horde.Build/Server/AdminController.cs
+++ b/Engine/Source/Programs/Horde/Horde.Build/Server/AdminController.cs
@@ -6,13 +6,13 @@ using System.Linq;
 using System.Security.Claims;
 using System.Threading.Tasks;
 using Horde.Build.Acls;
-using Horde.Build.Services;
+using Horde.Build.Agents.Software;
 using Horde.Build.Utilities;
 using Microsoft.AspNetCore.Authorization;
 using Microsoft.AspNetCore.Mvc;
 using Microsoft.Extensions.Options;
 
-namespace Horde.Build.Controllers
+namespace Horde.Build.Server
 {
 	/// <summary>
 	/// Object containing settings for the server
diff --git a/Engine/Source/Programs/Horde/Horde.Build/Server/Config.cs b/Engine/Source/Programs/Horde/Horde.Build/Server/Config.cs
--- a/Engine/Source/Programs/Horde/Horde.Build/Server/Config.cs
+++ b/Engine/Source/Programs/Horde/Horde.Build/Server/Config.cs
@@ -8,8 +8,10 @@ using System.Linq;
 using EpicGames.Core;
 using EpicGames.Horde.Common;
 using EpicGames.Serialization;
-using Horde.Build.Api;
-using Horde.Build.Models;
+using Horde.Build.Acls;
+using Horde.Build.Jobs.Graphs;
+using Horde.Build.Projects;
+using Horde.Build.Streams;
 using Horde.Build.Tools;
 using Horde.Build.Utilities;
 
diff --git a/Engine/Source/Programs/Horde/Horde.Build/Server/ConsistencyService.cs b/Engine/Source/Programs/Horde/Horde.Build/Server/ConsistencyService.cs
--- a/Engine/Source/Programs/Horde/Horde.Build/Server/ConsistencyService.cs
+++ b/Engine/Source/Programs/Horde/Horde.Build/Server/ConsistencyService.cs
@@ -5,14 +5,14 @@ using System.Collections.Generic;
 using System.Linq;
 using System.Threading;
 using System.Threading.Tasks;
-using Horde.Build.Collections;
-using Horde.Build.Models;
+using Horde.Build.Agents.Sessions;
+using Horde.Build.Agents.Leases;
 using Horde.Build.Utilities;
 using HordeCommon;
 using Microsoft.Extensions.Hosting;
 using Microsoft.Extensions.Logging;
 
-namespace Horde.Build.Services
+namespace Horde.Build.Server
 {
 	using SessionId = ObjectId<ISession>;
 
diff --git a/Engine/Source/Programs/Horde/Horde.Build/Server/DashboardController.Messages.cs b/Engine/Source/Programs/Horde/Horde.Build/Server/DashboardController.Messages.cs
--- a/Engine/Source/Programs/Horde/Horde.Build/Server/DashboardController.Messages.cs
+++ b/Engine/Source/Programs/Horde/Horde.Build/Server/DashboardController.Messages.cs
@@ -1,6 +1,6 @@
 // Copyright Epic Games, Inc. All Rights Reserved.
 
-namespace Horde.Build.Api
+namespace Horde.Build.Server
 {
 	/// <summary>
 	/// Setting information required by dashboard
diff --git a/Engine/Source/Programs/Horde/Horde.Build/Server/DashboardController.cs b/Engine/Source/Programs/Horde/Horde.Build/Server/DashboardController.cs
--- a/Engine/Source/Programs/Horde/Horde.Build/Server/DashboardController.cs
+++ b/Engine/Source/Programs/Horde/Horde.Build/Server/DashboardController.cs
@@ -1,11 +1,8 @@
 // Copyright Epic Games, Inc. All Rights Reserved.
 
-// Copyright Epic Games, Inc. All Rights Reserved.	
-
 using System;
 using System.Text;
 using System.Threading.Tasks;
-using Horde.Build.Api;
 using Horde.Build.Controllers;
 using Microsoft.AspNetCore.Authentication;
 using Microsoft.AspNetCore.Authentication.Cookies;
@@ -13,7 +10,7 @@ using Microsoft.AspNetCore.Authorization;
 using Microsoft.AspNetCore.Mvc;
 using Microsoft.Extensions.Options;
 
-namespace Horde.Build
+namespace Horde.Build.Server
 {
 	/// <summary>	
 	/// Dashboard authorization challenge controller	
diff --git a/Engine/Source/Programs/Horde/Horde.Build/Server/DowntimeService.cs b/Engine/Source/Programs/Horde/Horde.Build/Server/DowntimeService.cs
--- a/Engine/Source/Programs/Horde/Horde.Build/Server/DowntimeService.cs
+++ b/Engine/Source/Programs/Horde/Horde.Build/Server/DowntimeService.cs
@@ -4,7 +4,6 @@ using System;
 using System.Linq;
 using System.Threading;
 using System.Threading.Tasks;
-using Horde.Build.Models;
 using HordeCommon;
 using Microsoft.Extensions.Hosting;
 using Microsoft.Extensions.Logging;
diff --git a/Engine/Source/Programs/Horde/Horde.Build/Server/Globals.cs b/Engine/Source/Programs/Horde/Horde.Build/Server/Globals.cs
--- a/Engine/Source/Programs/Horde/Horde.Build/Server/Globals.cs
+++ b/Engine/Source/Programs/Horde/Horde.Build/Server/Globals.cs
@@ -1,4 +1,4 @@
-﻿// Copyright Epic Games, Inc. All Rights Reserved.
+// Copyright Epic Games, Inc. All Rights Reserved.
 
 using System;
 using System.Collections.Generic;
@@ -12,7 +12,7 @@ using Horde.Build.Utilities;
 using MongoDB.Bson;
 using MongoDB.Bson.Serialization.Attributes;
 
-namespace Horde.Build.Models
+namespace Horde.Build.Server
 {
 	/// <summary>
 	/// Base class for singleton documents
diff --git a/Engine/Source/Programs/Horde/Horde.Build/Server/HealthService.cs b/Engine/Source/Programs/Horde/Horde.Build/Server/HealthService.cs
--- a/Engine/Source/Programs/Horde/Horde.Build/Server/HealthService.cs
+++ b/Engine/Source/Programs/Horde/Horde.Build/Server/HealthService.cs
@@ -6,7 +6,7 @@ using Grpc.Health.V1;
 using Horde.Build.Server;
 using static Grpc.Health.V1.HealthCheckResponse.Types;
 
-namespace Horde.Build.Services
+namespace Horde.Build.Server
 {
 	/// <summary>
 	/// Implements the gRPC health checking protocol
diff --git a/Engine/Source/Programs/Horde/Horde.Build/Server/LifetimeController.cs b/Engine/Source/Programs/Horde/Horde.Build/Server/LifetimeController.cs
--- a/Engine/Source/Programs/Horde/Horde.Build/Server/LifetimeController.cs
+++ b/Engine/Source/Programs/Horde/Horde.Build/Server/LifetimeController.cs
@@ -4,7 +4,7 @@ using System.Threading.Tasks;
 using Horde.Build.Server;
 using Microsoft.AspNetCore.Mvc;
 
-namespace Horde.Build.Controllers
+namespace Horde.Build.Server
 {
 	/// <summary>
 	/// Controller for app lifetime related routes
diff --git a/Engine/Source/Programs/Horde/Horde.Build/Server/MongoService.cs b/Engine/Source/Programs/Horde/Horde.Build/Server/MongoService.cs
--- a/Engine/Source/Programs/Horde/Horde.Build/Server/MongoService.cs
+++ b/Engine/Source/Programs/Horde/Horde.Build/Server/MongoService.cs
@@ -18,8 +18,7 @@ using System.Threading.Channels;
 using System.Threading.Tasks;
 using EpicGames.Core;
 using EpicGames.Redis.Utility;
-using Horde.Build.Models;
-using Horde.Build.Services;
+using Horde.Build.Secrets;
 using Horde.Build.Utilities;
 using Horde.Build.Utiltiies;
 using Microsoft.Extensions.Hosting;
diff --git a/Engine/Source/Programs/Horde/Horde.Build/Server/Notices/INotice.cs b/Engine/Source/Programs/Horde/Horde.Build/Server/Notices/INotice.cs
--- a/Engine/Source/Programs/Horde/Horde.Build/Server/Notices/INotice.cs
+++ b/Engine/Source/Programs/Horde/Horde.Build/Server/Notices/INotice.cs
@@ -1,11 +1,12 @@
 // Copyright Epic Games, Inc. All Rights Reserved.
 
 using System;
+using Horde.Build.Users;
 using Horde.Build.Utilities;
 using MongoDB.Bson;
 using MongoDB.Bson.Serialization.Attributes;
 
-namespace Horde.Build.Models
+namespace Horde.Build.Server.Notices
 {
 	using UserId = ObjectId<IUser>;
 
diff --git a/Engine/Source/Programs/Horde/Horde.Build/Server/Notices/INoticeCollection.cs b/Engine/Source/Programs/Horde/Horde.Build/Server/Notices/INoticeCollection.cs
--- a/Engine/Source/Programs/Horde/Horde.Build/Server/Notices/INoticeCollection.cs
+++ b/Engine/Source/Programs/Horde/Horde.Build/Server/Notices/INoticeCollection.cs
@@ -3,11 +3,11 @@
 using System;
 using System.Collections.Generic;
 using System.Threading.Tasks;
-using Horde.Build.Models;
+using Horde.Build.Users;
 using Horde.Build.Utilities;
 using MongoDB.Bson;
 
-namespace Horde.Build.Collections
+namespace Horde.Build.Server.Notices
 {
 	using UserId = ObjectId<IUser>;
 
diff --git a/Engine/Source/Programs/Horde/Horde.Build/Server/Notices/NoticeCollection.cs b/Engine/Source/Programs/Horde/Horde.Build/Server/Notices/NoticeCollection.cs
--- a/Engine/Source/Programs/Horde/Horde.Build/Server/Notices/NoticeCollection.cs
+++ b/Engine/Source/Programs/Horde/Horde.Build/Server/Notices/NoticeCollection.cs
@@ -6,12 +6,11 @@ using MongoDB.Driver;
 using System;
 using System.Collections.Generic;
 using System.Linq;
-using Horde.Build.Models;
-using Horde.Build.Server;
 using Horde.Build.Utilities;
 using MongoDB.Bson;
+using Horde.Build.Users;
 
-namespace Horde.Build.Collections.Impl
+namespace Horde.Build.Server.Notices
 {
 	using UserId = ObjectId<IUser>;
 
diff --git a/Engine/Source/Programs/Horde/Horde.Build/Server/Notices/NoticeService.cs b/Engine/Source/Programs/Horde/Horde.Build/Server/Notices/NoticeService.cs
--- a/Engine/Source/Programs/Horde/Horde.Build/Server/Notices/NoticeService.cs
+++ b/Engine/Source/Programs/Horde/Horde.Build/Server/Notices/NoticeService.cs
@@ -3,12 +3,11 @@
 using System;
 using System.Collections.Generic;
 using System.Threading.Tasks;
-using Horde.Build.Collections;
-using Horde.Build.Models;
+using Horde.Build.Users;
 using Horde.Build.Utilities;
 using MongoDB.Bson;
 
-namespace Horde.Build.Services
+namespace Horde.Build.Server.Notices
 {
 	using UserId = ObjectId<IUser>;
 
diff --git a/Engine/Source/Programs/Horde/Horde.Build/Server/Notices/NoticesController.Messages.cs b/Engine/Source/Programs/Horde/Horde.Build/Server/Notices/NoticesController.Messages.cs
--- a/Engine/Source/Programs/Horde/Horde.Build/Server/Notices/NoticesController.Messages.cs
+++ b/Engine/Source/Programs/Horde/Horde.Build/Server/Notices/NoticesController.Messages.cs
@@ -2,11 +2,10 @@
 
 using System;
 using System.ComponentModel.DataAnnotations;
-using Horde.Build.Models;
+using Horde.Build.Users;
 
-namespace Horde.Build.Api
+namespace Horde.Build.Server.Notices
 {
-
 	/// <summary>
 	/// Parameters required to create a notice
 	/// </summary>
diff --git a/Engine/Source/Programs/Horde/Horde.Build/Server/Notices/NoticesController.cs b/Engine/Source/Programs/Horde/Horde.Build/Server/Notices/NoticesController.cs
--- a/Engine/Source/Programs/Horde/Horde.Build/Server/Notices/NoticesController.cs
+++ b/Engine/Source/Programs/Horde/Horde.Build/Server/Notices/NoticesController.cs
@@ -4,18 +4,14 @@ using System;
 using System.Collections.Generic;
 using System.Threading.Tasks;
 using Horde.Build.Acls;
-using Horde.Build.Api;
-using Horde.Build.Collections;
-using Horde.Build.Models;
-using Horde.Build.Services;
+using Horde.Build.Users;
 using Horde.Build.Utilities;
 using Microsoft.AspNetCore.Authorization;
 using Microsoft.AspNetCore.Mvc;
 using Microsoft.Extensions.Options;
 using MongoDB.Bson;
-using Horde.Build.Server;
 
-namespace Horde.Build.Controllers
+namespace Horde.Build.Server.Notices
 {
 	/// <summary>
 	/// Controller for the /api/v1/notices endpoint
diff --git a/Engine/Source/Programs/Horde/Horde.Build/Server/RequestTrackerService.cs b/Engine/Source/Programs/Horde/Horde.Build/Server/RequestTrackerService.cs
--- a/Engine/Source/Programs/Horde/Horde.Build/Server/RequestTrackerService.cs
+++ b/Engine/Source/Programs/Horde/Horde.Build/Server/RequestTrackerService.cs
@@ -10,7 +10,7 @@ using System.Threading.Tasks;
 using Microsoft.AspNetCore.Http;
 using Microsoft.Extensions.Logging;
 
-namespace Horde.Build.Utilities
+namespace Horde.Build.Server
 {
 	/// <summary>
 	/// Tracks open HTTP requests to ASP.NET
diff --git a/Engine/Source/Programs/Horde/Horde.Build/Server/RpcService.cs b/Engine/Source/Programs/Horde/Horde.Build/Server/RpcService.cs
--- a/Engine/Source/Programs/Horde/Horde.Build/Server/RpcService.cs
+++ b/Engine/Source/Programs/Horde/Horde.Build/Server/RpcService.cs
@@ -13,11 +13,19 @@ using EpicGames.Core;
 using Google.Protobuf.WellKnownTypes;
 using Grpc.Core;
 using Horde.Build.Acls;
-using Horde.Build.Collections;
+using Horde.Build.Agents;
+using Horde.Build.Agents.Pools;
+using Horde.Build.Agents.Sessions;
+using Horde.Build.Agents.Software;
 using Horde.Build.Jobs;
-using Horde.Build.Models;
+using Horde.Build.Jobs.Artifacts;
+using Horde.Build.Jobs.Graphs;
+using Horde.Build.Jobs.Templates;
+using Horde.Build.Jobs.TestData;
+using Horde.Build.Logs;
 using Horde.Build.Server;
-using Horde.Build.Tasks.Impl;
+using Horde.Build.Streams;
+using Horde.Build.Tasks;
 using Horde.Build.Utilities;
 using HordeCommon;
 using HordeCommon.Rpc;
@@ -26,14 +34,19 @@ using Microsoft.Extensions.Logging;
 using MongoDB.Bson;
 using MongoDB.Bson.Serialization;
 
-namespace Horde.Build.Services
+namespace Horde.Build.Server
 {
 	using AgentSoftwareChannelName = StringId<AgentSoftwareChannels>;
-	using IStream = Horde.Build.Models.IStream;
+	using IStream = Horde.Build.Streams.IStream;
 	using JobId = ObjectId<IJob>;
 	using LogId = ObjectId<ILogFile>;
 	using RpcAgentCapabilities = HordeCommon.Rpc.Messages.AgentCapabilities;
 	using RpcDeviceCapabilities = HordeCommon.Rpc.Messages.DeviceCapabilities;
+	using RpcGetStreamResponse = HordeCommon.Rpc.GetStreamResponse;
+	using RpcGetJobResponse = HordeCommon.Rpc.GetJobResponse;
+	using RpcGetStepResponse = HordeCommon.Rpc.GetStepResponse;
+	using RpcUpdateJobRequest = HordeCommon.Rpc.UpdateJobRequest;
+	using RpcUpdateStepRequest = HordeCommon.Rpc.UpdateStepRequest;
 	using SessionId = ObjectId<ISession>;
 	using StreamId = StringId<IStream>;
 
@@ -389,7 +402,7 @@ namespace Horde.Build.Services
 		/// <param name="request">Request arguments</param>
 		/// <param name="context">Context for the RPC call</param>
 		/// <returns>Information about the new agent</returns>
-		public override async Task<GetStreamResponse> GetStream(GetStreamRequest request, ServerCallContext context)
+		public override async Task<RpcGetStreamResponse> GetStream(GetStreamRequest request, ServerCallContext context)
 		{
 			StreamId streamIdValue = new StreamId(request.StreamId);
 
@@ -412,7 +425,7 @@ namespace Horde.Build.Services
 		/// <param name="request">Request arguments</param>
 		/// <param name="context">Context for the RPC call</param>
 		/// <returns>Information about the new agent</returns>
-		public override async Task<GetJobResponse> GetJob(GetJobRequest request, ServerCallContext context)
+		public override async Task<RpcGetJobResponse> GetJob(GetJobRequest request, ServerCallContext context)
 		{
 			JobId jobIdValue = new JobId(request.JobId.ToObjectId());
 
@@ -435,7 +448,7 @@ namespace Horde.Build.Services
 		/// <param name="request">Request arguments</param>
 		/// <param name="context">Context for the RPC call</param>
 		/// <returns>Information about the new agent</returns>
-		public override async Task<Empty> UpdateJob(UpdateJobRequest request, ServerCallContext context)
+		public override async Task<Empty> UpdateJob(RpcUpdateJobRequest request, ServerCallContext context)
 		{
 			JobId jobIdValue = new JobId(request.JobId);
 
@@ -470,7 +483,7 @@ namespace Horde.Build.Services
 			}
 
 			IJobStepBatch batch = AuthorizeBatch(job, request.BatchId.ToSubResourceId(), context);
-			job = await _jobService.UpdateBatchAsync(job, batchId, null, Api.JobStepBatchState.Starting);
+			job = await _jobService.UpdateBatchAsync(job, batchId, null, JobStepBatchState.Starting);
 
 			if (job == null)
 			{
@@ -500,7 +513,7 @@ namespace Horde.Build.Services
 			}
 
 			IJobStepBatch batch = AuthorizeBatch(job, request.BatchId.ToSubResourceId(), context);
-			await _jobService.UpdateBatchAsync(job, batch.Id, null, Api.JobStepBatchState.Complete);
+			await _jobService.UpdateBatchAsync(job, batch.Id, null, JobStepBatchState.Complete);
 			return new Empty();
 		}
 
@@ -534,7 +547,7 @@ namespace Horde.Build.Services
 
 			// Find the batch being executed
 			IJobStepBatch batch = AuthorizeBatch(job, request.BatchId.ToSubResourceId(), context);
-			if (batch.State != Api.JobStepBatchState.Starting && batch.State != Api.JobStepBatchState.Running)
+			if (batch.State != JobStepBatchState.Starting && batch.State != JobStepBatchState.Running)
 			{
 				return new BeginStepResponse { State = BeginStepResponse.Types.Result.Complete };
 			}
@@ -547,7 +560,7 @@ namespace Horde.Build.Services
 				if (stepIdx == batch.Steps.Count)
 				{
 					_logger.LogDebug("Job {JobId} batch {BatchId} is complete", job.Id, batch.Id);
-					if (await _jobService.TryUpdateBatchAsync(job, batch.Id, newState: Api.JobStepBatchState.Stopping) == null)
+					if (await _jobService.TryUpdateBatchAsync(job, batch.Id, newState: JobStepBatchState.Stopping) == null)
 					{
 						return null;
 					}
@@ -570,7 +583,7 @@ namespace Horde.Build.Services
 			// Create a log file if necessary
 			if (log.Value == null)
 			{
-				log.Value = await _logFileService.CreateLogFileAsync(job.Id, batch.SessionId, Api.LogType.Json);
+				log.Value = await _logFileService.CreateLogFileAsync(job.Id, batch.SessionId, LogType.Json);
 			}
 
 			// Get the node for this step
@@ -682,7 +695,7 @@ namespace Horde.Build.Services
 		/// <param name="request">Request arguments</param>
 		/// <param name="context">Context for the RPC call</param>
 		/// <returns>Information about the new agent</returns>
-		public override async Task<Empty> UpdateStep(UpdateStepRequest request, ServerCallContext context)
+		public override async Task<Empty> UpdateStep(RpcUpdateStepRequest request, ServerCallContext context)
 		{
 			IJob? job = await _jobService.GetJobAsync(new JobId(request.JobId));
 			if (job == null)
@@ -714,7 +727,7 @@ namespace Horde.Build.Services
 		/// <param name="request">Request arguments</param>
 		/// <param name="context">Context for the RPC call</param>
 		/// <returns>Information about the step</returns>
-		public override async Task<GetStepResponse> GetStep(GetStepRequest request, ServerCallContext context)
+		public override async Task<RpcGetStepResponse> GetStep(GetStepRequest request, ServerCallContext context)
 		{
 			IJob job = await GetJobAsync(new JobId(request.JobId));
 			IJobStepBatch batch = AuthorizeBatch(job, request.BatchId.ToSubResourceId(), context);
@@ -725,7 +738,7 @@ namespace Horde.Build.Services
 				throw new StructuredRpcException(StatusCode.NotFound, "Unable to find step {JobId}:{BatchId}:{StepId}", job.Id, batch.Id, stepId);
 			}
 
-			return new GetStepResponse { Outcome = step.Outcome, State = step.State, AbortRequested = step.AbortRequested || step.HasTimedOut(_clock.UtcNow) };
+			return new RpcGetStepResponse { Outcome = step.Outcome, State = step.State, AbortRequested = step.AbortRequested || step.HasTimedOut(_clock.UtcNow) };
 		}
 
 		/// <summary>
@@ -833,7 +846,7 @@ namespace Horde.Build.Services
 			{
 				throw new StructuredRpcException(StatusCode.NotFound, "Resource not found");
 			}
-			if (!Horde.Build.Services.LogFileService.AuthorizeForSession(logFile, context.GetHttpContext().User))
+			if (!LogFileService.AuthorizeForSession(logFile, context.GetHttpContext().User))
 			{
 				throw new StructuredRpcException(StatusCode.PermissionDenied, "Access denied");
 			}
diff --git a/Engine/Source/Programs/Horde/Horde.Build/Server/SchemaController.cs b/Engine/Source/Programs/Horde/Horde.Build/Server/SchemaController.cs
--- a/Engine/Source/Programs/Horde/Horde.Build/Server/SchemaController.cs
+++ b/Engine/Source/Programs/Horde/Horde.Build/Server/SchemaController.cs
@@ -12,7 +12,7 @@ using Horde.Build.Utilities;
 using Microsoft.AspNetCore.Mvc;
 using Microsoft.Extensions.Primitives;
 
-namespace Horde.Build.Controllers
+namespace Horde.Build.Server
 {
 	/// <summary>
 	/// Controller for the /api/v1/schema endpoint
diff --git a/Engine/Source/Programs/Horde/Horde.Build/Server/ServerController.cs b/Engine/Source/Programs/Horde/Horde.Build/Server/ServerController.cs
--- a/Engine/Source/Programs/Horde/Horde.Build/Server/ServerController.cs
+++ b/Engine/Source/Programs/Horde/Horde.Build/Server/ServerController.cs
@@ -2,12 +2,11 @@
 
 using System.Diagnostics;
 using System.Reflection;
-using Horde.Build.Api;
 using Microsoft.AspNetCore.Authorization;
 using Microsoft.AspNetCore.Mvc;
 using Microsoft.Extensions.Options;
 
-namespace Horde.Build.Controllers
+namespace Horde.Build.Server
 {
 	/// <summary>
 	/// Controller managing account status
diff --git a/Engine/Source/Programs/Horde/Horde.Build/Server/ServerMessages.cs b/Engine/Source/Programs/Horde/Horde.Build/Server/ServerMessages.cs
--- a/Engine/Source/Programs/Horde/Horde.Build/Server/ServerMessages.cs
+++ b/Engine/Source/Programs/Horde/Horde.Build/Server/ServerMessages.cs
@@ -5,7 +5,7 @@ using System.Diagnostics;
 using System.Reflection;
 using System.Runtime.InteropServices;
 
-namespace Horde.Build.Api
+namespace Horde.Build.Server
 {
 	/// <summary>
 	/// Server Info
diff --git a/Engine/Source/Programs/Horde/Horde.Build/Server/UpgradeService.cs b/Engine/Source/Programs/Horde/Horde.Build/Server/UpgradeService.cs
new file mode 100644
--- /dev/null
+++ b/Engine/Source/Programs/Horde/Horde.Build/Server/UpgradeService.cs
@@ -0,0 +1,106 @@
+// Copyright Epic Games, Inc. All Rights Reserved.
+
+using System;
+using System.Threading.Tasks;
+using Horde.Build.Jobs;
+using Horde.Build.Server;
+using Horde.Build.Users;
+using Microsoft.Extensions.DependencyInjection;
+using Microsoft.Extensions.Logging;
+
+namespace Horde.Build.Server
+{
+	/// <summary>
+	/// Service to perform schema upgrades on the database
+	/// </summary>
+	public class UpgradeService
+	{
+		/// <summary>
+		/// The current schema version
+		/// </summary>
+		public const int LatestSchemaVersion = 5;
+
+		/// <summary>
+		/// The database service instance
+		/// </summary>
+		readonly MongoService _mongoService;
+
+		/// <summary>
+		/// The DI service provider
+		/// </summary>
+		readonly IServiceProvider _serviceProvider;
+
+		/// <summary>
+		/// Logger instance
+		/// </summary>
+		readonly ILogger<UpgradeService> _logger;
+
+		/// <summary>
+		/// Constructor
+		/// </summary>
+		/// <param name="mongoService">The database service singleton</param>
+		/// <param name="serviceProvider">The DI service provider</param>
+		/// <param name="logger">Logger instance</param>
+		public UpgradeService(MongoService mongoService, IServiceProvider serviceProvider, ILogger<UpgradeService> logger)
+		{
+			_mongoService = mongoService;
+			_serviceProvider = serviceProvider;
+			_logger = logger;
+		}
+
+		/// <summary>
+		/// Update the database schema
+		/// </summary>
+		/// <param name="fromVersion">The version to upgrade from. Automatically detected if this is null.</param>
+		/// <returns>Async task</returns>
+		public async Task UpgradeSchemaAsync(int? fromVersion)
+		{
+			Globals globals = await _mongoService.GetGlobalsAsync();
+			int schemaVersion = fromVersion ?? globals.SchemaVersion ?? 0;
+
+			while (schemaVersion < LatestSchemaVersion)
+			{
+				_logger.LogInformation("Upgrading from schema version {Version}", schemaVersion);
+
+				// Handle the current version
+				//				if (SchemaVersion == 0)
+				//				{
+				//					IssueCollection IssueCollection = new IssueCollection(DatabaseService);
+				//					IJobCollection JobCollection = ServiceProvider.GetService<IJobCollection>();
+				//					await IssueCollection.UpdateSchemaLogIdFields(JobCollection, Logger);
+				//				}
+				//				if(SchemaVersion == 2)
+				//				{
+				//					IssueCollection IssueCollection = new IssueCollection(DatabaseService);
+				//					await IssueCollection.FixResolvedIssuesAsync(Logger);
+				//				}
+
+				if (schemaVersion == 3)
+				{
+					IJobCollection jobCollection = _serviceProvider.GetRequiredService<IJobCollection>();
+					await jobCollection.UpgradeDocumentsAsync();
+				}
+				if (schemaVersion == 4)
+				{
+					UserCollectionV1 userCollectionV1 = new UserCollectionV1(_serviceProvider.GetRequiredService<MongoService>());
+					using UserCollectionV2 userCollectionV2 = new UserCollectionV2(_serviceProvider.GetRequiredService<MongoService>(), _serviceProvider.GetRequiredService<ILogger<UserCollectionV2>>());
+					await userCollectionV2.ResaveDocumentsAsync(userCollectionV1);
+				}
+
+				// Increment the version number
+				schemaVersion++;
+
+				// Try to update the current schema version number
+				while (globals.SchemaVersion == null || globals.SchemaVersion < schemaVersion)
+				{
+					globals.SchemaVersion = schemaVersion;
+					if (await _mongoService.TryUpdateSingletonAsync(globals))
+					{
+						break;
+					}
+					globals = await _mongoService.GetGlobalsAsync();
+				}
+			}
+		}
+	}
+}
diff --git a/Engine/Source/Programs/Horde/Horde.Build/ServerSettings.cs b/Engine/Source/Programs/Horde/Horde.Build/ServerSettings.cs
--- a/Engine/Source/Programs/Horde/Horde.Build/ServerSettings.cs
+++ b/Engine/Source/Programs/Horde/Horde.Build/ServerSettings.cs
@@ -3,7 +3,7 @@
 using System;
 using System.Linq;
 using EpicGames.Horde.Storage;
-using Horde.Build.Fleet.Autoscale;
+using Horde.Build.Agents.Fleet;
 using Horde.Build.Utilities;
 using TimeZoneConverter;
 
diff --git a/Engine/Source/Programs/Horde/Horde.Build/Startup.cs b/Engine/Source/Programs/Horde/Horde.Build/Startup.cs
--- a/Engine/Source/Programs/Horde/Horde.Build/Startup.cs
+++ b/Engine/Source/Programs/Horde/Horde.Build/Startup.cs
@@ -22,31 +22,34 @@ using Grpc.Core;
 using Grpc.Core.Interceptors;
 using Horde.Build.Acls;
 using Horde.Build.Authentication;
-using Horde.Build.Collections;
-using Horde.Build.Collections.Impl;
-using Horde.Build.Commits;
-using Horde.Build.Commits.Impl;
+using Horde.Build.Agents;
+using Horde.Build.Agents.Fleet;
+using Horde.Build.Agents.Leases;
+using Horde.Build.Agents.Pools;
+using Horde.Build.Agents.Sessions;
+using Horde.Build.Agents.Software;
+using Horde.Build.Agents.Telemetry;
 using Horde.Build.Compute;
-using Horde.Build.Compute.Impl;
-using Horde.Build.Config;
-using Horde.Build.Fleet.Autoscale;
-using Horde.Build.Issues.Impl;
+using Horde.Build.Configuration;
+using Horde.Build.Devices;
+using Horde.Build.Issues;
+using Horde.Build.Issues.External;
 using Horde.Build.Jobs;
+using Horde.Build.Jobs.Artifacts;
+using Horde.Build.Jobs.Graphs;
+using Horde.Build.Jobs.Schedules;
+using Horde.Build.Jobs.Templates;
+using Horde.Build.Jobs.TestData;
+using Horde.Build.Jobs.Timing;
 using Horde.Build.Logs;
 using Horde.Build.Logs.Builder;
-using Horde.Build.Logs.Readers;
 using Horde.Build.Logs.Storage;
-using Horde.Build.Logs.Storage.Impl;
-using Horde.Build.Models;
 using Horde.Build.Notifications;
-using Horde.Build.Notifications.Impl;
+using Horde.Build.Secrets;
 using Horde.Build.Server;
-using Horde.Build.Services;
-using Horde.Build.Services.Impl;
 using Horde.Build.Storage;
 using Horde.Build.Storage.Backends;
 using Horde.Build.Tasks;
-using Horde.Build.Tasks.Impl;
 using Horde.Build.Tools;
 using Horde.Build.Utilities;
 using HordeCommon;
@@ -81,6 +84,15 @@ using Serilog.Events;
 using StackExchange.Redis;
 using StatsdClient;
 using Status = Grpc.Core.Status;
+using Horde.Build.Users;
+using Horde.Build.Perforce;
+using Horde.Build.Projects;
+using Horde.Build.Streams;
+using Horde.Build.Ugs;
+using Horde.Build.Auditing;
+using Horde.Build.Agents.Fleet.Providers;
+using Horde.Build.Server.Notices;
+using Horde.Build.Notifications.Sinks;
 
 namespace Horde.Build
 {
@@ -434,7 +446,7 @@ namespace Horde.Build
 			{
 				// This SharedCredentials option is primarily for development purposes.
 				(string accessKey, string secretAccessKey, string secretToken) = AwsHelper.ReadAwsCredentials(settings.S3AwsProfile);
-				awsOptions.Credentials = new Credentials(accessKey, secretAccessKey, secretToken, DateTime.Now + TimeSpan.FromHours(12));
+				awsOptions.Credentials = new Amazon.SecurityToken.Model.Credentials(accessKey, secretAccessKey, secretToken, DateTime.Now + TimeSpan.FromHours(12));
 			}
 			
 			services.AddSingleton(awsOptions);
diff --git a/Engine/Source/Programs/Horde/Horde.Build/Streams/IStream.cs b/Engine/Source/Programs/Horde/Horde.Build/Streams/IStream.cs
--- a/Engine/Source/Programs/Horde/Horde.Build/Streams/IStream.cs
+++ b/Engine/Source/Programs/Horde/Horde.Build/Streams/IStream.cs
@@ -7,13 +7,18 @@ using System.Linq;
 using System.Text.RegularExpressions;
 using EpicGames.Core;
 using Horde.Build.Acls;
-using Horde.Build.Api;
+using Horde.Build.Agents;
+using Horde.Build.Agents.Pools;
+using Horde.Build.Jobs.Schedules;
+using Horde.Build.Jobs.Templates;
+using Horde.Build.Projects;
 using Horde.Build.Server;
+using Horde.Build.Users;
 using Horde.Build.Utilities;
 using HordeCommon;
 using MongoDB.Bson.Serialization.Attributes;
 
-namespace Horde.Build.Models
+namespace Horde.Build.Streams
 {
 	using PoolId = StringId<IPool>;
 	using ProjectId = StringId<IProject>;
@@ -116,9 +121,9 @@ namespace Horde.Build.Models
 		/// Creates an API response object from this stream
 		/// </summary>
 		/// <returns>The response object</returns>
-		public Api.GetAgentTypeResponse ToApiResponse()
+		public GetAgentTypeResponse ToApiResponse()
 		{
-			return new Api.GetAgentTypeResponse(Pool.ToString(), Workspace, TempStorageDir, Environment);
+			return new GetAgentTypeResponse(Pool.ToString(), Workspace, TempStorageDir, Environment);
 		}
 
 		/// <summary>
@@ -236,9 +241,9 @@ namespace Horde.Build.Models
 		/// Creates an API response object from this stream
 		/// </summary>
 		/// <returns>The response object</returns>
-		public Api.GetWorkspaceTypeResponse ToApiResponse()
+		public GetWorkspaceTypeResponse ToApiResponse()
 		{
-			return new Api.GetWorkspaceTypeResponse(Cluster, ServerAndPort, UserName, Identifier, Stream, View, Incremental, UseAutoSdk);
+			return new GetWorkspaceTypeResponse(Cluster, ServerAndPort, UserName, Identifier, Stream, View, Incremental, UseAutoSdk);
 		}
 	}
 
@@ -733,13 +738,13 @@ namespace Horde.Build.Models
 		/// <param name="bIncludeAcl">Whether to include the ACL in the response object</param>
 		/// <param name="apiTemplateRefs">The template refs for this stream. Passed separately because they have their own ACL.</param>
 		/// <returns>New response instance</returns>
-		public static Api.GetStreamResponse ToApiResponse(this IStream stream, bool bIncludeAcl, List<GetTemplateRefResponse> apiTemplateRefs)
+		public static GetStreamResponse ToApiResponse(this IStream stream, bool bIncludeAcl, List<GetTemplateRefResponse> apiTemplateRefs)
 		{
 			List<GetStreamTabResponse> apiTabs = stream.Tabs.ConvertAll(x => x.ToResponse());
 			Dictionary<string, GetAgentTypeResponse> apiAgentTypes = stream.AgentTypes.ToDictionary(x => x.Key, x => x.Value.ToApiResponse());
 			Dictionary<string, GetWorkspaceTypeResponse> apiWorkspaceTypes = stream.WorkspaceTypes.ToDictionary(x => x.Key, x => x.Value.ToApiResponse());
 			GetAclResponse? apiAcl = (bIncludeAcl && stream.Acl != null)? new GetAclResponse(stream.Acl) : null;
-			return new Api.GetStreamResponse(stream.Id.ToString(), stream.ProjectId.ToString(), stream.Name, stream.ConfigRevision, stream.Order, stream.NotificationChannel, stream.NotificationChannelFilter, stream.TriageChannel, stream.DefaultPreflight?.ToRequest(), apiTabs, apiAgentTypes, apiWorkspaceTypes, apiTemplateRefs, apiAcl, stream.PausedUntil, stream.PauseComment, stream.Config.Workflows);
+			return new GetStreamResponse(stream.Id.ToString(), stream.ProjectId.ToString(), stream.Name, stream.ConfigRevision, stream.Order, stream.NotificationChannel, stream.NotificationChannelFilter, stream.TriageChannel, stream.DefaultPreflight?.ToRequest(), apiTabs, apiAgentTypes, apiWorkspaceTypes, apiTemplateRefs, apiAcl, stream.PausedUntil, stream.PauseComment, stream.Config.Workflows);
 		}
 
 		/// <summary>
diff --git a/Engine/Source/Programs/Horde/Horde.Build/Streams/IStreamCollection.cs b/Engine/Source/Programs/Horde/Horde.Build/Streams/IStreamCollection.cs
--- a/Engine/Source/Programs/Horde/Horde.Build/Streams/IStreamCollection.cs
+++ b/Engine/Source/Programs/Horde/Horde.Build/Streams/IStreamCollection.cs
@@ -3,12 +3,11 @@
 using System;
 using System.Collections.Generic;
 using System.Threading.Tasks;
-using Horde.Build.Api;
-using Horde.Build.Models;
-using Horde.Build.Server;
+using Horde.Build.Jobs;
+using Horde.Build.Projects;
 using Horde.Build.Utilities;
 
-namespace Horde.Build.Collections
+namespace Horde.Build.Streams
 {
 	using JobId = ObjectId<IJob>;
 	using ProjectId = StringId<IProject>;
diff --git a/Engine/Source/Programs/Horde/Horde.Build/Streams/StreamCollection.cs b/Engine/Source/Programs/Horde/Horde.Build/Streams/StreamCollection.cs
--- a/Engine/Source/Programs/Horde/Horde.Build/Streams/StreamCollection.cs
+++ b/Engine/Source/Programs/Horde/Horde.Build/Streams/StreamCollection.cs
@@ -7,17 +7,20 @@ using System.IO;
 using System.Linq;
 using System.Threading.Tasks;
 using Horde.Build.Acls;
-using Horde.Build.Api;
-using Horde.Build.Config;
-using Horde.Build.Models;
+using Horde.Build.Configuration;
+using Horde.Build.Jobs;
+using Horde.Build.Jobs.Schedules;
+using Horde.Build.Jobs.Templates;
+using Horde.Build.Projects;
 using Horde.Build.Server;
+using Horde.Build.Users;
 using Horde.Build.Utilities;
 using HordeCommon;
 using Microsoft.Extensions.Logging;
 using MongoDB.Bson.Serialization.Attributes;
 using MongoDB.Driver;
 
-namespace Horde.Build.Collections.Impl
+namespace Horde.Build.Streams
 {
 	using JobId = ObjectId<IJob>;
 	using ProjectId = StringId<IProject>;
diff --git a/Engine/Source/Programs/Horde/Horde.Build/Streams/StreamService.cs b/Engine/Source/Programs/Horde/Horde.Build/Streams/StreamService.cs
--- a/Engine/Source/Programs/Horde/Horde.Build/Streams/StreamService.cs
+++ b/Engine/Source/Programs/Horde/Horde.Build/Streams/StreamService.cs
@@ -6,14 +6,13 @@ using System.Linq;
 using System.Security.Claims;
 using System.Threading.Tasks;
 using Horde.Build.Acls;
-using Horde.Build.Api;
-using Horde.Build.Collections;
-using Horde.Build.Models;
+using Horde.Build.Jobs;
+using Horde.Build.Projects;
 using Horde.Build.Utilities;
 using Microsoft.Extensions.Caching.Memory;
 using MongoDB.Driver;
 
-namespace Horde.Build.Services
+namespace Horde.Build.Streams
 {
 	using JobId = ObjectId<IJob>;
 	using ProjectId = StringId<IProject>;
diff --git a/Engine/Source/Programs/Horde/Horde.Build/Streams/StreamTab.cs b/Engine/Source/Programs/Horde/Horde.Build/Streams/StreamTab.cs
--- a/Engine/Source/Programs/Horde/Horde.Build/Streams/StreamTab.cs
+++ b/Engine/Source/Programs/Horde/Horde.Build/Streams/StreamTab.cs
@@ -2,11 +2,10 @@
 
 using System;
 using System.Collections.Generic;
-using Horde.Build.Api;
 using Horde.Build.Utilities;
 using MongoDB.Bson.Serialization.Attributes;
 
-namespace Horde.Build.Models
+namespace Horde.Build.Streams
 {
 	using TemplateRefId = StringId<TemplateRef>;
 
diff --git a/Engine/Source/Programs/Horde/Horde.Build/Streams/StreamsController.Messages.cs b/Engine/Source/Programs/Horde/Horde.Build/Streams/StreamsController.Messages.cs
--- a/Engine/Source/Programs/Horde/Horde.Build/Streams/StreamsController.Messages.cs
+++ b/Engine/Source/Programs/Horde/Horde.Build/Streams/StreamsController.Messages.cs
@@ -4,12 +4,16 @@ using System;
 using System.Collections.Generic;
 using System.ComponentModel.DataAnnotations;
 using EpicGames.Core;
-using Horde.Build.Models;
+using Horde.Build.Acls;
+using Horde.Build.Jobs.Graphs;
+using Horde.Build.Jobs.Templates;
 using Horde.Build.Server;
+using Horde.Build.Jobs.Schedules;
+using Horde.Build.Users;
 using Horde.Build.Utilities;
 using HordeCommon;
 
-namespace Horde.Build.Api
+namespace Horde.Build.Streams
 {
 	using TemplateRefId = StringId<TemplateRef>;
 	using WorkflowId = StringId<WorkflowConfig>;
diff --git a/Engine/Source/Programs/Horde/Horde.Build/Streams/StreamsController.cs b/Engine/Source/Programs/Horde/Horde.Build/Streams/StreamsController.cs
--- a/Engine/Source/Programs/Horde/Horde.Build/Streams/StreamsController.cs
+++ b/Engine/Source/Programs/Horde/Horde.Build/Streams/StreamsController.cs
@@ -5,15 +5,16 @@ using System.Collections.Generic;
 using System.Linq;
 using System.Threading.Tasks;
 using Horde.Build.Acls;
-using Horde.Build.Api;
-using Horde.Build.Collections;
-using Horde.Build.Models;
-using Horde.Build.Services;
+using Horde.Build.Jobs;
+using Horde.Build.Jobs.Templates;
+using Horde.Build.Perforce;
+using Horde.Build.Projects;
+using Horde.Build.Users;
 using Horde.Build.Utilities;
 using Microsoft.AspNetCore.Authorization;
 using Microsoft.AspNetCore.Mvc;
 
-namespace Horde.Build.Controllers
+namespace Horde.Build.Streams
 {
 	using ProjectId = StringId<IProject>;
 	using StreamId = StringId<IStream>;
diff --git a/Engine/Source/Programs/Horde/Horde.Build/Tasks/ConformTaskSource.cs b/Engine/Source/Programs/Horde/Horde.Build/Tasks/ConformTaskSource.cs
--- a/Engine/Source/Programs/Horde/Horde.Build/Tasks/ConformTaskSource.cs
+++ b/Engine/Source/Programs/Horde/Horde.Build/Tasks/ConformTaskSource.cs
@@ -7,18 +7,20 @@ using System.Threading;
 using System.Threading.Tasks;
 using Google.Protobuf;
 using Google.Protobuf.WellKnownTypes;
-using Horde.Build.Api;
-using Horde.Build.Collections;
-using Horde.Build.Models;
+using Horde.Build.Agents;
+using Horde.Build.Agents.Leases;
+using Horde.Build.Agents.Pools;
+using Horde.Build.Jobs;
+using Horde.Build.Logs;
+using Horde.Build.Perforce;
 using Horde.Build.Server;
-using Horde.Build.Services;
 using Horde.Build.Utilities;
 using HordeCommon;
 using HordeCommon.Rpc.Tasks;
 using Microsoft.Extensions.Hosting;
 using Microsoft.Extensions.Logging;
 
-namespace Horde.Build.Tasks.Impl
+namespace Horde.Build.Tasks
 {
 	using JobId = ObjectId<IJob>;
 	using LeaseId = ObjectId<ILease>;
diff --git a/Engine/Source/Programs/Horde/Horde.Build/Tasks/ITaskSource.cs b/Engine/Source/Programs/Horde/Horde.Build/Tasks/ITaskSource.cs
--- a/Engine/Source/Programs/Horde/Horde.Build/Tasks/ITaskSource.cs
+++ b/Engine/Source/Programs/Horde/Horde.Build/Tasks/ITaskSource.cs
@@ -12,7 +12,8 @@ using EpicGames.Core;
 using Google.Protobuf;
 using Google.Protobuf.Reflection;
 using Google.Protobuf.WellKnownTypes;
-using Horde.Build.Models;
+using Horde.Build.Agents;
+using Horde.Build.Agents.Leases;
 using Horde.Build.Utilities;
 using HordeCommon;
 using Microsoft.Extensions.Logging;
diff --git a/Engine/Source/Programs/Horde/Horde.Build/Tasks/RestartTaskSource.cs b/Engine/Source/Programs/Horde/Horde.Build/Tasks/RestartTaskSource.cs
--- a/Engine/Source/Programs/Horde/Horde.Build/Tasks/RestartTaskSource.cs
+++ b/Engine/Source/Programs/Horde/Horde.Build/Tasks/RestartTaskSource.cs
@@ -1,20 +1,18 @@
 // Copyright Epic Games, Inc. All Rights Reserved.
 
-using System.Collections.Generic;
-using System.Runtime.CompilerServices;
 using System.Threading;
 using System.Threading.Tasks;
-using EpicGames.Core;
 using Google.Protobuf;
 using Google.Protobuf.WellKnownTypes;
-using Horde.Build.Api;
-using Horde.Build.Models;
-using Horde.Build.Services;
+using Horde.Build.Agents;
+using Horde.Build.Agents.Leases;
+using Horde.Build.Jobs;
+using Horde.Build.Logs;
 using Horde.Build.Utilities;
 using HordeCommon;
 using HordeCommon.Rpc.Tasks;
 
-namespace Horde.Build.Tasks.Impl
+namespace Horde.Build.Tasks
 {
 	using JobId = ObjectId<IJob>;
 	using LeaseId = ObjectId<ILease>;
diff --git a/Engine/Source/Programs/Horde/Horde.Build/Tasks/ShutdownTaskSource.cs b/Engine/Source/Programs/Horde/Horde.Build/Tasks/ShutdownTaskSource.cs
--- a/Engine/Source/Programs/Horde/Horde.Build/Tasks/ShutdownTaskSource.cs
+++ b/Engine/Source/Programs/Horde/Horde.Build/Tasks/ShutdownTaskSource.cs
@@ -4,14 +4,15 @@ using System.Threading;
 using System.Threading.Tasks;
 using Google.Protobuf;
 using Google.Protobuf.WellKnownTypes;
-using Horde.Build.Api;
-using Horde.Build.Models;
-using Horde.Build.Services;
+using Horde.Build.Agents;
+using Horde.Build.Agents.Leases;
+using Horde.Build.Jobs;
+using Horde.Build.Logs;
 using Horde.Build.Utilities;
 using HordeCommon;
 using HordeCommon.Rpc.Tasks;
 
-namespace Horde.Build.Tasks.Impl
+namespace Horde.Build.Tasks
 {
 	using JobId = ObjectId<IJob>;
 	using LeaseId = ObjectId<ILease>;
diff --git a/Engine/Source/Programs/Horde/Horde.Build/Tasks/UpgradeTaskSource.cs b/Engine/Source/Programs/Horde/Horde.Build/Tasks/UpgradeTaskSource.cs
--- a/Engine/Source/Programs/Horde/Horde.Build/Tasks/UpgradeTaskSource.cs
+++ b/Engine/Source/Programs/Horde/Horde.Build/Tasks/UpgradeTaskSource.cs
@@ -1,20 +1,20 @@
 // Copyright Epic Games, Inc. All Rights Reserved.
 
 using System;
-using System.Collections.Generic;
-using System.Runtime.CompilerServices;
 using System.Threading;
 using System.Threading.Tasks;
 using Google.Protobuf;
 using Google.Protobuf.WellKnownTypes;
-using Horde.Build.Api;
-using Horde.Build.Models;
-using Horde.Build.Services;
+using Horde.Build.Agents;
+using Horde.Build.Agents.Leases;
+using Horde.Build.Agents.Software;
+using Horde.Build.Jobs;
+using Horde.Build.Logs;
 using Horde.Build.Utilities;
 using HordeCommon;
 using HordeCommon.Rpc.Tasks;
 
-namespace Horde.Build.Tasks.Impl
+namespace Horde.Build.Tasks
 {
 	using AgentSoftwareChannelName = StringId<AgentSoftwareChannels>;
 	using JobId = ObjectId<IJob>;
diff --git a/Engine/Source/Programs/Horde/Horde.Build/Tools/ToolsController.cs b/Engine/Source/Programs/Horde/Horde.Build/Tools/ToolsController.cs
--- a/Engine/Source/Programs/Horde/Horde.Build/Tools/ToolsController.cs
+++ b/Engine/Source/Programs/Horde/Horde.Build/Tools/ToolsController.cs
@@ -13,7 +13,7 @@ using System.IO;
 using System.Linq;
 using System.Threading.Tasks;
 
-namespace Horde.Build.Tools.Impl
+namespace Horde.Build.Tools
 {
 	using ToolId = StringId<Tool>;
 	using ToolDeploymentId = ObjectId<ToolDeployment>;
diff --git a/Engine/Source/Programs/Horde/Horde.Build/Ugs/IUgsMetadataCollection.cs b/Engine/Source/Programs/Horde/Horde.Build/Ugs/IUgsMetadataCollection.cs
--- a/Engine/Source/Programs/Horde/Horde.Build/Ugs/IUgsMetadataCollection.cs
+++ b/Engine/Source/Programs/Horde/Horde.Build/Ugs/IUgsMetadataCollection.cs
@@ -3,10 +3,8 @@
 using System;
 using System.Collections.Generic;
 using System.Threading.Tasks;
-using Horde.Build.Api;
-using Horde.Build.Models;
 
-namespace Horde.Build.Collections
+namespace Horde.Build.Ugs
 {
 	/// <summary>
 	/// Collection of stream documents
diff --git a/Engine/Source/Programs/Horde/Horde.Build/Ugs/UgsController.Messages.cs b/Engine/Source/Programs/Horde/Horde.Build/Ugs/UgsController.Messages.cs
--- a/Engine/Source/Programs/Horde/Horde.Build/Ugs/UgsController.Messages.cs
+++ b/Engine/Source/Programs/Horde/Horde.Build/Ugs/UgsController.Messages.cs
@@ -5,9 +5,10 @@ using System.Collections.Generic;
 using System.ComponentModel.DataAnnotations;
 using System.Linq;
 using EpicGames.Core;
-using Horde.Build.Models;
+using Horde.Build.Issues;
+using Horde.Build.Users;
 
-namespace Horde.Build.Api
+namespace Horde.Build.Ugs
 {
 	/// <summary>
 	/// Review by a user of a particular change
diff --git a/Engine/Source/Programs/Horde/Horde.Build/Ugs/UgsController.cs b/Engine/Source/Programs/Horde/Horde.Build/Ugs/UgsController.cs
--- a/Engine/Source/Programs/Horde/Horde.Build/Ugs/UgsController.cs
+++ b/Engine/Source/Programs/Horde/Horde.Build/Ugs/UgsController.cs
@@ -4,17 +4,16 @@ using System;
 using System.Collections.Generic;
 using System.Linq;
 using System.Threading.Tasks;
-using Horde.Build.Api;
-using Horde.Build.Collections;
-using Horde.Build.Models;
-using Horde.Build.Services;
+using Horde.Build.Users;
+using Horde.Build.Logs;
+using Horde.Build.Issues;
 using Horde.Build.Utilities;
 using Microsoft.AspNetCore.Mvc;
 using Microsoft.Extensions.Logging;
 using Microsoft.Extensions.Options;
 using MongoDB.Bson;
 
-namespace Horde.Build.Controllers
+namespace Horde.Build.Ugs
 {
 	using LogId = ObjectId<ILogFile>;
 	using UserId = ObjectId<IUser>;
diff --git a/Engine/Source/Programs/Horde/Horde.Build/Ugs/UgsMetadata.cs b/Engine/Source/Programs/Horde/Horde.Build/Ugs/UgsMetadata.cs
--- a/Engine/Source/Programs/Horde/Horde.Build/Ugs/UgsMetadata.cs
+++ b/Engine/Source/Programs/Horde/Horde.Build/Ugs/UgsMetadata.cs
@@ -2,9 +2,8 @@
 
 using System;
 using System.Collections.Generic;
-using Horde.Build.Api;
 
-namespace Horde.Build.Models
+namespace Horde.Build.Ugs
 {
 	/// <summary>
 	/// Metadata for a particular user
diff --git a/Engine/Source/Programs/Horde/Horde.Build/Ugs/UgsMetadataCollection.cs b/Engine/Source/Programs/Horde/Horde.Build/Ugs/UgsMetadataCollection.cs
--- a/Engine/Source/Programs/Horde/Horde.Build/Ugs/UgsMetadataCollection.cs
+++ b/Engine/Source/Programs/Horde/Horde.Build/Ugs/UgsMetadataCollection.cs
@@ -5,15 +5,13 @@ using System.Collections.Generic;
 using System.Diagnostics.CodeAnalysis;
 using System.Linq;
 using System.Threading.Tasks;
-using Horde.Build.Api;
-using Horde.Build.Models;
 using Horde.Build.Server;
 using Horde.Build.Utilities;
 using MongoDB.Bson;
 using MongoDB.Bson.Serialization.Attributes;
 using MongoDB.Driver;
 
-namespace Horde.Build.Collections.Impl
+namespace Horde.Build.Ugs
 {
 	/// <summary>
 	/// Concrete implementation of IUgsMetadataCollection
diff --git a/Engine/Source/Programs/Horde/Horde.Build/Users/IAvatar.cs b/Engine/Source/Programs/Horde/Horde.Build/Users/IAvatar.cs
--- a/Engine/Source/Programs/Horde/Horde.Build/Users/IAvatar.cs
+++ b/Engine/Source/Programs/Horde/Horde.Build/Users/IAvatar.cs
@@ -1,6 +1,6 @@
 // Copyright Epic Games, Inc. All Rights Reserved.
 
-namespace Horde.Build.Models
+namespace Horde.Build.Users
 {
 	/// <summary>
 	/// Avatar for the user
diff --git a/Engine/Source/Programs/Horde/Horde.Build/Users/IAvatarService.cs b/Engine/Source/Programs/Horde/Horde.Build/Users/IAvatarService.cs
--- a/Engine/Source/Programs/Horde/Horde.Build/Users/IAvatarService.cs
+++ b/Engine/Source/Programs/Horde/Horde.Build/Users/IAvatarService.cs
@@ -1,9 +1,8 @@
 // Copyright Epic Games, Inc. All Rights Reserved.
 
 using System.Threading.Tasks;
-using Horde.Build.Models;
 
-namespace Horde.Build.Services
+namespace Horde.Build.Users
 {
 	/// <summary>
 	/// Provides avatar images
diff --git a/Engine/Source/Programs/Horde/Horde.Build/Users/IServiceAccount.cs b/Engine/Source/Programs/Horde/Horde.Build/Users/IServiceAccount.cs
--- a/Engine/Source/Programs/Horde/Horde.Build/Users/IServiceAccount.cs
+++ b/Engine/Source/Programs/Horde/Horde.Build/Users/IServiceAccount.cs
@@ -3,7 +3,7 @@
 using System.Collections.Generic;
 using MongoDB.Bson;
 
-namespace Horde.Build.Models
+namespace Horde.Build.Users
 {
 	/// <summary>
 	/// A service account used for authenticating during server-to-server communication
diff --git a/Engine/Source/Programs/Horde/Horde.Build/Users/IServiceAccountCollection.cs b/Engine/Source/Programs/Horde/Horde.Build/Users/IServiceAccountCollection.cs
--- a/Engine/Source/Programs/Horde/Horde.Build/Users/IServiceAccountCollection.cs
+++ b/Engine/Source/Programs/Horde/Horde.Build/Users/IServiceAccountCollection.cs
@@ -2,10 +2,9 @@
 
 using System.Collections.Generic;
 using System.Threading.Tasks;
-using Horde.Build.Models;
 using MongoDB.Bson;
 
-namespace Horde.Build.Collections
+namespace Horde.Build.Users
 {
 	/// <summary>
 	/// Interface for a collection of service account documents
diff --git a/Engine/Source/Programs/Horde/Horde.Build/Users/IUser.cs b/Engine/Source/Programs/Horde/Horde.Build/Users/IUser.cs
--- a/Engine/Source/Programs/Horde/Horde.Build/Users/IUser.cs
+++ b/Engine/Source/Programs/Horde/Horde.Build/Users/IUser.cs
@@ -1,10 +1,11 @@
 // Copyright Epic Games, Inc. All Rights Reserved.
 
 using System.Collections.Generic;
+using Horde.Build.Jobs;
 using Horde.Build.Utilities;
 using MongoDB.Bson;
 
-namespace Horde.Build.Models
+namespace Horde.Build.Users
 {
 	using JobId = ObjectId<IJob>;
 	using UserId = ObjectId<IUser>;
diff --git a/Engine/Source/Programs/Horde/Horde.Build/Users/IUserCollection.cs b/Engine/Source/Programs/Horde/Horde.Build/Users/IUserCollection.cs
--- a/Engine/Source/Programs/Horde/Horde.Build/Users/IUserCollection.cs
+++ b/Engine/Source/Programs/Horde/Horde.Build/Users/IUserCollection.cs
@@ -3,11 +3,11 @@
 using System.Collections.Generic;
 using System.Security.Claims;
 using System.Threading.Tasks;
-using Horde.Build.Models;
+using Horde.Build.Jobs;
 using Horde.Build.Utilities;
 using MongoDB.Bson;
 
-namespace Horde.Build.Collections
+namespace Horde.Build.Users
 {
 	using JobId = ObjectId<IJob>;
 	using UserId = ObjectId<IUser>;
diff --git a/Engine/Source/Programs/Horde/Horde.Build/Users/ServiceAccountCollection.cs b/Engine/Source/Programs/Horde/Horde.Build/Users/ServiceAccountCollection.cs
--- a/Engine/Source/Programs/Horde/Horde.Build/Users/ServiceAccountCollection.cs
+++ b/Engine/Source/Programs/Horde/Horde.Build/Users/ServiceAccountCollection.cs
@@ -4,13 +4,12 @@ using System;
 using System.Collections.Generic;
 using System.Linq;
 using System.Threading.Tasks;
-using Horde.Build.Models;
 using Horde.Build.Server;
 using MongoDB.Bson;
 using MongoDB.Bson.Serialization.Attributes;
 using MongoDB.Driver;
 
-namespace Horde.Build.Collections.Impl
+namespace Horde.Build.Users
 {
 	/// <summary>
 	/// Collection of service account documents
diff --git a/Engine/Source/Programs/Horde/Horde.Build/Users/UserCollection.cs b/Engine/Source/Programs/Horde/Horde.Build/Users/UserCollection.cs
--- a/Engine/Source/Programs/Horde/Horde.Build/Users/UserCollection.cs
+++ b/Engine/Source/Programs/Horde/Horde.Build/Users/UserCollection.cs
@@ -5,7 +5,7 @@ using System.Collections.Generic;
 using System.Linq;
 using System.Security.Claims;
 using System.Threading.Tasks;
-using Horde.Build.Models;
+using Horde.Build.Jobs;
 using Horde.Build.Server;
 using Horde.Build.Utilities;
 using MongoDB.Bson;
@@ -13,7 +13,7 @@ using MongoDB.Bson.Serialization;
 using MongoDB.Bson.Serialization.Attributes;
 using MongoDB.Driver;
 
-namespace Horde.Build.Collections.Impl
+namespace Horde.Build.Users
 {
 	using JobId = ObjectId<IJob>;
 	using UserId = ObjectId<IUser>;
diff --git a/Engine/Source/Programs/Horde/Horde.Build/Users/UserCollectionV2.cs b/Engine/Source/Programs/Horde/Horde.Build/Users/UserCollectionV2.cs
--- a/Engine/Source/Programs/Horde/Horde.Build/Users/UserCollectionV2.cs
+++ b/Engine/Source/Programs/Horde/Horde.Build/Users/UserCollectionV2.cs
@@ -4,7 +4,7 @@ using System;
 using System.Collections.Generic;
 using System.Linq;
 using System.Threading.Tasks;
-using Horde.Build.Models;
+using Horde.Build.Jobs;
 using Horde.Build.Server;
 using Horde.Build.Utilities;
 using Microsoft.Extensions.Caching.Memory;
@@ -14,7 +14,7 @@ using MongoDB.Bson.Serialization;
 using MongoDB.Bson.Serialization.Attributes;
 using MongoDB.Driver;
 
-namespace Horde.Build.Collections.Impl
+namespace Horde.Build.Users
 {
 	using JobId = ObjectId<IJob>;
 	using UserId = ObjectId<IUser>;
diff --git a/Engine/Source/Programs/Horde/Horde.Build/Users/UserController.cs b/Engine/Source/Programs/Horde/Horde.Build/Users/UserController.cs
--- a/Engine/Source/Programs/Horde/Horde.Build/Users/UserController.cs
+++ b/Engine/Source/Programs/Horde/Horde.Build/Users/UserController.cs
@@ -3,15 +3,12 @@
 using System.Collections.Generic;
 using System.Linq;
 using System.Threading.Tasks;
-using Horde.Build.Api;
-using Horde.Build.Collections;
-using Horde.Build.Models;
-using Horde.Build.Services;
+using Horde.Build.Jobs;
 using Horde.Build.Utilities;
 using Microsoft.AspNetCore.Authorization;
 using Microsoft.AspNetCore.Mvc;
 
-namespace Horde.Build.Controllers
+namespace Horde.Build.Users
 {
 	using JobId = ObjectId<IJob>;
 	using UserId = ObjectId<IUser>;
diff --git a/Engine/Source/Programs/Horde/Horde.Build/Users/UserMessages.cs b/Engine/Source/Programs/Horde/Horde.Build/Users/UserMessages.cs
new file mode 100644
--- /dev/null
+++ b/Engine/Source/Programs/Horde/Horde.Build/Users/UserMessages.cs
@@ -0,0 +1,171 @@
+// Copyright Epic Games, Inc. All Rights Reserved.
+
+using System;
+using System.Collections.Generic;
+using System.Linq;
+using System.Text.Json;
+using EpicGames.Core;
+using MongoDB.Bson;
+
+namespace Horde.Build.Users
+{
+	/// <summary>
+	/// Response describing the current user
+	/// </summary>
+	public class GetUserResponse
+	{
+		/// <summary>
+		/// Id of the user
+		/// </summary>
+		public string Id { get; set; }
+
+		/// <summary>
+		/// Name of the user
+		/// </summary>
+		public string Name { get; set; }
+
+		/// <summary>
+		/// Avatar image URL (24px)
+		/// </summary>
+		public string? Image24 { get; set; }
+
+		/// <summary>
+		/// Avatar image URL (32px)
+		/// </summary>
+		public string? Image32 { get; set; }
+
+		/// <summary>
+		/// Avatar image URL (48px)
+		/// </summary>
+		public string? Image48 { get; set; }
+
+		/// <summary>
+		/// Avatar image URL (72px)
+		/// </summary>
+		public string? Image72 { get; set; }
+
+		/// <summary>
+		/// Email of the user
+		/// </summary>
+		public string? Email { get; set; }
+
+		/// <summary>
+		/// Claims for the user
+		/// </summary>
+		public List<UserClaim>? Claims { get; set; }
+
+		/// <summary>
+		/// Whether to enable experimental features for this user
+		/// </summary>
+		public bool? EnableExperimentalFeatures { get; set; }
+
+		/// <summary>
+		/// Settings for the dashboard
+		/// </summary>
+		public object? DashboardSettings { get; set; }
+
+		/// <summary>
+		/// List of pinned job ids
+		/// </summary>
+		public List<string>? PinnedJobIds { get; set; }
+
+		/// <summary>
+		/// Constructor
+		/// </summary>
+		public GetUserResponse(IUser user, IAvatar? avatar, IUserClaims? claims, IUserSettings? settings)
+		{
+			Id = user.Id.ToString();
+			Name = user.Name;
+			Email = user.Email;
+
+			Image24 = avatar?.Image24;
+			Image32 = avatar?.Image32;
+			Image48 = avatar?.Image48;
+			Image72 = avatar?.Image72;
+						
+			Claims = claims?.Claims.Select(x => new UserClaim(x)).ToList();
+
+			if (settings != null)
+			{
+				EnableExperimentalFeatures = settings.EnableExperimentalFeatures;
+				
+				DashboardSettings = BsonTypeMapper.MapToDotNetValue(settings.DashboardSettings);
+				PinnedJobIds = settings.PinnedJobIds.ConvertAll(x => x.ToString());
+			}
+		}
+	}
+
+	/// <summary>
+	/// Basic information about a user. May be embedded in other responses.
+	/// </summary>
+	public class GetThinUserInfoResponse
+	{
+		/// <summary>
+		/// Id of the user
+		/// </summary>
+		public string Id { get; set; }
+
+		/// <summary>
+		/// Name of the user
+		/// </summary>
+		public string Name { get; set; }
+
+		/// <summary>
+		/// The user's email address
+		/// </summary>
+		public string? Email { get; set; }
+
+		/// <summary>
+		/// The user login [DEPRECATED]
+		/// </summary>
+		internal string? Login { get; set; }
+
+		/// <summary>
+		/// Constructor
+		/// </summary>
+		/// <param name="user"></param>
+		public GetThinUserInfoResponse(IUser? user)
+		{
+			if (user == null)
+			{
+				Id = String.Empty;
+				Name = "(Unknown)";
+				Email = null;
+				Login = null;
+			}
+			else
+			{
+				Id = user.Id.ToString();
+				Name = user.Name;
+				Email = user.Email;
+				Login = user.Login;
+			}
+		}
+	}
+
+	/// <summary>
+	/// Request to update settings for a user
+	/// </summary>
+	public class UpdateUserRequest
+	{
+		/// <summary>
+		/// Whether to enable experimental features for this user
+		/// </summary>
+		public bool? EnableExperimentalFeatures { get; set; }
+
+		/// <summary>
+		/// New dashboard settings
+		/// </summary>
+		public JsonElement? DashboardSettings { get; set; }
+
+		/// <summary>
+		/// Job ids to add to the pinned list
+		/// </summary>
+		public List<string>? AddPinnedJobIds { get; set; }
+
+		/// <summary>
+		/// Jobs ids to remove from the pinned list
+		/// </summary>
+		public List<string>? RemovePinnedJobIds { get; set; }
+	}
+}
diff --git a/Engine/Source/Programs/Horde/Horde.Build/Users/UsersController.cs b/Engine/Source/Programs/Horde/Horde.Build/Users/UsersController.cs
--- a/Engine/Source/Programs/Horde/Horde.Build/Users/UsersController.cs
+++ b/Engine/Source/Programs/Horde/Horde.Build/Users/UsersController.cs
@@ -4,15 +4,11 @@ using System;
 using System.Collections.Generic;
 using System.Linq;
 using System.Threading.Tasks;
-using Horde.Build.Api;
-using Horde.Build.Collections;
-using Horde.Build.Models;
-using Horde.Build.Services;
 using Horde.Build.Utilities;
 using Microsoft.AspNetCore.Authorization;
 using Microsoft.AspNetCore.Mvc;
 
-namespace Horde.Build.Controllers
+namespace Horde.Build.Users
 {
 	using UserId = ObjectId<IUser>;
 
diff --git a/Engine/Source/Programs/Horde/Horde.Build/Utilities/HordeControllerBase.cs b/Engine/Source/Programs/Horde/Horde.Build/Utilities/HordeControllerBase.cs
--- a/Engine/Source/Programs/Horde/Horde.Build/Utilities/HordeControllerBase.cs
+++ b/Engine/Source/Programs/Horde/Horde.Build/Utilities/HordeControllerBase.cs
@@ -2,7 +2,11 @@
 
 using EpicGames.Core;
 using Horde.Build.Acls;
-using Horde.Build.Models;
+using Horde.Build.Agents;
+using Horde.Build.Agents.Leases;
+using Horde.Build.Jobs;
+using Horde.Build.Projects;
+using Horde.Build.Streams;
 using Horde.Build.Tools;
 using Microsoft.AspNetCore.Http;
 using Microsoft.AspNetCore.Mvc;
diff --git a/Engine/Source/Programs/Horde/Horde.Build/Utilities/SingletonDocument.cs b/Engine/Source/Programs/Horde/Horde.Build/Utilities/SingletonDocument.cs
--- a/Engine/Source/Programs/Horde/Horde.Build/Utilities/SingletonDocument.cs
+++ b/Engine/Source/Programs/Horde/Horde.Build/Utilities/SingletonDocument.cs
@@ -3,7 +3,6 @@
 using System;
 using System.Reflection;
 using System.Threading.Tasks;
-using Horde.Build.Models;
 using Horde.Build.Server;
 using MongoDB.Bson;
 
diff --git a/Engine/Source/Programs/Horde/Horde.Build/Utilities/TracingExtensions.cs b/Engine/Source/Programs/Horde/Horde.Build/Utilities/TracingExtensions.cs
--- a/Engine/Source/Programs/Horde/Horde.Build/Utilities/TracingExtensions.cs
+++ b/Engine/Source/Programs/Horde/Horde.Build/Utilities/TracingExtensions.cs
@@ -3,7 +3,7 @@
 using System;
 using System.Collections.Generic;
 using EpicGames.Core;
-using Horde.Build.Models;
+using Horde.Build.Jobs;
 using OpenTracing;
 
 namespace Horde.Build.Utilities
diff --git a/Engine/Source/Programs/Horde/Horde.Build/Utilities/ViewMap.cs b/Engine/Source/Programs/Horde/Horde.Build/Utilities/ViewMap.cs
--- a/Engine/Source/Programs/Horde/Horde.Build/Utilities/ViewMap.cs
+++ b/Engine/Source/Programs/Horde/Horde.Build/Utilities/ViewMap.cs
@@ -9,7 +9,7 @@ using EpicGames.Core;
 
 namespace Horde.Build.Utilities
 {
-	using P4 = Perforce.P4;
+	using P4 = global::Perforce.P4;
 
 	/// <summary>
 	/// Stores a mapping from one set of paths to another
