Horde: Remove some legacy classes, and a few final traces of legacy namespaces.

diff --git a/Engine/Source/Programs/Horde/Horde.Build/Api/Changes.cs b/Engine/Source/Programs/Horde/Horde.Build/Api/Changes.cs
deleted file mode 100644
diff --git a/Engine/Source/Programs/Horde/Horde.Build/Api/Config.cs b/Engine/Source/Programs/Horde/Horde.Build/Api/Config.cs
deleted file mode 100644
diff --git a/Engine/Source/Programs/Horde/Horde.Build/Commands/Config/SchemasCommand.cs b/Engine/Source/Programs/Horde/Horde.Build/Commands/Config/SchemasCommand.cs
--- a/Engine/Source/Programs/Horde/Horde.Build/Commands/Config/SchemasCommand.cs
+++ b/Engine/Source/Programs/Horde/Horde.Build/Commands/Config/SchemasCommand.cs
@@ -3,7 +3,6 @@
 using System;
 using System.Threading.Tasks;
 using EpicGames.Core;
-using Horde.Build.Controllers;
 using Horde.Build.Server;
 using Horde.Build.Utilities;
 using Microsoft.Extensions.Logging;
diff --git a/Engine/Source/Programs/Horde/Horde.Build/Configuration/ConfigUpdateService.cs b/Engine/Source/Programs/Horde/Horde.Build/Configuration/ConfigUpdateService.cs
--- a/Engine/Source/Programs/Horde/Horde.Build/Configuration/ConfigUpdateService.cs
+++ b/Engine/Source/Programs/Horde/Horde.Build/Configuration/ConfigUpdateService.cs
@@ -85,9 +85,6 @@ namespace Horde.Build.Configuration
 				_ticker = clock.AddSharedTicker<ConfigUpdateService>(TimeSpan.FromMinutes(1.0), TickLeaderAsync, logger);
 			}
 			_logger = logger;
-
-			// This will trigger if the local Horde.json user configuration is changed
-			_settings.OnChange(OnUserConfigUpdated);
 		}
 
 		/// <inheritdoc/>
@@ -449,331 +446,5 @@ namespace Horde.Build.Configuration
 				await UpdateConfigAsync(configUri);
 			}
 		}
-
-		//
-		// On premises configuration handling (aka Horde.json settings, though needs to have perforce global config handled too, checkout/modify/submit)
-		//
-
-		/// <summary>
-		/// Update the global configuration
-		/// </summary>
-		/// <param name="request"></param>
-		/// <returns></returns>
-		public async Task<bool> UpdateGlobalConfig(UpdateGlobalConfigRequest request)
-		{
-			if (_settings.CurrentValue.ConfigPath == null || !Path.IsPathRooted(_settings.CurrentValue.ConfigPath) || _settings.CurrentValue.ConfigPath.StartsWith("//", StringComparison.Ordinal))
-			{
-				throw new Exception($"Global config path must be rooted for service updates.  (Perforce paths are not currently supported for live updates), {_settings.CurrentValue.ConfigPath}");
-			}
-
-			FileReference globalConfigFile = new FileReference(_settings.CurrentValue.ConfigPath);
-			DirectoryReference globalConfigDirectory = globalConfigFile.Directory;
-
-			// make sure the directory exists
-			if (!DirectoryReference.Exists(globalConfigDirectory))
-			{
-				DirectoryReference.CreateDirectory(globalConfigDirectory);
-			}
-
-			bool configDirty = false;
-
-			// projects
-			if (request.ProjectsJson != null)
-			{
-				configDirty = true;
-
-				// write out the projects
-				foreach (KeyValuePair<string, string> project in request.ProjectsJson)
-				{
-					await FileReference.WriteAllTextAsync(FileReference.Combine(globalConfigDirectory, project.Key), project.Value);
-
-					if (request.ProjectLogo != null)
-					{
-						byte[] bytes = Convert.FromBase64String(request.ProjectLogo);
-						await FileReference.WriteAllBytesAsync(FileReference.Combine(globalConfigDirectory, project.Key.Replace(".json", ".png", StringComparison.OrdinalIgnoreCase)), bytes);
-					}
-				}
-			}
-
-			// streams
-			if (request.StreamsJson != null)
-			{
-				configDirty = true;
-
-				// write out the streams
-				foreach (KeyValuePair<string, string> stream in request.StreamsJson)
-				{
-					await FileReference.WriteAllTextAsync(FileReference.Combine(globalConfigDirectory, stream.Key), stream.Value);
-				}
-			}
-
-			// update global config path
-			if (!String.IsNullOrEmpty(request.GlobalsJson))
-			{
-				configDirty = true;
-
-				await FileReference.WriteAllTextAsync(globalConfigFile, request.GlobalsJson);
-			}
-
-			if (configDirty == true)
-			{
-				await UpdateConfigAsync(new Uri("file://" + globalConfigFile.ToString()));
-			}
-
-			// create the default pool 
-			if (request.DefaultPoolName != null)
-			{
-				PoolId poolIdValue = PoolId.Sanitize(request.DefaultPoolName);
-
-				IPool? pool = await _poolService.GetPoolAsync(poolIdValue);
-				if (pool == null)
-				{
-					await _poolService.CreatePoolAsync(request.DefaultPoolName, condition: "OSFamily == 'Windows'", properties: new Dictionary<string, string>() { ["Color"] = "0" });
-				}
-			}
-
-			return true;
-
-		}
-
-		/// <summary>
-		/// Update the server settings
-		/// </summary>
-		/// <param name="request"></param>
-		/// <returns></returns>
-		public async Task<ServerUpdateResponse> UpdateServerSettings(UpdateServerSettingsRequest request)
-		{
-
-			ServerUpdateResponse response = new ServerUpdateResponse();
-
-			Dictionary<string, object> newUserSettings = new Dictionary<string, object>();
-
-			try
-			{
-
-				if (request.Settings == null || request.Settings.Count == 0)
-				{
-					return response;
-				}
-
-				if (_userConfigUpdated != null)
-				{
-					response.Errors.Add("User configuation already being updated");
-					return response;
-				}
-
-				// Load the current user configuration
-				IConfiguration config = new ConfigurationBuilder()
-					.AddJsonFile(Program.UserConfigFile.FullName, optional: true)
-					.Build();
-
-				ServerSettings userSettings = new ServerSettings();
-				config.GetSection("Horde").Bind(userSettings);
-
-				// Figure out current and new settings
-				HashSet<string> userProps = new HashSet<string>(StringComparer.OrdinalIgnoreCase);
-
-				if (FileReference.Exists(Program.UserConfigFile))
-				{
-					byte[] data = await FileReference.ReadAllBytesAsync(Program.UserConfigFile);
-
-					JsonDocument document = JsonDocument.Parse(data);
-					JsonElement hordeElement;
-
-					if (document.RootElement.TryGetProperty("Horde", out hordeElement))
-					{
-						foreach (JsonProperty property in hordeElement.EnumerateObject())
-						{
-							userProps.Add(property.Name);
-						}
-					}
-				}
-
-				foreach (string name in request.Settings.Keys)
-				{
-					userProps.Add(name);
-				}
-
-				// Apply the changes from the request
-				foreach (KeyValuePair<string, object> pair in request.Settings)
-				{
-					PropertyInfo? property = userSettings.GetType().GetProperty(pair.Key, BindingFlags.IgnoreCase | BindingFlags.Public | BindingFlags.Instance);
-					if (property == null)
-					{
-						response.Errors.Add($"Horde configuration property {pair.Key} does not exist when reading server settings");
-						_logger.LogError("Horde configuration property {Key} does not exist when reading server settings", pair.Key);
-						continue;
-					}
-
-					if (property.SetMethod == null)
-					{
-						response.Errors.Add($"Horde configuration property {pair.Key} does not have a set method");
-						_logger.LogError("Horde configuration property {Key} does not have a set method", pair.Key);
-						continue;
-					}
-
-					// explicit setting removal
-					if (pair.Value == null)
-					{
-						userProps.Remove(pair.Key);
-						continue;
-					}
-
-					JsonElement element = (JsonElement)pair.Value;
-
-					object? value = null;
-
-					switch (element.ValueKind)
-					{
-						case JsonValueKind.True:
-							value = true;
-							break;
-						case JsonValueKind.False:
-							value = false;
-							break;
-						case JsonValueKind.Number:
-							value = element.GetDouble();
-							break;
-						case JsonValueKind.String:
-							value = element.GetString();
-							break;
-					}
-
-					// unable to map type
-					if (value == null)
-					{
-						response.Errors.Add($"Unable to map type for Property {property.Name}");
-						continue;
-					}
-
-					// handle common conversions
-					if (property.GetType() != value.GetType())
-					{
-						try
-						{
-							value = Convert.ChangeType(value, property.PropertyType, CultureInfo.CurrentCulture);
-						}
-						catch (Exception ex)
-						{
-							response.Errors.Add($"Property {property.Name} raised exception during conversion, {ex.Message}");
-							continue;
-						}
-
-						if (value == null)
-						{
-							response.Errors.Add($"Property {property.Name} had null value on conversion");
-							continue;
-						}
-					}
-					else
-					{
-						response.Errors.Add($"Property {property.Name} is not assignable to {value}");
-						continue;
-					}
-
-
-					//  Set the value, providing some validation
-					try
-					{
-						property.SetMethod.Invoke(userSettings, new object[] { value });
-					}
-					catch (Exception ex)
-					{
-						response.Errors.Add($"Exception updating property {pair.Key}, {ex.Message}");
-					}
-				}
-
-				// Construct the new user settings
-				foreach (string name in userProps)
-				{
-					PropertyInfo? property = userSettings.GetType().GetProperty(name, BindingFlags.IgnoreCase | BindingFlags.Public | BindingFlags.Instance);
-					if (property == null)
-					{
-						response.Errors.Add($"Horde configuration property {name} does not exist when writing server settings");
-						_logger.LogError("Horde configuration property {Name} does not exist when writing server settings", name);
-
-						continue;
-					}
-
-					if (property.GetMethod == null)
-					{
-						response.Errors.Add($"Horde configuration property {name} does not have a get method");
-						_logger.LogError("Horde configuration property {Name} does not have a get method", name);
-						continue;
-					}
-
-					object? result = property.GetMethod.Invoke(userSettings, null);
-
-					if (result == null)
-					{
-						response.Errors.Add($"Horde configuration property {name} was null while writing and should have already been filtered out");
-						_logger.LogError("Horde configuration property {Name} was null while writing and should have already been filtered out", name);
-						continue;
-					}
-
-					newUserSettings.Add(property.Name, result);
-
-				}
-			}
-			catch (Exception ex)
-			{
-				response.Errors.Add($"Exception while updating settings: {ex.Message}");
-				_logger.LogError(ex, "{Error}", response.Errors.Last());
-			}
-
-			if (response.Errors.Count != 0)
-			{
-				return response;
-			}
-
-			Dictionary<string, object> newLocalSettings = new Dictionary<string, object>();
-			newLocalSettings["Horde"] = newUserSettings;
-
-			try
-			{
-				_userConfigUpdated = new TaskCompletionSource<bool>();
-
-				// This will trigger a setting update as the user config json is set to reload on change
-				try
-				{
-					await FileReference.WriteAllBytesAsync(Program.UserConfigFile, JsonSerializer.SerializeToUtf8Bytes(newLocalSettings, new JsonSerializerOptions { WriteIndented = true, DefaultIgnoreCondition = JsonIgnoreCondition.WhenWritingNull }));
-				}
-				catch (Exception ex)
-				{
-					response.Errors.Add($"Unable to serialize json settings to {Program.UserConfigFile}, {ex.Message}");
-					_logger.LogError(ex, "Unable to serialize json settings to {ConfigFile}, {Message}", Program.UserConfigFile.ToString(), ex.Message);
-				}
-
-				if (response.Errors.Count == 0)
-				{
-					if (await Task.WhenAny(_userConfigUpdated.Task, Task.Delay(5000)) != _userConfigUpdated.Task)
-					{
-						response.Errors.Add("Server update timed out while awaiting write");
-						_logger.LogError("Server update timed out after writing config");
-					}
-				}
-			}
-			finally
-			{
-				_userConfigUpdated = null;
-			}
-
-			return response;
-
-		}
-		void OnUserConfigUpdated(ServerSettings settings, string name)
-		{
-			NumUserConfigUpdates++;
-			_userConfigUpdated?.SetResult(true);
-		}
-
-		private TaskCompletionSource<bool>? _userConfigUpdated;
-
-		/// <summary>
-		/// The number of server configuration updates that have been made while the server is running
-		/// This is used to detect whether server may need to be restarted
-		/// </summary>
-		public int NumUserConfigUpdates { get; private set; } = 0;
-
 	}
 }
diff --git a/Engine/Source/Programs/Horde/Horde.Build/Controllers/ErrorController.cs b/Engine/Source/Programs/Horde/Horde.Build/Controllers/ErrorController.cs
deleted file mode 100644
diff --git a/Engine/Source/Programs/Horde/Horde.Build/Perforce/ChangeMessages.cs b/Engine/Source/Programs/Horde/Horde.Build/Perforce/ChangeMessages.cs
new file mode 100644
--- /dev/null
+++ b/Engine/Source/Programs/Horde/Horde.Build/Perforce/ChangeMessages.cs
@@ -0,0 +1,90 @@
+// Copyright Epic Games, Inc. All Rights Reserved.
+
+using System.Collections.Generic;
+using Horde.Build.Perforce;
+using Horde.Build.Users;
+
+namespace Horde.Build.Perforce
+{
+	/// <summary>
+	/// Information about a submitted changelist
+	/// </summary>
+	public class GetChangeSummaryResponse
+	{
+		/// <summary>
+		/// The source changelist number
+		/// </summary>
+		public int Number { get; set; }
+
+		/// <summary>
+		/// Name of the user that authored this change [DEPRECATED]
+		/// </summary>
+		public string Author { get; set; }
+
+		/// <summary>
+		/// Information about the change author
+		/// </summary>
+		public GetThinUserInfoResponse AuthorInfo { get; set; }
+
+		/// <summary>
+		/// The description text
+		/// </summary>
+		public string Description { get; set; }
+
+		/// <summary>
+		/// Constructor
+		/// </summary>
+		/// <param name="changeSummary">The commit to construct from</param>
+		public GetChangeSummaryResponse(ChangeSummary changeSummary)
+		{
+			Number = changeSummary.Number;
+			Author = changeSummary.Author.Name;
+			AuthorInfo = new GetThinUserInfoResponse(changeSummary.Author);
+			Description = changeSummary.Description;
+		}
+	}
+
+	/// <summary>
+	/// Information about a submitted changelist
+	/// </summary>
+	public class GetChangeDetailsResponse
+	{
+		/// <summary>
+		/// The source changelist number
+		/// </summary>
+		public int Number { get; set; }
+
+		/// <summary>
+		/// Name of the user that authored this change [DEPRECATED]
+		/// </summary>
+		public string Author { get; set; }
+
+		/// <summary>
+		/// Information about the user that authored this change
+		/// </summary>
+		public GetThinUserInfoResponse AuthorInfo { get; set; }
+
+		/// <summary>
+		/// The description text
+		/// </summary>
+		public string Description { get; set; }
+
+		/// <summary>
+		/// List of files that were modified, relative to the stream base
+		/// </summary>
+		public List<string> Files { get; set; }
+
+		/// <summary>
+		/// Constructor
+		/// </summary>
+		/// <param name="changeDetails">The commit to construct from</param>
+		public GetChangeDetailsResponse(ChangeDetails changeDetails)
+		{
+			Number = changeDetails.Number;
+			Author = changeDetails.Author.Name;
+			AuthorInfo = new GetThinUserInfoResponse(changeDetails.Author);
+			Description = changeDetails.Description;
+			Files = changeDetails.Files.ConvertAll(x => x.Path);
+		}
+	}
+}
diff --git a/Engine/Source/Programs/Horde/Horde.Build/Server/DashboardController.cs b/Engine/Source/Programs/Horde/Horde.Build/Server/DashboardController.cs
--- a/Engine/Source/Programs/Horde/Horde.Build/Server/DashboardController.cs
+++ b/Engine/Source/Programs/Horde/Horde.Build/Server/DashboardController.cs
@@ -3,7 +3,6 @@
 using System;
 using System.Text;
 using System.Threading.Tasks;
-using Horde.Build.Controllers;
 using Microsoft.AspNetCore.Authentication;
 using Microsoft.AspNetCore.Authentication.Cookies;
 using Microsoft.AspNetCore.Authorization;
