Horde: Rename files containing API messages to not have a dotted name, to stop nesting within solution explorer.

diff --git a/Engine/Source/Programs/Horde/Horde.Build/Agents/AgentMessages.cs b/Engine/Source/Programs/Horde/Horde.Build/Agents/AgentMessages.cs
new file mode 100644
--- /dev/null
+++ b/Engine/Source/Programs/Horde/Horde.Build/Agents/AgentMessages.cs
@@ -0,0 +1,500 @@
+// Copyright Epic Games, Inc. All Rights Reserved.
+
+using System;
+using System.Collections.Generic;
+using System.ComponentModel.DataAnnotations;
+using System.Linq;
+using Horde.Build.Acls;
+using Horde.Build.Agents;
+using EpicGames.Core;
+using HordeCommon;
+using Horde.Build.Agents.Leases;
+using Horde.Build.Agents.Sessions;
+
+namespace Horde.Build.Agents
+{
+	/// <summary>
+	/// Parameters to register a new agent
+	/// </summary>
+	public class CreateAgentRequest
+	{
+		/// <summary>
+		/// Friendly name for the agent
+		/// </summary>
+		[Required]
+		public string Name { get; set; } = null!; // Enforced by [required] attribute
+
+		/// <summary>
+		/// Whether the agent is currently enabled
+		/// </summary>
+		public bool Enabled { get; set; } = true;
+
+		/// <summary>
+		/// Whether the agent is ephemeral (ie. should not be shown when inactive)
+		/// </summary>
+		public bool Ephemeral { get; set; }
+
+		/// <summary>
+		/// Per-agent override for the desired agent software channel
+		/// </summary>
+		public string? Channel { get; set; }
+
+		/// <summary>
+		/// Pools for this agent
+		/// </summary>
+		public List<string>? Pools { get; set; }
+	}
+
+	/// <summary>
+	/// Response from creating an agent
+	/// </summary>
+	public class CreateAgentResponse
+	{
+		/// <summary>
+		/// Unique id for the new agent
+		/// </summary>
+		public string Id { get; set; }
+
+		/// <summary>
+		/// Constructor
+		/// </summary>
+		/// <param name="id">Unique id for this agent</param>
+		public CreateAgentResponse(string id)
+		{
+			Id = id;
+		}
+	}
+
+	/// <summary>
+	/// Parameters to update an agent
+	/// </summary>
+	public class UpdateAgentRequest
+	{
+		/// <summary>
+		/// Whether the agent is currently enabled
+		/// </summary>
+		public bool? Enabled { get; set; }
+
+		/// <summary>
+		/// Whether the agent is ephemeral
+		/// </summary>
+		public bool? Ephemeral { get; set; }
+
+		/// <summary>
+		/// Request a conform be performed using the current agent set
+		/// </summary>
+		public bool? RequestConform { get; set; }
+
+		/// <summary>
+		/// Request that a full conform be performed, removing all intermediate files
+		/// </summary>
+		public bool? RequestFullConform { get; set; }
+
+		/// <summary>
+		/// Request the machine be restarted
+		/// </summary>
+		public bool? RequestRestart { get; set; }
+
+		/// <summary>
+		/// Request the machine be shut down
+		/// </summary>
+		public bool? RequestShutdown { get; set; }
+
+		/// <summary>
+		/// Per-agent override for the desired agent software channel
+		/// </summary>
+		public string? Channel { get; set; }
+
+		/// <summary>
+		/// Pools for this agent
+		/// </summary>
+		public List<string>? Pools { get; set; }
+
+		/// <summary>
+		/// New ACL for this agent
+		/// </summary>
+		public UpdateAclRequest? Acl { get; set; }
+		
+		/// <summary>
+		/// New comment
+		/// </summary>
+		public string? Comment { get; set; }
+	}
+
+	/// <summary>
+	/// Response for queries to find a particular lease within an agent
+	/// </summary>
+	public class GetAgentLeaseResponse
+	{
+		/// <summary>
+		/// Identifier for the lease
+		/// </summary>
+		public string Id { get; set; }
+
+		/// <summary>
+		/// The agent id
+		/// </summary>
+		public string? AgentId { get; set; }
+
+		/// <summary>
+		/// Cost of this agent, per hour
+		/// </summary>
+		public double? AgentRate { get; set; }
+
+		/// <summary>
+		/// Name of the lease
+		/// </summary>
+		public string Name { get; set; }
+
+		/// <summary>
+		/// Log id for this lease
+		/// </summary>
+		public string? LogId { get; set; }
+
+		/// <summary>
+		/// Time at which the lease started (UTC)
+		/// </summary>
+		public DateTime StartTime { get; set; }
+
+		/// <summary>
+		/// Time at which the lease started (UTC)
+		/// </summary>
+		public DateTime? FinishTime { get; set; }
+
+		/// <summary>
+		/// Whether this lease has started executing on the agent yet
+		/// </summary>
+		public bool Executing { get; set; }
+
+		/// <summary>
+		/// Details of the payload being executed
+		/// </summary>
+		public Dictionary<string, string>? Details { get; set; }
+
+		/// <summary>
+		/// Outcome of the lease
+		/// </summary>
+		public LeaseOutcome? Outcome { get; set; }
+
+		/// <summary>
+		/// State of the lease (for AgentLeases)
+		/// </summary>
+		public LeaseState? State { get; set; }
+
+		/// <summary>
+		/// Constructor
+		/// </summary>
+		/// <param name="lease">The lease to initialize from</param>
+		/// <param name="details">The payload details</param>
+		public GetAgentLeaseResponse(AgentLease lease, Dictionary<string, string>? details)
+		{
+			Id = lease.Id.ToString();
+			Name = lease.Name;
+			LogId = lease.LogId?.ToString();
+			State = lease.State;
+			StartTime = lease.StartTime;
+			Executing = lease.Active;
+			FinishTime = lease.ExpiryTime;
+			Details = details;
+		}
+
+		/// <summary>
+		/// Converts this lease to a response object
+		/// </summary>
+		/// <param name="lease">The lease to initialize from</param>
+		/// <param name="details">The payload details</param>
+		/// <param name="agentRate">Rate for running this agent</param>
+		public GetAgentLeaseResponse(ILease lease, Dictionary<string, string>? details, double? agentRate)
+		{
+			Id = lease.Id.ToString();
+			AgentId = lease.AgentId.ToString();
+			AgentRate = agentRate;
+			Name = lease.Name;
+			LogId = lease.LogId?.ToString();
+			StartTime = lease.StartTime;
+			Executing = (lease.FinishTime == null);
+			FinishTime = lease.FinishTime;
+			Details = details;
+			Outcome = lease.Outcome;
+		}
+	}
+
+	/// <summary>
+	/// Information about an agent session
+	/// </summary>
+	public class GetAgentSessionResponse
+	{
+		/// <summary>
+		/// Unique id for this session
+		/// </summary>
+		[Required]
+		public string Id { get; set; }
+
+		/// <summary>
+		/// Start time for this session
+		/// </summary>
+		[Required]
+		public DateTime StartTime { get; set; }
+
+		/// <summary>
+		/// Finishing time for this session
+		/// </summary>
+		public DateTime? FinishTime { get; set; }
+
+		/// <summary>
+		/// Properties of this agent
+		/// </summary>
+		public List<string>? Properties { get; set; }
+
+		/// <summary>
+		/// Version of the software running during this session
+		/// </summary>
+		public string? Version { get; set; }
+
+		/// <summary>
+		/// Constructor
+		/// </summary>
+		/// <param name="session">The session to construct from</param>
+		public GetAgentSessionResponse(ISession session)
+		{
+			Id = session.Id.ToString();
+			StartTime = session.StartTime;
+			FinishTime = session.FinishTime;
+			Properties = (session.Properties != null) ? new List<string>(session.Properties) : null;
+			Version = session.Version?.ToString();
+		}
+	}
+
+	/// <summary>
+	/// Information about a workspace synced on an agent
+	/// </summary>
+	public class GetAgentWorkspaceResponse
+	{
+		/// <summary>
+		/// The Perforce server and port to connect to
+		/// </summary>
+		public string? Cluster { get; set; }
+
+		/// <summary>
+		/// User to log into Perforce with (eg. buildmachine)
+		/// </summary>
+		public string? UserName { get; set; }
+
+		/// <summary>
+		/// Identifier to distinguish this workspace from other workspaces
+		/// </summary>
+		public string Identifier { get; set; }
+
+		/// <summary>
+		/// The stream to sync
+		/// </summary>
+		public string Stream { get; set; }
+
+		/// <summary>
+		/// Custom view for the workspace
+		/// </summary>
+		public List<string>? View { get; set; }
+
+		/// <summary>
+		/// Whether to use an incremental workspace
+		/// </summary>
+		public bool BIncremental { get; set; }
+
+		/// <summary>
+		/// Constructor
+		/// </summary>
+		/// <param name="workspace">The workspace to construct from</param>
+		public GetAgentWorkspaceResponse(AgentWorkspace workspace)
+		{
+			Cluster = workspace.Cluster;
+			UserName = workspace.UserName;
+			Identifier = workspace.Identifier;
+			Stream = workspace.Stream;
+			View = workspace.View;
+			BIncremental = workspace.BIncremental;
+		}
+	}
+
+	/// <summary>
+	/// Information about an agent
+	/// </summary>
+	public class GetAgentResponse
+	{
+		/// <summary>
+		/// The agent's unique ID
+		/// </summary>
+		public string Id { get; set; }
+
+		/// <summary>
+		/// Friendly name of the agent
+		/// </summary>
+		public string Name { get; set; }
+
+		/// <summary>
+		/// Whether the agent is currently enabled
+		/// </summary>
+		public bool Enabled { get; set; }
+
+		/// <summary>
+		/// Cost estimate per-hour for this agent
+		/// </summary>
+		public double? Rate { get; set; }
+
+		/// <summary>
+		/// The current session id
+		/// </summary>
+		public string? SessionId { get; set; }
+
+		/// <summary>
+		/// Whether the agent is ephemeral
+		/// </summary>
+		public bool Ephemeral { get; set; }
+
+		/// <summary>
+		/// Whether the agent is currently online
+		/// </summary>
+		public bool Online { get; set; }
+
+		/// <summary>
+		/// Whether this agent has expired
+		/// </summary>
+		public bool Deleted { get; set; }
+
+		/// <summary>
+		/// Whether a conform job is pending
+		/// </summary>
+		public bool PendingConform { get; set; }
+
+		/// <summary>
+		/// Whether a full conform job is pending
+		/// </summary>
+		public bool PendingFullConform { get; set; }
+
+		/// <summary>
+		/// Whether a restart is pending
+		/// </summary>
+		public bool PendingRestart { get; set; }
+		
+		/// <summary>
+		/// Whether a shutdown is pending
+		/// </summary>
+		public bool PendingShutdown { get; set; }
+
+		/// <summary>
+		/// The reason for the last shutdown
+		/// </summary>
+		public string LastShutdownReason { get; set; }
+
+		/// <summary>
+		/// Last time a conform was attempted
+		/// </summary>
+		public DateTime LastConformTime { get; set; }
+
+		/// <summary>
+		/// Number of times a conform has been attempted
+		/// </summary>
+		public int? ConformAttemptCount { get; set; }
+
+		/// <summary>
+		/// Last time a conform was attempted
+		/// </summary>
+		public DateTime? NextConformTime { get; set; }
+
+		/// <summary>
+		/// The current client version
+		/// </summary>
+		public string? Version { get; set; }
+
+		/// <summary>
+		/// Per-agent override for the desired agent software channel
+		/// </summary>
+		public string? Channel { get; set; }
+
+		/// <summary>
+		/// Properties for the agent
+		/// </summary>
+		public List<string> Properties { get; set; }
+
+		/// <summary>
+		/// Resources for the agent
+		/// </summary>
+		public Dictionary<string, int> Resources { get; set; }
+
+		/// <summary>
+		/// Last update time of this agent
+		/// </summary>
+		public DateTime? UpdateTime { get; set; }
+
+		/// <summary>
+		/// Pools for this agent
+		/// </summary>
+		public List<string>? Pools { get; set; }
+
+		/// <summary>
+		/// List of workspaces currently synced to this machine
+		/// </summary>
+		public List<GetAgentWorkspaceResponse> Workspaces { get; set; } = new List<GetAgentWorkspaceResponse>();
+
+		/// <summary>
+		/// Capabilities of this agent
+		/// </summary>
+		public object? Capabilities { get; }
+
+		/// <summary>
+		/// Array of active leases.
+		/// </summary>
+		public List<GetAgentLeaseResponse> Leases { get; }
+
+		/// <summary>
+		/// Per-object permissions
+		/// </summary>
+		public GetAclResponse? Acl { get; set; }
+		
+		/// <summary>
+		/// Comment for this agent
+		/// </summary>
+		public string? Comment { get; set; }
+
+		/// <summary>
+		/// Constructor
+		/// </summary>
+		/// <param name="agent">The agent to construct from</param>
+		/// <param name="leases">Active leases</param>
+		/// <param name="rate">Rate for this agent</param>
+		/// <param name="bIncludeAcl">Whether to include the ACL in the response</param>
+		public GetAgentResponse(IAgent agent, List<GetAgentLeaseResponse> leases, double? rate, bool bIncludeAcl)
+		{
+			Id = agent.Id.ToString();
+			Name = agent.Id.ToString();
+			Enabled = agent.Enabled;
+			Rate = rate;
+			Properties = new List<string>(agent.Properties);
+			Resources = new Dictionary<string, int>(agent.Resources);
+			SessionId = agent.SessionId?.ToString();
+			Online = agent.IsSessionValid(DateTime.UtcNow);
+			Deleted = agent.Deleted;
+			PendingConform = agent.RequestConform;
+			PendingFullConform = agent.RequestFullConform;
+			PendingRestart = agent.RequestRestart;
+			PendingShutdown = agent.RequestShutdown;
+			LastShutdownReason = agent.LastShutdownReason ?? "Unknown";
+			LastConformTime = agent.LastConformTime;
+			NextConformTime =   agent.LastConformTime;
+			ConformAttemptCount = agent.ConformAttemptCount;
+			Version = agent.Version?.ToString() ?? "Unknown";
+			if(agent.Channel != null)
+			{
+				Version += $" ({agent.Channel})";
+			}
+			Version = agent.Version?.ToString();
+			Channel = agent.Channel?.ToString();
+			UpdateTime = agent.UpdateTime;
+			Pools = agent.GetPools().Select(x => x.ToString()).ToList();
+			Workspaces = agent.Workspaces.ConvertAll(x => new GetAgentWorkspaceResponse(x));
+			Capabilities = new { Devices = new[] { new { agent.Properties, agent.Resources } } };
+			Leases = leases;
+			Acl = (bIncludeAcl && agent.Acl != null) ? new GetAclResponse(agent.Acl) : null;
+			Comment = agent.Comment;
+		}
+	}
+}
diff --git a/Engine/Source/Programs/Horde/Horde.Build/Agents/AgentsController.Messages.cs b/Engine/Source/Programs/Horde/Horde.Build/Agents/AgentsController.Messages.cs
deleted file mode 100644
diff --git a/Engine/Source/Programs/Horde/Horde.Build/Agents/Leases/LeaseMessages.cs b/Engine/Source/Programs/Horde/Horde.Build/Agents/Leases/LeaseMessages.cs
new file mode 100644
--- /dev/null
+++ b/Engine/Source/Programs/Horde/Horde.Build/Agents/Leases/LeaseMessages.cs
@@ -0,0 +1,16 @@
+// Copyright Epic Games, Inc. All Rights Reserved.
+
+namespace Horde.Build.Agents.Leases
+{
+
+	/// <summary>
+	/// Updates an existing lease
+	/// </summary>
+	public class UpdateLeaseRequest
+	{
+		/// <summary>
+		/// Mark this lease as aborted
+		/// </summary>
+		public bool? Aborted { get; set; }
+	}
+}
diff --git a/Engine/Source/Programs/Horde/Horde.Build/Agents/Leases/LeasesController.Messages.cs b/Engine/Source/Programs/Horde/Horde.Build/Agents/Leases/LeasesController.Messages.cs
deleted file mode 100644
diff --git a/Engine/Source/Programs/Horde/Horde.Build/Agents/Pools/PoolMessages.cs b/Engine/Source/Programs/Horde/Horde.Build/Agents/Pools/PoolMessages.cs
new file mode 100644
--- /dev/null
+++ b/Engine/Source/Programs/Horde/Horde.Build/Agents/Pools/PoolMessages.cs
@@ -0,0 +1,332 @@
+// Copyright Epic Games, Inc. All Rights Reserved.
+
+using System.Collections.Generic;
+using System.ComponentModel.DataAnnotations;
+using System.Linq;
+using EpicGames.Horde.Common;
+using Horde.Build.Agents.Fleet;
+
+namespace Horde.Build.Agents.Pools
+{
+	/// <see cref="LeaseUtilizationSettings" />
+	public class LeaseUtilizationSettingsMessage
+	{
+		/// <summary>
+		/// Construct a public REST API representation from the internal one
+		/// </summary>
+		/// <param name="settings"></param>
+#pragma warning disable IDE0060 // Remove unused parameter
+		public LeaseUtilizationSettingsMessage(LeaseUtilizationSettings settings)
+#pragma warning restore IDE0060 // Remove unused parameter
+		{
+		}
+		
+		/// <summary>
+		/// Convert this public REST API object to an internal representation
+		/// </summary>
+		/// <returns></returns>
+		public Horde.Build.Agents.Fleet.LeaseUtilizationSettings Convert()
+		{
+			return new Horde.Build.Agents.Fleet.LeaseUtilizationSettings();
+		}
+	}
+	
+	/// <see cref="JobQueueSettings" />
+	public class JobQueueSettingsMessage
+	{
+		/// <summary>
+		/// Factor translating queue size to additional agents to grow the pool with
+		/// The result is always rounded up to nearest integer. 
+		/// Example: if there are 20 jobs in queue, a factor 0.25 will result in 5 new agents being added (20 * 0.25)
+		/// </summary>
+		public double ScaleOutFactor { get; set; }
+
+		/// <summary>
+		/// Factor by which to shrink the pool size with when queue is empty
+		/// The result is always rounded up to nearest integer.
+		/// Example: when the queue size is zero, a default value of 0.9 will shrink the pool by 10% (current agent count * 0.9)
+		/// </summary>
+		public double ScaleInFactor { get; set; }
+
+		/// <summary>
+		/// Construct a public REST API representation from the internal one
+		/// </summary>
+		/// <param name="settings"></param>
+		public JobQueueSettingsMessage(JobQueueSettings settings)
+		{
+			ScaleInFactor = settings.ScaleInFactor;
+			ScaleOutFactor = settings.ScaleOutFactor;
+		}
+		
+		/// <summary>
+		/// Convert this public REST API object to an internal representation
+		/// </summary>
+		/// <returns></returns>
+		public JobQueueSettings Convert()
+		{
+			return new JobQueueSettings(ScaleOutFactor, ScaleInFactor);
+		}
+	}
+	
+	/// <summary>
+	/// Parameters to create a new pool
+	/// </summary>
+	public class CreatePoolRequest
+	{
+		/// <summary>
+		/// Name for the new pool
+		/// </summary>
+		[Required]
+		public string Name { get; set; } = null!;
+
+		/// <summary>
+		/// Condition to satisfy for agents to be included in this pool
+		/// </summary>
+		public Condition? Condition { get; set; }
+
+		/// <summary>
+		/// Whether to enable autoscaling for this pool
+		/// </summary>
+		public bool? EnableAutoscaling { get; set; }
+		
+		/// <summary>
+		/// Cooldown time between scale-out events in seconds
+		/// </summary>
+		public int? ScaleOutCooldown { get; set; }
+		
+		/// <summary>
+		/// Cooldown time between scale-in events in seconds
+		/// </summary>
+		public int? ScaleInCooldown { get; set; }
+		
+		/// <summary>
+		/// Pool sizing strategy
+		/// </summary>
+		public PoolSizeStrategy? SizeStrategy { get; set; }
+		
+		/// <summary>
+		/// Settings for lease utilization pool sizing strategy (if used)
+		/// </summary>
+		public LeaseUtilizationSettingsMessage? LeaseUtilizationSettings { get; set; }
+		
+		/// <summary>
+		/// Settings for job queue pool sizing strategy (if used) 
+		/// </summary>
+		public JobQueueSettingsMessage? JobQueueSettings { get; set; }
+
+		/// <summary>
+		/// The minimum nunmber of agents to retain in this pool
+		/// </summary>
+		public int? MinAgents { get; set; }
+
+		/// <summary>
+		/// The minimum number of idle agents in this pool, if autoscaling is enabled
+		/// </summary>
+		public int? NumReserveAgents { get; set; }
+
+		/// <summary>
+		/// Properties for the new pool
+		/// </summary>
+		public Dictionary<string, string>? Properties { get; set; }
+	}
+
+	/// <summary>
+	/// Response from creating a new pool
+	/// </summary>
+	public class CreatePoolResponse
+	{
+		/// <summary>
+		/// Unique id for the new pool
+		/// </summary>
+		public string Id { get; set; }
+
+		/// <summary>
+		/// Constructor
+		/// </summary>
+		/// <param name="id">Unique id for the new pool</param>
+		public CreatePoolResponse(string id)
+		{
+			Id = id;
+		}
+	}
+
+	/// <summary>
+	/// Parameters to update a pool
+	/// </summary>
+	public class UpdatePoolRequest
+	{
+		/// <summary>
+		/// Optional new name for the pool
+		/// </summary>
+		public string? Name { get; set; }
+
+		/// <summary>
+		/// Requirements for this pool
+		/// </summary>
+		public Condition? Condition { get; set; }
+
+		/// <summary>
+		/// Whether to enable autoscaling for this pool
+		/// </summary>
+		public bool? EnableAutoscaling { get; set; }
+		
+		/// <summary>
+		/// Cooldown time between scale-out events in seconds
+		/// </summary>
+		public int? ScaleOutCooldown { get; set; }
+		
+		/// <summary>
+		/// Cooldown time between scale-in events in seconds
+		/// </summary>
+		public int? ScaleInCooldown { get; set; }
+		
+		/// <summary>
+		/// Pool sizing strategy
+		/// </summary>
+		public PoolSizeStrategy? SizeStrategy { get; set; }
+
+		/// <summary>
+		/// Set pool to use default strategy
+		/// </summary>
+		public bool? UseDefaultStrategy { get; set; }
+
+		/// <summary>
+		/// Settings for lease utilization pool sizing strategy (if used)
+		/// </summary>
+		public LeaseUtilizationSettings? LeaseUtilizationSettings { get; set; }
+		
+		/// <summary>
+		/// Settings for job queue pool sizing strategy (if used) 
+		/// </summary>
+		public JobQueueSettings? JobQueueSettings { get; set; }
+
+		/// <summary>
+		/// The minimum nunmber of agents to retain in this pool
+		/// </summary>
+		public int? MinAgents { get; set; }
+
+		/// <summary>
+		/// The minimum number of idle agents in this pool, if autoscaling is enabled
+		/// </summary>
+		public int? NumReserveAgents { get; set; }
+
+		/// <summary>
+		/// Properties to update for the pool. Properties set to null will be removed.
+		/// </summary>
+		public Dictionary<string, string?>? Properties { get; set; }
+	}
+
+	/// <summary>
+	/// Parameters to update a pool
+	/// </summary>
+	public class BatchUpdatePoolRequest
+	{
+		/// <summary>
+		///  ID of the pool to update
+		/// </summary>
+		[Required]
+		public string Id { get; set; } = null!;
+
+		/// <summary>
+		/// Optional new name for the pool
+		/// </summary>
+		public string? Name { get; set; }
+
+		/// <summary>
+		/// Properties to update for the pool. Properties set to null will be removed.
+		/// </summary>
+		public Dictionary<string, string?>? Properties { get; set; }
+	}
+
+	/// <summary>
+	/// Response describing a pool
+	/// </summary>
+	public class GetPoolResponse
+	{
+		/// <summary>
+		/// Unique id of the pool
+		/// </summary>
+		public string Id { get; set; }
+
+		/// <summary>
+		/// Name of the pool
+		/// </summary>
+		public string Name { get; set; }
+
+		/// <summary>
+		/// Condition for agents to be auto-added to the pool
+		/// </summary>
+		public Condition? Condition { get; set; }
+
+		/// <summary>
+		/// Whether to enable autoscaling for this pool
+		/// </summary>
+		public bool EnableAutoscaling { get; set; }
+		
+		/// <summary>
+		/// Cooldown time between scale-out events in seconds
+		/// </summary>
+		public int? ScaleOutCooldown { get; set; }
+		
+		/// <summary>
+		/// Cooldown time between scale-in events in seconds
+		/// </summary>
+		public int? ScaleInCooldown { get; set; }
+		
+		/// <summary>
+		/// Pool sizing strategy to be used for this pool
+		/// </summary>
+		public PoolSizeStrategy? SizeStrategy { get; set; }
+		
+		/// <summary>
+		/// Settings for lease utilization pool sizing strategy (if used)
+		/// </summary>
+		public LeaseUtilizationSettingsMessage? LeaseUtilizationSettings { get; set; }
+		
+		/// <summary>
+		/// Settings for job queue pool sizing strategy (if used) 
+		/// </summary>
+		public JobQueueSettingsMessage? JobQueueSettings { get; set; }
+
+		/// <summary>
+		/// The minimum nunmber of agents to retain in this pool
+		/// </summary>
+		public int? MinAgents { get; set; }
+
+		/// <summary>
+		/// The minimum number of idle agents in this pool, if autoscaling is enabled
+		/// </summary>
+		public int? NumReserveAgents { get; set; }
+
+		/// <summary>
+		/// List of workspaces that this agent contains
+		/// </summary>
+		public List<GetAgentWorkspaceResponse> Workspaces { get; set; }
+
+		/// <summary>
+		/// Arbitrary properties for this pool.
+		/// </summary>
+		public IReadOnlyDictionary<string, string> Properties { get; set; }
+
+		/// <summary>
+		/// Constructor
+		/// </summary>
+		/// <param name="pool">The pool to construct from</param>
+		public GetPoolResponse(IPool pool)
+		{
+			Id = pool.Id.ToString();
+			Name = pool.Name;
+			Condition = pool.Condition;
+			EnableAutoscaling = pool.EnableAutoscaling;
+			ScaleOutCooldown = pool.ScaleOutCooldown == null ? null : (int)pool.ScaleOutCooldown.Value.TotalSeconds;
+			ScaleInCooldown = pool.ScaleInCooldown == null ? null : (int)pool.ScaleInCooldown.Value.TotalSeconds;
+			SizeStrategy = pool.SizeStrategy;
+			LeaseUtilizationSettings = pool.LeaseUtilizationSettings == null ? null : new LeaseUtilizationSettingsMessage(pool.LeaseUtilizationSettings);
+			JobQueueSettings = pool.JobQueueSettings == null ? null : new JobQueueSettingsMessage(pool.JobQueueSettings);
+			MinAgents = pool.MinAgents;
+			NumReserveAgents = pool.NumReserveAgents;
+			Workspaces = pool.Workspaces.Select(x => new GetAgentWorkspaceResponse(x)).ToList();
+			Properties = pool.Properties;
+		}
+	}
+}
diff --git a/Engine/Source/Programs/Horde/Horde.Build/Agents/Pools/PoolsController.Messages.cs b/Engine/Source/Programs/Horde/Horde.Build/Agents/Pools/PoolsController.Messages.cs
deleted file mode 100644
diff --git a/Engine/Source/Programs/Horde/Horde.Build/Devices/DeviceMessages.cs b/Engine/Source/Programs/Horde/Horde.Build/Devices/DeviceMessages.cs
new file mode 100644
--- /dev/null
+++ b/Engine/Source/Programs/Horde/Horde.Build/Devices/DeviceMessages.cs
@@ -0,0 +1,737 @@
+// Copyright Epic Games, Inc. All Rights Reserved.
+
+using System;
+using System.Collections.Generic;
+using System.ComponentModel.DataAnnotations;
+using System.Linq;
+
+namespace Horde.Build.Devices
+{
+
+	/// <summary>
+	/// Create device platform request
+	/// </summary>
+	public class CreateDevicePlatformRequest
+	{
+		/// <summary>
+		/// The name of the platform
+		/// </summary>
+		[Required]
+		public string Name { get; set; } = null!;
+
+	}
+
+	/// <summary>
+	/// Create device platform response object
+	/// </summary>
+	public class CreateDevicePlatformResponse
+	{
+		/// <summary>
+		/// Id of newly created platform
+		/// </summary>
+		public string Id { get; set; }
+
+		/// <summary>
+		/// Constructor
+		/// </summary>
+		public CreateDevicePlatformResponse(string id)
+		{
+			Id = id;
+		}
+	}
+
+	/// <summary>
+	/// Update requesty object for a device platform
+	/// </summary>
+	public class UpdateDevicePlatformRequest
+	{
+		/// <summary>
+		/// The vendor model ids for the platform
+		/// </summary>
+		public string[] ModelIds { get; set; } = null!;
+	}
+
+	/// <summary>
+	/// Get object response which describes a device platform
+	/// </summary>
+	public class GetDevicePlatformResponse
+	{
+		/// <summary>
+		/// Unique id of device platform
+		/// </summary>
+		public string Id { get; set; }
+
+		/// <summary>
+		/// Friendly name of device platform
+		/// </summary>
+		public string Name { get; set; }
+
+		/// <summary>
+		/// Platform vendor models
+		/// </summary>
+		public string[] ModelIds { get; set; }
+
+		/// <summary>
+		/// Response constructor
+		/// </summary>
+		public GetDevicePlatformResponse(string id, string name, string[] modelIds)
+		{
+			Id = id;
+			Name = name;
+			ModelIds = modelIds;
+		}
+	}
+
+	/// <summary>
+	/// Device pool creation request object
+	/// </summary>
+	public class CreateDevicePoolRequest
+	{
+		/// <summary>
+		/// The name for the new pool
+		/// </summary>
+		[Required]
+		public string Name { get; set; } = null!;
+
+		/// <summary>
+		/// The name for the new pool
+		/// </summary>
+		[Required]
+		public DevicePoolType PoolType { get; set; }
+
+		/// <summary>
+		/// Projects associated with this device pool
+		/// </summary>
+		public List<string>? ProjectIds { get; set; }
+	}
+
+	/// <summary>
+	/// Device pool update request object
+	/// </summary>
+	public class UpdateDevicePoolRequest
+	{
+		/// <summary>
+		/// Id of the device pool to update
+		/// </summary>
+		public string Id { get; set; } = null!;
+
+		/// <summary>
+		/// Projects associated with this device pool
+		/// </summary>
+		public List<string>? ProjectIds { get; set; }
+	}
+
+	/// <summary>
+	/// Device pool creation response object
+	/// </summary>
+	public class CreateDevicePoolResponse
+	{
+		/// <summary>
+		/// Id of the newly created device pool
+		/// </summary>
+		public string Id { get; set; }
+
+		/// <summary>
+		/// Constructor
+		/// </summary>
+		public CreateDevicePoolResponse(string id)
+		{
+			Id = id;
+		}
+	}
+
+	/// <summary>
+	/// Device pool response object
+	/// </summary>
+	public class GetDevicePoolResponse
+	{
+		/// <summary>
+		/// Id of  the device pool
+		/// </summary>
+		public string Id { get; set; }
+
+		/// <summary>
+		/// Name of the device pool
+		/// </summary>
+		public string Name { get; set; }
+
+		/// <summary>
+		/// Type of the device pool
+		/// </summary>
+		public DevicePoolType PoolType { get; set; }
+
+		/// <summary>
+		/// Whether there is write access to the pool
+		/// </summary>
+		public bool WriteAccess { get; set; }
+
+		/// <summary>
+		/// Constructor
+		/// </summary>
+		public GetDevicePoolResponse(string id, string name, DevicePoolType poolType, bool writeAccess)
+		{
+			Id = id;
+			Name = name;
+			PoolType = poolType;
+			WriteAccess = writeAccess;
+		}
+	}
+
+	// Devices 
+
+	/// <summary>
+	/// Device creation request object
+	/// </summary>
+	public class CreateDeviceRequest
+	{
+		/// <summary>
+		/// The platform of the device 
+		/// </summary>
+		[Required]
+		public string PlatformId { get; set; } = null!;
+
+		/// <summary>
+		/// The pool to assign the device
+		/// </summary>
+		[Required]
+		public string PoolId { get; set; } = null!;
+
+		/// <summary>
+		/// The friendly name of the device
+		/// </summary>
+		[Required]
+		public string Name { get; set; } = null!;
+
+		/// <summary>
+		/// Whether to create the device in enabled state
+		/// </summary>
+		public bool? Enabled { get; set; }
+
+		/// <summary>
+		/// The network address of the device
+		/// </summary>
+		public string? Address { get; set; }
+
+		/// <summary>
+		/// The vendor model id of the device
+		/// </summary>
+		public string? ModelId { get; set; }
+	}
+
+	/// <summary>
+	/// Device creation response object
+	/// </summary>
+	public class CreateDeviceResponse
+	{
+		/// <summary>
+		/// The id of the newly created device
+		/// </summary>
+		[Required]
+		public string Id { get; set; } = null!;
+
+		/// <summary>
+		/// Constructor
+		/// </summary>
+		/// <param name="id"></param>
+		public CreateDeviceResponse(string id)
+		{
+			Id = id;
+		}
+	}
+
+	/// <summary>
+	/// Get response object which describes a device
+	/// </summary>
+	public class GetDeviceUtilizationResponse
+	{
+		/// <summary>
+		/// The job id which utilized device
+		/// </summary>
+		public string? JobId { get; set; }
+
+		/// <summary>
+		/// The job's step id
+		/// </summary>
+		public string? StepId { get; set; }
+
+		/// <summary>
+		/// The time device was reserved
+		/// </summary>
+		public DateTime ReservationStartUtc { get; set; }
+
+		/// <summary>
+		/// The time device was freed
+		/// </summary>
+		public DateTime? ReservationFinishUtc { get; set; }
+
+		/// <summary>
+		/// Constructor
+		/// </summary>
+		/// <param name="telemetry"></param>
+		public GetDeviceUtilizationResponse(DeviceUtilizationTelemetry telemetry)
+		{
+			JobId = telemetry.JobId;
+			StepId = telemetry.StepId;
+			ReservationStartUtc = telemetry.ReservationStartUtc;
+			ReservationFinishUtc = telemetry.ReservationFinishUtc;
+		}
+	}
+
+	/// <summary>
+	/// Get response object which describes a device
+	/// </summary>
+	public class GetDeviceResponse
+	{
+		/// <summary>
+		/// The unique id of the device
+		/// </summary>
+		public string Id { get; set; }
+
+		/// <summary>
+		/// The platform of the device
+		/// </summary>
+		public string PlatformId { get; set; }
+
+		/// <summary>
+		/// The pool the device belongs to
+		/// </summary>
+		public string PoolId { get; set; }
+
+		/// <summary>
+		/// The friendly name of the device
+		/// </summary>
+		public string Name { get; set; }
+
+		/// <summary>
+		/// Whether the device is currently enabled
+		/// </summary>
+		public bool Enabled { get; set; }
+
+		/// <summary>
+		///  The address of the device (if it allows network connections)
+		/// </summary>
+		public string? Address { get; set; }
+
+		/// <summary>
+		/// The vendor model id of the device
+		/// </summary>
+		public string? ModelId { get; set; }
+
+		/// <summary>
+		/// Any notes provided for the device
+		/// </summary>
+		public string? Notes { get; set; }
+
+		/// <summary>
+		/// If the device has a marked problem
+		/// </summary>
+		public DateTime? ProblemTime { get; set; }
+
+		/// <summary>
+		/// If the device is in maintenance mode
+		/// </summary>
+		public DateTime? MaintenanceTime { get; set; }
+
+		/// <summary>
+		/// The user id that has the device checked out
+		/// </summary>
+		public string? CheckedOutByUserId { get; set; }
+
+		/// <summary>
+		/// The last time the device was checked out
+		/// </summary>
+		public DateTime? CheckOutTime { get; set; }
+
+		/// <summary>
+		/// The last user to modifiy the device
+		/// </summary>
+		public string? ModifiedByUser { get; set; }
+
+		/// <summary>
+		///  Device Utilization data
+		/// </summary>
+		public List<GetDeviceUtilizationResponse>? Utilization { get; set; }
+
+		/// <summary>
+		/// Device response constructor
+		/// </summary>
+		public GetDeviceResponse(string id, string platformId, string poolId, string name, bool enabled, string? address, string? modelId, string? modifiedByUser, string? notes, DateTime? problemTime, DateTime? maintenanceTime, List<DeviceUtilizationTelemetry>? utilization, string? checkedOutByUser = null, DateTime? checkOutTime = null)
+		{
+			Id = id;
+			Name = name;
+			PlatformId = platformId;
+			PoolId = poolId;
+			Enabled = enabled;
+			Address = address;
+			ModelId = modelId;
+			ModifiedByUser = modifiedByUser;
+			Notes = notes;
+			ProblemTime = problemTime;
+			MaintenanceTime = maintenanceTime;
+			Utilization = utilization?.Select(u => new GetDeviceUtilizationResponse(u)).ToList();
+			CheckedOutByUserId = checkedOutByUser;
+			CheckOutTime = checkOutTime;
+		}
+	}
+
+	/// <summary>
+	/// Device update request object
+	/// </summary>
+	public class UpdateDeviceRequest
+	{
+		/// <summary>
+		/// The device pool id
+		/// </summary>
+		public string? PoolId { get; set; }
+
+		/// <summary>
+		/// The device name
+		/// </summary>
+		public string? Name { get; set; }
+
+		/// <summary>
+		/// IP address or hostname of device
+		/// </summary>
+		public string? Address { get; set; }
+
+		/// <summary>
+		/// Device vendor model id
+		/// </summary>
+		public string? ModelId { get; set; }
+
+		/// <summary>
+		/// Markdown notes
+		/// </summary>
+		public string? Notes { get; set; }
+
+		/// <summary>
+		/// Whether device is enabled
+		/// </summary>
+		public bool? Enabled { get; set; }
+
+		/// <summary>
+		/// Whether the device is in maintenance mode
+		/// </summary>
+		public bool? Maintenance { get; set; }
+
+		/// <summary>
+		/// Whether to set or clear any device problem state
+		/// </summary>
+		public bool? Problem { get; set; }
+	}
+
+	/// <summary>
+	/// Device checkout request object
+	/// </summary>
+	public class CheckoutDeviceRequest
+	{
+		/// <summary>
+		/// Whether to checkout or in the device
+		/// </summary>
+		public bool Checkout { get; set; }
+	}
+
+	/// <summary>
+	/// Device reservation request object
+	/// </summary>
+	public class DeviceReservationRequest
+	{
+		/// <summary>
+		/// Device reservation platform id
+		/// </summary>
+		[Required]
+		public string PlatformId { get; set; } = null!;
+
+		/// <summary>
+		/// The optional vendor model ids to include for this device
+		/// </summary>
+		public List<string>? IncludeModels { get; set; }
+
+		/// <summary>
+		/// The optional vendor model ids to exclude for this device
+		/// </summary>
+		public List<string>? ExcludeModels { get; set; }
+	}
+
+	/// <summary>
+	/// Reservation request object
+	/// </summary>
+	public class CreateDeviceReservationRequest
+	{
+		/// <summary>
+		/// What pool to reserve devices in
+		/// </summary>
+		[Required]
+		public string PoolId { get; set; } = null!;
+
+		/// <summary>
+		/// Devices to reserve
+		/// </summary>
+		[Required]
+		public List<DeviceReservationRequest> Devices { get; set; } = null!;
+	}
+
+	/// <summary>
+	/// Device reservation response object
+	/// </summary>
+	public class CreateDeviceReservationResponse
+	{
+		/// <summary>
+		/// The reservation id of newly created reservation
+		/// </summary>
+		public string Id { get; set; } = null!;
+
+		/// <summary>
+		/// The devices that were reserved
+		/// </summary>
+		public List<GetDeviceResponse> Devices { get; set; } = null!;
+	}
+
+	/// <summary>
+	/// A reservation containing one or more devices
+	/// </summary>
+	public class GetDeviceReservationResponse
+	{
+		/// <summary>
+		/// Randomly generated unique id for this reservation
+		/// </summary>
+		[Required]
+		public string Id { get; set; } = null!;
+
+		/// <summary>
+		/// Which device pool the reservation is in
+		/// </summary>
+		[Required]
+		public string PoolId { get; set; } = null!;
+
+		/// <summary>
+		/// The reserved devices
+		/// </summary>
+		[Required]
+		public List<string> Devices { get; set; } = null!;
+
+		/// <summary>
+		/// JobID holding reservation
+		/// </summary>
+		public string? JobId { get; set; }
+
+		/// <summary>
+		/// Job step id holding reservation
+		/// </summary>
+		public string? StepId { get; set; }
+
+		/// <summary>
+		/// Reservations held by a user, requires a token
+		/// </summary>
+		public string? UserId { get; set; }
+
+		/// <summary>
+		/// The hostname of machine holding reservation
+		/// </summary>
+		public string? Hostname { get; set; }
+
+		/// <summary>
+		/// The optional reservation details 
+		/// </summary>
+		public string? ReservationDetails { get; set; }
+
+		/// <summary>
+		/// The UTC time when the reservation was created
+		/// </summary>
+		public DateTime CreateTimeUtc { get; set; }
+
+		/// <summary>
+		/// The legacy reservation system guid, to be removed once can update Gauntlet client in all streams
+		/// </summary>
+		public string LegacyGuid { get; set; } = null!;
+	}
+
+	// Legacy clients
+
+	/// <summary>
+	/// Reservation request for legacy clients
+	/// </summary>
+	public class LegacyCreateReservationRequest
+	{
+		/// <summary>
+		/// The device types to reserve, these are mapped to platforms
+		/// </summary>
+		[Required]
+		public string[] DeviceTypes { get; set; } = null!;
+
+		/// <summary>
+		/// The hostname of machine reserving devices 
+		/// </summary>
+		[Required]
+		public string Hostname { get; set; } = null!;
+
+		/// <summary>
+		/// The duration of reservation
+		/// </summary>
+		[Required]
+		public string Duration { get; set; } = null!;
+
+		/// <summary>
+		/// Reservation details string
+		/// </summary>
+		public string? ReservationDetails { get; set; } = null;
+
+		/// <summary>
+		/// The PoolId of reservation request 
+		/// </summary>
+		public string? PoolId { get; set; } = null;
+
+		/// <summary>
+		/// The JobId of reservation request 
+		/// </summary>
+		public string? JobId { get; set; } = null;
+
+		/// <summary>
+		/// The StepId of reservation request 
+		/// </summary>
+		public string? StepId { get; set; } = null;
+
+	}
+
+	/// <summary>
+	/// Reservation response for legacy clients
+	/// </summary>
+	public class GetLegacyReservationResponse
+	{
+		/// <summary>
+		/// The names of the devices that were reserved
+		/// </summary>
+		[Required]
+		public string[] DeviceNames { get; set; } = null!;
+
+		/// <summary>
+		/// The corresponding perf specs of the reserved devices
+		/// </summary>
+		[Required]
+		public string[] DevicePerfSpecs { get; set; } = null!;
+
+		/// <summary>
+		/// The corresponding perf specs of the reserved devices
+		/// </summary>
+		[Required]
+		public string[] DeviceModels { get; set; } = null!;
+
+		/// <summary>
+		/// The host name of the machine making the reservation
+		/// </summary>
+		[Required]
+		public string HostName { get; set; } = null!;
+
+		/// <summary>
+		/// The start time of the reservation
+		/// </summary>
+		[Required]
+		public string StartDateTime { get; set; } = null!;
+
+		/// <summary>
+		/// The duration of the reservation (before renew)
+		/// </summary>
+		[Required]
+		public string Duration { get; set; } = null!;
+
+		/// <summary>
+		/// The JobId of reservation request 
+		/// </summary>
+		public string? JobId { get; set; } = null;
+
+		/// <summary>
+		/// The StepId of reservation request 
+		/// </summary>
+		public string? StepId { get; set; } = null;
+
+		/// <summary>
+		/// The legacy guid of the reservation
+		/// </summary>
+		[Required]
+		public string Guid { get; set; } = null!;
+	}
+
+	/// <summary>
+	/// Device response for legacy clients
+	/// </summary>
+	public class GetLegacyDeviceResponse
+	{
+		/// <summary>
+		/// The  name of the reserved device
+		/// </summary>
+		[Required]
+		public string Name { get; set; } = null!;
+
+		/// <summary>
+		/// The (legacy) type of the device, mapped from platform id
+		/// </summary>
+		[Required]
+		public string Type { get; set; } = null!;
+
+		/// <summary>
+		/// The IP or hostname of device
+		/// </summary>
+		[Required]
+		public string IPOrHostName { get; set; } = null!;
+
+		/// <summary>
+		/// The (legacy) perf spec of the device
+		/// </summary>
+		[Required]
+		public string PerfSpec { get; set; } = null!;
+
+		/// <summary>
+		/// The device model information
+		/// </summary>
+		[Required]
+		public string Model { get; set; } = null!;
+
+		/// <summary>
+		/// The available start time which is parsed client side
+		/// </summary>
+		[Required]
+		public string AvailableStartTime { get; set; } = null!;
+
+		/// <summary>
+		/// The available end time which is parsed client side
+		/// </summary>
+		[Required]
+		public string AvailableEndTime { get; set; } = null!;
+
+		/// <summary>
+		/// Whether device is enabled
+		/// </summary>
+		[Required]
+		public bool Enabled { get; set; }
+
+		/// <summary>
+		/// Associated device data
+		/// </summary>
+		[Required]
+		public string DeviceData { get; set; } = null!;
+	}
+
+	/// <summary>
+	/// Platform mapping update for V1 devices
+	/// </summary>
+	public class UpdatePlatformMapRequest
+	{
+		/// <summary>
+		/// Platform name => platform id
+		/// </summary>
+		[Required]
+		public Dictionary<string, string> PlatformMap { get; set; } = null!;
+
+		/// <summary>
+		/// Platform id => name
+		/// </summary>
+		[Required]
+		public Dictionary<string, string> PlatformReverseMap { get; set; } = null!;
+
+		/// <summary>
+		/// Perfspec v1 => model names
+		/// </summary>
+		[Required]
+		public Dictionary<string, string> PerfSpecHighMap { get; set; } = null!;
+	}
+}
diff --git a/Engine/Source/Programs/Horde/Horde.Build/Devices/DevicesController.Messages.cs b/Engine/Source/Programs/Horde/Horde.Build/Devices/DevicesController.Messages.cs
deleted file mode 100644
diff --git a/Engine/Source/Programs/Horde/Horde.Build/Issues/IssueMessages.cs b/Engine/Source/Programs/Horde/Horde.Build/Issues/IssueMessages.cs
new file mode 100644
--- /dev/null
+++ b/Engine/Source/Programs/Horde/Horde.Build/Issues/IssueMessages.cs
@@ -0,0 +1,1073 @@
+// Copyright Epic Games, Inc. All Rights Reserved.
+
+using System;
+using System.Collections.Generic;
+using System.Linq;
+using EpicGames.Core;
+using Horde.Build.Issues;
+using Horde.Build.Issues.External;
+using Horde.Build.Server;
+using Horde.Build.Streams;
+using Horde.Build.Users;
+using Horde.Build.Utilities;
+using MongoDB.Bson;
+
+namespace Horde.Build.Issues
+{
+	using StreamId = StringId<IStream>;
+	using TemplateRefId = StringId<TemplateRef>;
+	using WorkflowId = StringId<WorkflowConfig>;
+
+	/// <summary>
+	/// Identifies a particular changelist and job
+	/// </summary>
+	public class GetIssueStepResponse
+	{
+		/// <summary>
+		/// The changelist number
+		/// </summary>
+		public int Change { get; set; }
+
+		/// <summary>
+		/// Severity of the issue in this step
+		/// </summary>
+		public IssueSeverity Severity { get; set; }
+
+		/// <summary>
+		/// Name of the job containing this step
+		/// </summary>
+		public string JobName { get; set; }
+
+		/// <summary>
+		/// The unique job id
+		/// </summary>
+		public string JobId { get; set; }
+
+		/// <summary>
+		/// The unique batch id
+		/// </summary>
+		public string BatchId { get; set; }
+
+		/// <summary>
+		/// The unique step id
+		/// </summary>
+		public string StepId { get; set; }
+
+		/// <summary>
+		/// Time at which the step ran
+		/// </summary>
+		public DateTime StepTime { get; set; }
+
+		/// <summary>
+		/// The unique log id
+		/// </summary>
+		public string? LogId { get; set; }
+
+		/// <summary>
+		/// Constructor
+		/// </summary>
+		/// <param name="issueStep">The issue step to construct from</param>
+		public GetIssueStepResponse(IIssueStep issueStep)
+		{
+			Change = issueStep.Change;
+			Severity = issueStep.Severity;
+			JobName = issueStep.JobName;
+			JobId = issueStep.JobId.ToString();
+			BatchId = issueStep.BatchId.ToString();
+			StepId = issueStep.StepId.ToString();
+			StepTime = issueStep.StepTime;
+			LogId = issueStep.LogId?.ToString();
+		}
+	}
+
+	/// <summary>
+	/// Trace of a set of node failures across multiple steps
+	/// </summary>
+	public class GetIssueSpanResponse
+	{
+		/// <summary>
+		/// Unique id of this span
+		/// </summary>
+		public string Id { get; set; }
+
+		/// <summary>
+		/// The template containing this step
+		/// </summary>
+		public string TemplateId { get; set; }
+
+		/// <summary>
+		/// Name of the step
+		/// </summary>
+		public string Name { get; set; }
+
+		/// <summary>
+		/// Workflow that this span belongs to
+		/// </summary>
+		public WorkflowId? WorkflowId { get; set; }
+
+		/// <summary>
+		/// The previous build 
+		/// </summary>
+		public GetIssueStepResponse? LastSuccess { get; set; }
+
+		/// <summary>
+		/// The failing builds for a particular event
+		/// </summary>
+		public List<GetIssueStepResponse> Steps { get; set; }
+
+		/// <summary>
+		/// The following successful build
+		/// </summary>
+		public GetIssueStepResponse? NextSuccess { get; set; }
+
+		/// <summary>
+		/// Constructor
+		/// </summary>
+		/// <param name="span">The node to construct from</param>
+		/// <param name="steps">Failing steps for this span</param>
+		public GetIssueSpanResponse(IIssueSpan span, List<IIssueStep> steps)
+		{
+			Id = span.Id.ToString();
+			Name = span.NodeName;
+			TemplateId = span.TemplateRefId.ToString();
+			WorkflowId = span.LastFailure.Annotations.WorkflowId;
+			LastSuccess = (span.LastSuccess != null) ? new GetIssueStepResponse(span.LastSuccess) : null;
+			Steps = steps.ConvertAll(x => new GetIssueStepResponse(x));
+			NextSuccess = (span.NextSuccess != null) ? new GetIssueStepResponse(span.NextSuccess) : null;
+		}
+	}
+
+	/// <summary>
+	/// Information about a particular step
+	/// </summary>
+	public class GetIssueStreamResponse
+	{
+		/// <summary>
+		/// Unique id of the stream
+		/// </summary>
+		public string StreamId { get; set; }
+
+		/// <summary>
+		/// Minimum changelist affected by this issue (ie. last successful build)
+		/// </summary>
+		public int? MinChange { get; set; }
+
+		/// <summary>
+		/// Maximum changelist affected by this issue (ie. next successful build)
+		/// </summary>
+		public int? MaxChange { get; set; }
+
+		/// <summary>
+		/// Map of steps to (event signature id -> trace id)
+		/// </summary>
+		public List<GetIssueSpanResponse> Nodes { get; set; } = new List<GetIssueSpanResponse>();
+
+		/// <summary>
+		/// Constructor
+		/// </summary>
+		/// <param name="stream">The stream to construct from</param>
+		/// <param name="spans">List of spans for the given stream</param>
+		/// <param name="steps">List of steps for the given stream</param>
+		public GetIssueStreamResponse(IStream stream, List<IIssueSpan> spans, List<IIssueStep> steps)
+		{
+			StreamId = stream.Id.ToString();
+
+			foreach (IIssueSpan span in spans)
+			{
+				if (span.LastSuccess != null && (MinChange == null || span.LastSuccess.Change < MinChange.Value))
+				{
+					MinChange = span.LastSuccess.Change;
+				}
+				if (span.NextSuccess != null && (MaxChange == null || span.NextSuccess.Change > MaxChange.Value))
+				{
+					MaxChange = span.NextSuccess.Change;
+				}
+				Nodes.Add(new GetIssueSpanResponse(span, steps.Where(y => y.SpanId == span.Id).ToList()));
+			}
+		}
+	}
+
+	/// <summary>
+	/// Outcome of a particular build
+	/// </summary>
+	public enum IssueBuildOutcome
+	{
+		/// <summary>
+		/// Unknown outcome
+		/// </summary>
+		Unknown,
+
+		/// <summary>
+		/// Build succeeded
+		/// </summary>
+		Success,
+
+		/// <summary>
+		/// Build failed
+		/// </summary>
+		Error,
+
+		/// <summary>
+		/// Build finished with warnings
+		/// </summary>
+		Warning,
+	}
+
+	/// <summary>
+	/// Information about a diagnostic
+	/// </summary>
+	public class GetIssueDiagnosticResponse
+	{
+		/// <summary>
+		/// The corresponding build id
+		/// </summary>
+		public long? BuildId { get; set; }
+
+		/// <summary>
+		/// Message for the diagnostic
+		/// </summary>
+		public string Message { get; set; }
+
+		/// <summary>
+		/// Link to the error
+		/// </summary>
+		public Uri Url { get; set; }
+
+		/// <summary>
+		/// Constructor
+		/// </summary>
+		/// <param name="buildId">The corresponding build id</param>
+		/// <param name="message">Message for the diagnostic</param>
+		/// <param name="url">Link to the diagnostic</param>
+		public GetIssueDiagnosticResponse(long? buildId, string message, Uri url)
+		{
+			BuildId = buildId;
+			Message = message;
+			Url = url;
+		}
+	}
+
+	/// <summary>
+	/// Information about a template affected by an issue
+	/// </summary>
+	public class GetIssueAffectedTemplateResponse
+	{
+		/// <summary>
+		/// The template id
+		/// </summary>
+		public string TemplateId { get; set; }
+
+		/// <summary>
+		/// The template name
+		/// </summary>
+		public string TemplateName { get; set; }
+
+		/// <summary>
+		/// Whether it has been resolved or not
+		/// </summary>
+		public bool Resolved { get; set; }
+
+		/// <summary>
+		/// The issue severity of the affected template
+		/// </summary>
+		public IssueSeverity Severity { get; set; }
+
+		/// <summary>
+		/// Constructor
+		/// </summary>
+		/// <param name="templateId"></param>
+		/// <param name="templateName"></param>
+		/// <param name="resolved"></param>
+		/// <param name="severity"></param>
+		public GetIssueAffectedTemplateResponse(string templateId, string templateName, bool resolved, IssueSeverity severity = IssueSeverity.Unspecified)
+		{
+			TemplateId = templateId;
+			TemplateName = templateName;
+			Resolved = resolved;
+			Severity = severity;
+		}
+	}
+
+	/// <summary>
+	/// Summary for the state of a stream in an issue
+	/// </summary>
+	public class GetIssueAffectedStreamResponse
+	{
+		/// <summary>
+		/// Id of the stream
+		/// </summary>
+		public string StreamId { get; set; }
+
+		/// <summary>
+		/// Name of the stream
+		/// </summary>
+		public string StreamName { get; set; }
+
+		/// <summary>
+		/// Whether the issue has been resolved in this stream
+		/// </summary>
+		public bool Resolved { get; set; }
+
+		/// <summary>
+		/// The affected templates
+		/// </summary>
+		public List<GetIssueAffectedTemplateResponse> AffectedTemplates { get; set; }
+
+		/// <summary>
+		/// List of affected template ids
+		/// </summary>
+		public List<string> TemplateIds { get; set; }
+
+		/// <summary>
+		/// List of resolved template ids
+		/// </summary>
+		public List<string> ResolvedTemplateIds { get; set; }
+
+		/// <summary>
+		/// List of unresolved template ids
+		/// </summary>
+		public List<string> UnresolvedTemplateIds { get; set; }
+
+		/// <summary>
+		/// Constructor
+		/// </summary>
+		/// <param name="details">Issue to construct from</param>
+		/// <param name="stream"></param>
+		/// <param name="spans">The spans to construct from</param>
+		public GetIssueAffectedStreamResponse(IIssueDetails details, IStream? stream, IEnumerable<IIssueSpan> spans)
+		{
+			IIssueSpan firstSpan = spans.First();
+			StreamId = firstSpan.StreamId.ToString();
+			StreamName = firstSpan.StreamName;
+			Resolved = spans.All(x => x.NextSuccess != null);
+
+			AffectedTemplates = new List<GetIssueAffectedTemplateResponse>();
+			foreach (IGrouping<TemplateRefId, IIssueSpan> template in spans.GroupBy(x => x.TemplateRefId))
+			{
+				string templateName = template.Key.ToString();
+				if (stream != null && stream.Templates.TryGetValue(template.Key, out TemplateRef? templateRef))
+				{
+					templateName = templateRef.Name;
+				}
+
+				HashSet<ObjectId> unresolvedTemplateSpans = new HashSet<ObjectId>(template.Where(x => x.NextSuccess == null).Select(x => x.Id));
+
+				IIssueStep? templateStep = details.Steps.Where(x => unresolvedTemplateSpans.Contains(x.SpanId)).OrderByDescending(x => x.StepTime).FirstOrDefault();
+
+				AffectedTemplates.Add(new GetIssueAffectedTemplateResponse(template.Key.ToString(), templateName, template.All(x => x.NextSuccess != null), templateStep?.Severity ?? IssueSeverity.Unspecified));
+			}
+
+			HashSet<TemplateRefId> templateIdsSet = new HashSet<TemplateRefId>(spans.Select(x => x.TemplateRefId));
+			TemplateIds = templateIdsSet.Select(x => x.ToString()).ToList();
+
+			HashSet<TemplateRefId> unresolvedTemplateIdsSet = new HashSet<TemplateRefId>(spans.Where(x => x.NextSuccess == null).Select(x => x.TemplateRefId));
+			UnresolvedTemplateIds = unresolvedTemplateIdsSet.Select(x => x.ToString()).ToList();
+			ResolvedTemplateIds = templateIdsSet.Except(unresolvedTemplateIdsSet).Select(x => x.ToString()).ToList();
+		}
+	}
+
+	/// <summary>
+	/// Stores information about a build health issue
+	/// </summary>
+	public class GetIssueResponse
+	{
+		/// <summary>
+		/// The unique object id
+		/// </summary>
+		public int Id { get; set; }
+
+		/// <summary>
+		/// Time at which the issue was created
+		/// </summary>
+		public DateTime CreatedAt { get; set; }
+
+		/// <summary>
+		/// Time at which the issue was retrieved
+		/// </summary>
+		public DateTime RetrievedAt { get; set; }
+
+		/// <summary>
+		/// The associated project for the issue
+		/// </summary>
+		public string? Project { get; set; }
+
+		/// <summary>
+		/// The summary text for this issue
+		/// </summary>
+		public string Summary { get; set; }
+
+		/// <summary>
+		/// Detailed description text
+		/// </summary>
+		public string? Description { get; set; }
+
+		/// <summary>
+		/// Severity of this issue
+		/// </summary>
+		public IssueSeverity Severity { get; set; }
+
+		/// <summary>
+		/// Whether the issue is promoted
+		/// </summary>
+		public bool Promoted { get; set; }
+
+		/// <summary>
+		/// Owner of the issue [DEPRECATED]
+		/// </summary>
+		public string? Owner { get; set; }
+
+		/// <summary>
+		/// User id of the owner [DEPRECATED]
+		/// </summary>
+		public string? OwnerId { get; set; }
+
+		/// <summary>
+		/// Owner of the issue
+		/// </summary>
+		public GetThinUserInfoResponse? OwnerInfo { get; set; }
+
+		/// <summary>
+		/// User that nominated the current owner [DEPRECATED]
+		/// </summary>
+		public string? NominatedBy { get; set; }
+
+		/// <summary>
+		/// Owner of the issue
+		/// </summary>
+		public GetThinUserInfoResponse? NominatedByInfo { get; set; }
+
+		/// <summary>
+		/// Time that the issue was acknowledged
+		/// </summary>
+		public DateTime? AcknowledgedAt { get; set; }
+
+		/// <summary>
+		/// Changelist that fixed this issue
+		/// </summary>
+		public int? FixChange { get; set; }
+
+		/// <summary>
+		/// Time at which the issue was resolved
+		/// </summary>
+		public DateTime? ResolvedAt { get; set; }
+
+		/// <summary>
+		/// Name of the user that resolved the issue [DEPRECATED]
+		/// </summary>
+		public string? ResolvedBy { get; set; }
+
+		/// <summary>
+		/// User id of the person that resolved the issue [DEPRECATED]
+		/// </summary>
+		public string? ResolvedById { get; set; }
+
+		/// <summary>
+		/// User that resolved the issue
+		/// </summary>
+		public GetThinUserInfoResponse? ResolvedByInfo { get; set; }
+
+		/// <summary>
+		/// Time at which the issue was verified
+		/// </summary>
+		public DateTime? VerifiedAt { get; set; }
+
+		/// <summary>
+		/// Time that the issue was last seen
+		/// </summary>
+		public DateTime LastSeenAt { get; set; }
+
+		/// <summary>
+		/// List of stream paths affected by this issue
+		/// </summary>
+		public List<string> Streams { get; set; }
+
+		/// <summary>
+		/// List of affected stream ids
+		/// </summary>
+		public List<string> ResolvedStreams { get; set; }
+
+		/// <summary>
+		/// List of unresolved streams
+		/// </summary>
+		public List<string> UnresolvedStreams { get; set; }
+
+		/// <summary>
+		/// List of affected streams
+		/// </summary>
+		public List<GetIssueAffectedStreamResponse> AffectedStreams { get; set; }
+
+		/// <summary>
+		/// Most likely suspects for causing this issue [DEPRECATED]
+		/// </summary>
+		public List<string> PrimarySuspects { get; set; }
+
+		/// <summary>
+		/// User ids of the most likely suspects [DEPRECATED]
+		/// </summary>
+		public List<string> PrimarySuspectIds { get; set; }
+
+		/// <summary>
+		/// Most likely suspects for causing this issue
+		/// </summary>
+		public List<GetThinUserInfoResponse> PrimarySuspectsInfo { get; set; }
+
+		/// <summary>
+		/// Whether to show alerts for this issue
+		/// </summary>
+		public bool ShowDesktopAlerts { get; set; }
+
+		/// <summary>
+		/// Key for this issue in external issue tracker
+		/// </summary>
+		public string? ExternalIssueKey { get; set; }
+
+		/// <summary>
+		/// User who quarantined the issue
+		/// </summary>
+		public GetThinUserInfoResponse? QuarantinedByUserInfo { get; set; }
+
+		/// <summary>
+		/// The UTC time when the issue was quarantined
+		/// </summary>
+		public DateTime? QuarantineTimeUtc { get; set; }
+
+
+		/// <summary>
+		/// Constructs a new issue
+		/// </summary>
+		/// <param name="details">Issue to construct from</param>
+		/// <param name="affectedStreams">The affected streams</param>
+		/// <param name="showDesktopAlerts">Whether to show alerts for this issue</param>
+		public GetIssueResponse(IIssueDetails details, List<GetIssueAffectedStreamResponse> affectedStreams, bool showDesktopAlerts)
+		{
+			IIssue issue = details.Issue;
+			Id = issue.Id;
+			CreatedAt = issue.CreatedAt;
+			RetrievedAt = DateTime.UtcNow;
+			Summary = String.IsNullOrEmpty(issue.UserSummary)? issue.Summary : issue.UserSummary;
+			Description = issue.Description;
+			Severity = issue.Severity;
+			Promoted = issue.Promoted;
+			Owner = details.Owner?.Login;
+			OwnerId = details.Owner?.Id.ToString();
+			if(details.Owner != null)
+			{
+				OwnerInfo = new GetThinUserInfoResponse(details.Owner);
+			}
+			NominatedBy = details.NominatedBy?.Login;
+			if (details.NominatedBy != null)
+			{
+				NominatedByInfo = new GetThinUserInfoResponse(details.NominatedBy);
+			}
+			AcknowledgedAt = issue.AcknowledgedAt;
+			FixChange = issue.FixChange;
+			ResolvedAt = issue.ResolvedAt;
+			ResolvedBy = details.ResolvedBy?.Login;
+			ResolvedById = details.ResolvedBy?.Id.ToString();
+			if (details.ResolvedBy != null)
+			{
+				ResolvedByInfo = new GetThinUserInfoResponse(details.ResolvedBy);
+			}
+			VerifiedAt = issue.VerifiedAt;
+			LastSeenAt = issue.LastSeenAt;
+			Streams = details.Spans.Select(x => x.StreamName).Distinct().ToList()!;
+			ResolvedStreams = new List<string>();
+			UnresolvedStreams = new List<string>();
+			AffectedStreams = affectedStreams;
+			foreach (IGrouping<StreamId, IIssueSpan> stream in details.Spans.GroupBy(x => x.StreamId))
+			{
+				if (stream.All(x => x.NextSuccess != null))
+				{
+					ResolvedStreams.Add(stream.Key.ToString());
+				}
+				else
+				{
+					UnresolvedStreams.Add(stream.Key.ToString());
+				}
+			}
+			PrimarySuspects = details.SuspectUsers.Where(x => x.Login != null).Select(x => x.Login).ToList();
+			PrimarySuspectIds= details.SuspectUsers.Select(x => x.Id.ToString()).ToList();
+			PrimarySuspectsInfo = details.SuspectUsers.ConvertAll(x => new GetThinUserInfoResponse(x));
+			ShowDesktopAlerts = showDesktopAlerts;
+			ExternalIssueKey = details.ExternalIssueKey;
+			QuarantinedByUserInfo = details.QuarantinedBy != null ? new GetThinUserInfoResponse(details.QuarantinedBy) : null;
+			QuarantineTimeUtc = details.QuarantineTimeUtc;
+		}
+	}
+
+	/// <summary>
+	/// Information about a span within an issue
+	/// </summary>
+	public class FindIssueSpanResponse
+	{
+		/// <summary>
+		/// Unique id of this span
+		/// </summary>
+		public string Id { get; set; }
+
+		/// <summary>
+		/// The template containing this step
+		/// </summary>
+		public string TemplateId { get; set; }
+
+		/// <summary>
+		/// Name of the step
+		/// </summary>
+		public string Name { get; set; }
+
+		/// <summary>
+		/// Workflow for this span
+		/// </summary>
+		public WorkflowId? WorkflowId { get; set; }
+
+		/// <summary>
+		/// The previous build 
+		/// </summary>
+		public GetIssueStepResponse? LastSuccess { get; set; }
+
+		/// <summary>
+		/// The following successful build
+		/// </summary>
+		public GetIssueStepResponse? NextSuccess { get; set; }
+
+		/// <summary>
+		/// Constructor
+		/// </summary>
+		/// <param name="span"></param>
+		/// <param name="workflowId"></param>
+		public FindIssueSpanResponse(IIssueSpan span, WorkflowId? workflowId)
+		{
+			Id = span.Id.ToString();
+			TemplateId = span.TemplateRefId.ToString();
+			Name = span.NodeName;
+			WorkflowId = workflowId;
+			if (span.LastSuccess != null)
+			{
+				LastSuccess = new GetIssueStepResponse(span.LastSuccess);
+			}
+			if (span.NextSuccess != null)
+			{
+				NextSuccess = new GetIssueStepResponse(span.NextSuccess);
+			}
+		}
+	}
+
+	/// <summary>
+	/// Stores information about a build health issue
+	/// </summary>
+	public class FindIssueResponse
+	{
+		/// <summary>
+		/// The unique object id
+		/// </summary>
+		public int Id { get; set; }
+
+		/// <summary>
+		/// Time at which the issue was created
+		/// </summary>
+		public DateTime CreatedAt { get; set; }
+
+		/// <summary>
+		/// Time at which the issue was retrieved
+		/// </summary>
+		public DateTime RetrievedAt { get; set; }
+
+		/// <summary>
+		/// The associated project for the issue
+		/// </summary>
+		public string? Project { get; set; }
+
+		/// <summary>
+		/// The summary text for this issue
+		/// </summary>
+		public string Summary { get; set; }
+
+		/// <summary>
+		/// Detailed description text
+		/// </summary>
+		public string? Description { get; set; }
+
+		/// <summary>
+		/// Severity of this issue
+		/// </summary>
+		public IssueSeverity Severity { get; set; }
+
+		/// <summary>
+		/// Severity of this issue in the stream
+		/// </summary>
+		public IssueSeverity? StreamSeverity { get; set; }
+
+		/// <summary>
+		/// Whether the issue is promoted
+		/// </summary>
+		public bool Promoted { get; set; }
+
+		/// <summary>
+		/// Owner of the issue
+		/// </summary>
+		public GetThinUserInfoResponse? Owner { get; set; }
+
+		/// <summary>
+		/// Owner of the issue
+		/// </summary>
+		public GetThinUserInfoResponse? NominatedBy { get; set; }
+
+		/// <summary>
+		/// Time that the issue was acknowledged
+		/// </summary>
+		public DateTime? AcknowledgedAt { get; set; }
+
+		/// <summary>
+		/// Changelist that fixed this issue
+		/// </summary>
+		public int? FixChange { get; set; }
+
+		/// <summary>
+		/// Time at which the issue was resolved
+		/// </summary>
+		public DateTime? ResolvedAt { get; set; }
+
+		/// <summary>
+		/// User that resolved the issue
+		/// </summary>
+		public GetThinUserInfoResponse? ResolvedBy { get; set; }
+
+		/// <summary>
+		/// Time at which the issue was verified
+		/// </summary>
+		public DateTime? VerifiedAt { get; set; }
+
+		/// <summary>
+		/// Time that the issue was last seen
+		/// </summary>
+		public DateTime LastSeenAt { get; set; }
+
+		/// <summary>
+		/// Spans for this issue
+		/// </summary>
+		public List<FindIssueSpanResponse> Spans { get; set; }
+
+		/// <summary>
+		/// Key for this issue in external issue tracker
+		/// </summary>
+		public string? ExternalIssueKey { get; set; }
+
+		/// <summary>
+		/// User who quarantined the issue
+		/// </summary>
+		public GetThinUserInfoResponse? QuarantinedBy { get; set; }
+
+		/// <summary>
+		/// The UTC time when the issue was quarantined
+		/// </summary>
+		public DateTime? QuarantineTimeUtc { get; set; }
+
+
+		/// <summary>
+		/// Constructs a new issue
+		/// </summary>
+		/// <param name="issue">The isseu information</param>
+		/// <param name="owner">Owner of the issue</param>
+		/// <param name="nominatedBy">User that nominated the current fixer</param>
+		/// <param name="resolvedBy">User that resolved the issue</param>
+		/// <param name="quarantinedBy">User that quarantined the issue</param>
+		/// <param name="streamSeverity">The current severity in the stream</param>
+		/// <param name="spans">Spans for this issue</param>
+		public FindIssueResponse(IIssue issue, IUser? owner, IUser? nominatedBy, IUser? resolvedBy, IUser? quarantinedBy, IssueSeverity? streamSeverity, List<FindIssueSpanResponse> spans)
+		{
+			Id = issue.Id;
+			CreatedAt = issue.CreatedAt;
+			RetrievedAt = DateTime.UtcNow;
+			Summary = String.IsNullOrEmpty(issue.UserSummary) ? issue.Summary : issue.UserSummary;
+			Description = issue.Description;
+			Severity = issue.Severity;
+			StreamSeverity = streamSeverity;
+			Promoted = issue.Promoted;
+			if (owner != null)
+			{
+				Owner = new GetThinUserInfoResponse(owner);
+			}
+			if (nominatedBy != null)
+			{
+				NominatedBy = new GetThinUserInfoResponse(nominatedBy);
+			}
+			AcknowledgedAt = issue.AcknowledgedAt;
+			FixChange = issue.FixChange;
+			ResolvedAt = issue.ResolvedAt;
+			if (resolvedBy != null)
+			{
+				ResolvedBy = new GetThinUserInfoResponse(resolvedBy);
+			}
+			VerifiedAt = issue.VerifiedAt;
+			LastSeenAt = issue.LastSeenAt;
+			Spans = spans;
+			ExternalIssueKey = issue.ExternalIssueKey;
+			if (quarantinedBy != null)
+			{
+				QuarantinedBy = new GetThinUserInfoResponse(quarantinedBy);
+				QuarantineTimeUtc = issue.QuarantineTimeUtc;
+			}
+		}
+	}
+
+	/// <summary>
+	/// Request an issue to be updated
+	/// </summary>
+	public class UpdateIssueRequest
+	{
+		/// <summary>
+		/// Summary of the issue
+		/// </summary>
+		public string? Summary { get; set; }
+
+		/// <summary>
+		/// Description of the issue
+		/// </summary>
+		public string? Description { get; set; }
+
+		/// <summary>
+		/// Whether the issue is promoted or not
+		/// </summary>
+		public bool? Promoted { get; set; }
+
+		/// <summary>
+		/// New user id for owner of the issue, can be cleared by passing empty string
+		/// </summary>
+		public string? OwnerId { get; set; }
+
+		/// <summary>
+		/// User id that nominated the new owner
+		/// </summary>
+		public string? NominatedById { get; set; }
+
+		/// <summary>
+		/// Whether the issue has been acknowledged
+		/// </summary>
+		public bool? Acknowledged { get; set; }
+
+		/// <summary>
+		/// Whether the user has declined this issue
+		/// </summary>
+		public bool? Declined { get; set; }
+
+		/// <summary>
+		/// The change at which the issue is claimed fixed. 0 = not fixed, -1 = systemic issue.
+		/// </summary>
+		public int? FixChange { get; set; }
+
+		/// <summary>
+		/// Whether the issue should be marked as resolved
+		/// </summary>
+		public bool? Resolved { get; set; }
+
+		/// <summary>
+		/// List of spans to add to this issue
+		/// </summary>
+		public List<string>? AddSpans { get; set; }
+
+		/// <summary>
+		/// List of spans to remove from this issue
+		/// </summary>
+		public List<string>? RemoveSpans { get; set; }
+
+		/// <summary>
+		/// A key to issue in external tracker
+		/// </summary>
+		public string? ExternalIssueKey { get; set; }
+
+		/// <summary>
+		/// Id of user quarantining issue
+		/// </summary>
+		public string? QuarantinedById { get; set; }
+
+	}
+
+
+	/// <summary>
+	/// External issue project information
+	/// </summary>
+	public class GetExternalIssueProjectResponse
+	{
+		/// <summary>
+		/// The project key
+		/// </summary>
+		public string ProjectKey { get; set; }
+
+		/// <summary>
+		/// The name of the project
+		/// </summary>
+		public string Name { get; set; }
+
+		/// <summary>
+		/// The id of the project
+		/// </summary>
+		public string Id { get; set; }
+
+		/// <summary>
+		/// component id => name
+		/// </summary>
+		public Dictionary<string, string> Components { get; set; }
+
+		/// <summary>
+		/// IssueType id => name
+		/// </summary>
+		public Dictionary<string, string> IssueTypes { get; set; }
+
+		/// <summary>
+		/// Constructor
+		/// </summary>
+		/// <param name="key"></param>
+		/// <param name="name"></param>
+		/// <param name="id"></param>
+		/// <param name="components"></param>
+		/// <param name="issueTypes"></param>
+		public GetExternalIssueProjectResponse(string key, string name, string id, Dictionary<string, string> components, Dictionary<string, string> issueTypes)
+		{
+			ProjectKey = key;
+			Name = name; 
+			Id = id; 
+			Components = new Dictionary<string, string>(components); 
+			IssueTypes = new Dictionary<string, string>(issueTypes);
+		}
+	}
+
+
+	/// <summary>
+	/// Request an issue to be created on external issue tracking system
+	/// </summary>
+	public class CreateExternalIssueRequest
+	{
+		/// <summary>
+		/// Horde issue which is linked to external issue
+		/// </summary>
+		public int IssueId { get; set; } = 0;
+
+		/// <summary>
+		/// StreamId of a stream with this issue
+		/// </summary>
+		public string StreamId { get; set; } = String.Empty;
+
+		/// <summary>
+		/// Summary text for external issue
+		/// </summary>
+		public string Summary { get; set; } = String.Empty;
+
+		/// <summary>
+		/// External issue project id
+		/// </summary>
+		public string ProjectId { get; set; } = String.Empty;
+
+		/// <summary>
+		/// External issue component id
+		/// </summary>
+		public string ComponentId { get; set; } = String.Empty;
+
+		/// <summary>
+		/// External issue type id
+		/// </summary>
+		public string IssueTypeId { get; set; } = String.Empty;
+
+		/// <summary>
+		/// Optional description text for external issue
+		/// </summary>
+		public string? Description { get; set; }
+
+		/// <summary>
+		/// Optional link to issue on Horde
+		/// </summary>
+		public string? HordeIssueLink { get; set; }
+
+	}
+
+	/// <summary>
+	/// Response for externally created issue
+	/// </summary>
+	public class CreateExternalIssueResponse
+	{
+		/// <summary>
+		/// External issue key 
+		/// </summary>
+		public string Key { get; set; }
+
+		/// <summary>
+		/// Link to issue on external tracking site
+		/// </summary>
+		public string? Link { get; set; }
+
+		/// <summary>
+		/// Constructor
+		/// </summary>
+		/// <param name="key"></param>
+		/// <param name="link"></param>
+		public CreateExternalIssueResponse(string key, string? link)
+		{
+			Key = key;
+			Link = link;
+		}
+
+	}
+
+	/// <summary>
+	/// External issue response object
+	/// </summary>
+	public class GetExternalIssueResponse : IExternalIssue
+	{
+		/// <summary>
+		/// The external issue key
+		/// </summary>
+		public string Key { get; set; }
+
+		/// <summary>
+		/// The issue link on external tracking site
+		/// </summary>
+		public string? Link { get; set; }
+
+		/// <summary>
+		/// The issue status name, "To Do", "In Progress", etc
+		/// </summary>
+		public string? StatusName { get; set; }
+
+		/// <summary>
+		/// The issue resolution name, "Fixed", "Closed", etc
+		/// </summary>
+		public string? ResolutionName { get; set; }
+
+		/// <summary>
+		/// The issue priority name, "1 - Critical", "2 - Major", etc
+		/// </summary>
+		public string? PriorityName { get; set; }
+
+		/// <summary>
+		/// The current assignee's user name
+		/// </summary>
+		public string? AssigneeName { get; set; }
+
+		/// <summary>
+		/// The current assignee's display name
+		/// </summary>
+		public string? AssigneeDisplayName { get; set; }
+
+		/// <summary>
+		/// The current assignee's email address
+		/// </summary>
+		public string? AssigneeEmailAddress { get; set; }
+
+		/// <summary>
+		/// Response constructor
+		/// </summary>
+		public GetExternalIssueResponse(IExternalIssue issue)
+		{
+			Key = issue.Key;
+			Link = issue.Link;
+			StatusName = issue.StatusName;
+			ResolutionName = issue.ResolutionName;
+			PriorityName = issue.PriorityName;
+			AssigneeName = issue.AssigneeName;
+			AssigneeDisplayName = issue.AssigneeDisplayName;
+			AssigneeEmailAddress = issue.AssigneeEmailAddress;
+		}
+
+
+	}
+
+}
diff --git a/Engine/Source/Programs/Horde/Horde.Build/Issues/IssuesController.Messages.cs b/Engine/Source/Programs/Horde/Horde.Build/Issues/IssuesController.Messages.cs
deleted file mode 100644
diff --git a/Engine/Source/Programs/Horde/Horde.Build/Jobs/Artifacts/ArtifactMessages.cs b/Engine/Source/Programs/Horde/Horde.Build/Jobs/Artifacts/ArtifactMessages.cs
new file mode 100644
--- /dev/null
+++ b/Engine/Source/Programs/Horde/Horde.Build/Jobs/Artifacts/ArtifactMessages.cs
@@ -0,0 +1,118 @@
+// Copyright Epic Games, Inc. All Rights Reserved.
+
+using System.Collections.Generic;
+using Horde.Build.Jobs.Artifacts;
+
+namespace Horde.Build.Jobs.Artifacts
+{
+	/// <summary>
+	/// Response from creating an artifact
+	/// </summary>
+	public class CreateArtifactResponse
+	{
+		/// <summary>
+		/// Unique id for this artifact
+		/// </summary>
+		public string Id { get; set; }
+
+		/// <summary>
+		/// Constructor
+		/// </summary>
+		/// <param name="id">The artifact file id</param>
+		public CreateArtifactResponse(string id)
+		{
+			Id = id;
+		}
+	}
+
+	/// <summary>
+	/// Request to get a zip of artifacts
+	/// </summary>
+	public class GetArtifactZipRequest
+	{
+		/// <summary>
+		/// Job id to get all artifacts for
+		/// </summary>
+		public string? JobId { get; set; }
+
+		/// <summary>
+		/// Step id to filter by
+		/// </summary>
+		public string? StepId { get; set; }
+
+		/// <summary>
+		/// Further filter by a list of artifact ids
+		/// </summary>
+		public List<string>? ArtifactIds { get; set; }
+
+		/// <summary>
+		/// Constructor
+		/// </summary>
+		/// <param name="jobId">Job id to get all artifacts for</param>
+		/// <param name="stepId">step to filter by</param>
+		/// <param name="artifactIds">The artifact ids.  Returns all artifacts for a job </param>
+		public GetArtifactZipRequest(string? jobId, string? stepId, List<string>? artifactIds)
+		{
+			JobId = jobId;
+			StepId = stepId;
+			ArtifactIds = artifactIds;
+		}
+	}
+
+	/// <summary>
+	/// Response describing an artifact
+	/// </summary>
+	public class GetArtifactResponse
+	{
+		/// <summary>
+		/// Unique id of the artifact
+		/// </summary>
+		public string Id { get; set; }
+
+		/// <summary>
+		/// Name of the artifact
+		/// </summary>
+		public string Name { get; set; }
+
+		/// <summary>
+		/// Unique id of the job for this artifact
+		/// </summary>
+		public string JobId { get; set; }
+
+		/// <summary>
+		/// Optional unique id of the step for this artifact
+		/// </summary>
+		public string? StepId { get; set; }
+
+		/// <summary>
+		/// Type of artifact
+		/// </summary>
+		public string MimeType { get; set; }
+
+		/// <summary>
+		/// Length of artifact
+		/// </summary>
+		public long Length { get; set; }
+
+		/// <summary>
+		/// Short-lived code that can be used to download this artifact through direct download links in the browser
+		/// </summary>
+		public string? Code { get; set; }
+
+		/// <summary>
+		/// Constructor
+		/// </summary>
+		/// <param name="artifact">The artifact to construct from</param>
+		/// <param name="code">The direct download code</param>
+		public GetArtifactResponse(IArtifact artifact, string? code)
+		{
+			Id = artifact.Id.ToString();
+			Name = artifact.Name;
+			JobId = artifact.JobId.ToString();
+			StepId = artifact.StepId.ToString();
+			MimeType = artifact.MimeType;
+			Length = artifact.Length;
+			Code = code;
+		}
+	}
+}
diff --git a/Engine/Source/Programs/Horde/Horde.Build/Jobs/Artifacts/ArtifactsController.Messages.cs b/Engine/Source/Programs/Horde/Horde.Build/Jobs/Artifacts/ArtifactsController.Messages.cs
deleted file mode 100644
diff --git a/Engine/Source/Programs/Horde/Horde.Build/Jobs/JobMessages.cs b/Engine/Source/Programs/Horde/Horde.Build/Jobs/JobMessages.cs
new file mode 100644
--- /dev/null
+++ b/Engine/Source/Programs/Horde/Horde.Build/Jobs/JobMessages.cs
@@ -0,0 +1,1139 @@
+// Copyright Epic Games, Inc. All Rights Reserved.
+
+using System;
+using System.Collections.Generic;
+using System.ComponentModel.DataAnnotations;
+using System.Linq;
+using System.Text.Json.Serialization;
+using EpicGames.Core;
+using Horde.Build.Acls;
+using Horde.Build.Jobs;
+using Horde.Build.Jobs.Graphs;
+using Horde.Build.Streams;
+using Horde.Build.Users;
+using HordeCommon;
+
+namespace Horde.Build.Jobs
+{
+	/// <summary>
+	/// State of the job
+	/// </summary>
+	public enum JobState
+	{
+		/// <summary>
+		/// Waiting for resources
+		/// </summary>
+		Waiting,
+
+		/// <summary>
+		/// Currently running one or more steps
+		/// </summary>
+		Running,
+
+		/// <summary>
+		/// All steps have completed
+		/// </summary>
+		Complete,
+	}
+
+	/// <summary>
+	/// Parameters required to create a job
+	/// </summary>
+	public class CreateJobRequest
+	{
+		/// <summary>
+		/// The stream that this job belongs to
+		/// </summary>
+		[Required]
+		public string StreamId { get; set; }
+
+		/// <summary>
+		/// The template for this job
+		/// </summary>
+		[Required]
+		public string TemplateId { get; set; }
+
+		/// <summary>
+		/// Name of the job
+		/// </summary>
+		public string? Name { get; set; }
+
+		/// <summary>
+		/// The changelist number to build. Can be null for latest.
+		/// </summary>
+		public int? Change { get; set; }
+
+		/// <summary>
+		/// Parameters to use when selecting the change to execute at.
+		/// </summary>
+		public ChangeQueryRequest? ChangeQuery { get; set; }
+
+		/// <summary>
+		/// The preflight changelist number
+		/// </summary>
+		public int? PreflightChange { get; set; }
+
+		/// <summary>
+		/// Priority for the job
+		/// </summary>
+		public Priority? Priority { get; set; }
+
+		/// <summary>
+		/// Whether to automatically submit the preflighted change on completion
+		/// </summary>
+		public bool? AutoSubmit { get; set; }
+
+		/// <summary>
+		/// Whether to update issues based on the outcome of this job
+		/// </summary>
+		public bool? UpdateIssues { get; set; }
+
+		/// <summary>
+		/// Arguments for the job
+		/// </summary>
+		public List<string>? Arguments { get; set; }
+
+		/// <summary>
+		/// Private constructor for serialization
+		/// </summary>
+		public CreateJobRequest(string streamId, string templateId)
+		{
+			StreamId = streamId;
+			TemplateId = templateId;
+		}
+	}
+
+	/// <summary>
+	/// Response from creating a new job
+	/// </summary>
+	public class CreateJobResponse
+	{
+		/// <summary>
+		/// Unique id for the new job
+		/// </summary>
+		public string Id { get; set; }
+
+		/// <summary>
+		/// Constructor
+		/// </summary>
+		/// <param name="id">Unique id for the new job</param>
+		public CreateJobResponse(string id)
+		{
+			Id = id;
+		}
+	}
+
+	/// <summary>
+	/// Updates an existing job
+	/// </summary>
+	public class UpdateJobRequest
+	{
+		/// <summary>
+		/// New name for the job
+		/// </summary>
+		public string? Name { get; set; }
+
+		/// <summary>
+		/// New priority for the job
+		/// </summary>
+		public Priority? Priority { get; set; }
+
+		/// <summary>
+		/// Set whether the job should be automatically submitted or not
+		/// </summary>
+		public bool? AutoSubmit { get; set; }
+
+		/// <summary>
+		/// Mark this job as aborted
+		/// </summary>
+		public bool? Aborted { get; set; }
+
+		/// <summary>
+		/// New list of arguments for the job. Only -Target= arguments can be modified after the job has started.
+		/// </summary>
+		public List<string>? Arguments { get; set; }
+
+		/// <summary>
+		/// Custom permissions for this object
+		/// </summary>
+		public UpdateAclRequest? Acl { get; set; }
+	}
+
+	/// <summary>
+	/// Information about a report associated with a job
+	/// </summary>
+	public class GetReportResponse
+	{
+		/// <summary>
+		/// Name of the report
+		/// </summary>
+		public string Name { get; set; }
+
+		/// <summary>
+		/// Report placement
+		/// </summary>
+		public HordeCommon.Rpc.ReportPlacement Placement { get; set; }
+
+		/// <summary>
+		/// The artifact id
+		/// </summary>
+		public string ArtifactId { get; set; }
+
+		/// <summary>
+		/// Constructor
+		/// </summary>
+		/// <param name="report"></param>
+		public GetReportResponse(IReport report)
+		{
+			Name = report.Name;
+			Placement = report.Placement;
+			ArtifactId = report.ArtifactId.ToString();
+		}
+	}
+
+	/// <summary>
+	/// Information about a job
+	/// </summary>
+	public class GetJobResponse
+	{
+		/// <summary>
+		/// Unique Id for the job
+		/// </summary>
+		public string Id { get; set; }
+
+		/// <summary>
+		/// Unique id of the stream containing this job
+		/// </summary>
+		public string StreamId { get; set; }
+
+		/// <summary>
+		/// Name of the job
+		/// </summary>
+		public string Name { get; set; }
+
+		/// <summary>
+		/// The changelist number to build
+		/// </summary>
+		public int Change { get; set; }
+
+		/// <summary>
+		/// The code changelist
+		/// </summary>
+		public int? CodeChange { get; set; }
+
+		/// <summary>
+		/// The preflight changelist number
+		/// </summary>
+		public int? PreflightChange { get; set; }
+
+		/// <summary>
+		/// The cloned preflight changelist number
+		/// </summary>
+		public int? ClonedPreflightChange { get; set; }
+
+		/// <summary>
+		/// The template type
+		/// </summary>
+		public string TemplateId { get; set; }
+
+		/// <summary>
+		/// Hash of the actual template data
+		/// </summary>
+		public string TemplateHash { get; set; }
+
+		/// <summary>
+		/// Hash of the graph for this job
+		/// </summary>
+		public string? GraphHash { get; set; }
+
+		/// <summary>
+		/// The user that started this job [DEPRECATED]
+		/// </summary>
+		public string? StartedByUserId { get; set; }
+
+		/// <summary>
+		/// The user that started this job [DEPRECATED]
+		/// </summary>
+		public string? StartedByUser { get; set; }
+
+		/// <summary>
+		/// The user that started this job
+		/// </summary>
+		public GetThinUserInfoResponse? StartedByUserInfo { get; set; }
+
+		/// <summary>
+		/// The user that aborted this job [DEPRECATED]
+		/// </summary>
+		public string? AbortedByUser { get; set; }
+
+		/// <summary>
+		/// The user that aborted this job
+		/// </summary>
+		public GetThinUserInfoResponse? AbortedByUserInfo { get; set; }
+
+		/// <summary>
+		/// Priority of the job
+		/// </summary>
+		public Priority Priority { get; set; }
+
+		/// <summary>
+		/// Whether the change will automatically be submitted or not
+		/// </summary>
+		public bool AutoSubmit { get; set; }
+
+		/// <summary>
+		/// The submitted changelist number
+		/// </summary>
+		public int? AutoSubmitChange { get; }
+
+		/// <summary>
+		/// Message produced by trying to auto-submit the change
+		/// </summary>
+		public string? AutoSubmitMessage { get; }
+
+		/// <summary>
+		/// Time that the job was created
+		/// </summary>
+		public DateTimeOffset CreateTime { get; set; }
+
+		/// <summary>
+		/// The global job state
+		/// </summary>
+		public JobState State { get; set; }
+
+		/// <summary>
+		/// Array of jobstep batches
+		/// </summary>
+		public List<GetBatchResponse>? Batches { get; set; }
+
+		/// <summary>
+		/// List of labels
+		/// </summary>
+		public List<GetLabelStateResponse>? Labels { get; set; }
+
+		/// <summary>
+		/// The default label, containing the state of all steps that are otherwise not matched.
+		/// </summary>
+		public GetDefaultLabelStateResponse? DefaultLabel { get; set; } 
+
+		/// <summary>
+		/// List of reports
+		/// </summary>
+		public List<GetReportResponse>? Reports { get; set; }
+
+		/// <summary>
+		/// Parameters for the job
+		/// </summary>
+		public List<string> Arguments { get; set; }
+
+		/// <summary>
+		/// The last update time for this job
+		/// </summary>
+		public DateTimeOffset UpdateTime { get; set; }
+
+		/// <summary>
+		/// Whether issues are being updated by this job
+		/// </summary>
+		public bool UpdateIssues { get; set; }
+
+		/// <summary>
+		/// Per-object permissions
+		/// </summary>
+		public GetAclResponse? Acl { get; set; }
+
+		/// <summary>
+		/// Constructor
+		/// </summary>
+		/// <param name="job">Job to create a response for</param>
+		/// <param name="startedByUserInfo">User that started this job</param>
+		/// <param name="abortedByUserInfo">User that aborted this job</param>
+		/// <param name="aclResponse">The ACL response object</param>
+		public GetJobResponse(IJob job, GetThinUserInfoResponse? startedByUserInfo, GetThinUserInfoResponse? abortedByUserInfo, GetAclResponse? aclResponse)
+		{
+			Id = job.Id.ToString();
+			StreamId = job.StreamId.ToString();
+			Name = job.Name;
+			Change = job.Change;
+			CodeChange = (job.CodeChange != 0) ? (int?)job.CodeChange : null;
+			PreflightChange = (job.PreflightChange != 0) ? (int?)job.PreflightChange : null;
+			ClonedPreflightChange = (job.ClonedPreflightChange != 0) ? (int?)job.ClonedPreflightChange : null;
+			TemplateId = job.TemplateId.ToString();
+			TemplateHash = job.TemplateHash?.ToString() ?? String.Empty;
+			GraphHash = job.GraphHash.ToString();
+			StartedByUserId = job.StartedByUserId?.ToString();
+			StartedByUser = startedByUserInfo?.Login;
+			StartedByUserInfo = startedByUserInfo;
+			AbortedByUser = abortedByUserInfo?.Login;
+			AbortedByUserInfo = abortedByUserInfo;
+			CreateTime = new DateTimeOffset(job.CreateTimeUtc);
+			State = job.GetState();
+			Priority = job.Priority;
+			AutoSubmit = job.AutoSubmit;
+			AutoSubmitChange = job.AutoSubmitChange;
+			AutoSubmitMessage = job.AutoSubmitMessage;
+			Reports = job.Reports?.ConvertAll(x => new GetReportResponse(x));
+			Arguments = job.Arguments.ToList();
+			UpdateTime = new DateTimeOffset(job.UpdateTimeUtc);
+			UpdateIssues = job.UpdateIssues;
+			Acl = aclResponse;
+		}
+	}
+
+	/// <summary>
+	/// The timing info for 
+	/// </summary>
+	public class GetJobTimingResponse
+	{
+		/// <summary>
+		/// The job object
+		/// </summary>
+		public IJob? Job { get; set; }
+		
+		/// <summary>
+		/// A full job response
+		/// </summary>
+		public GetJobResponse? JobResponse { get; set; }
+		
+		/// <summary>
+		/// Timing info for each step
+		/// </summary>
+		public Dictionary<string, GetStepTimingInfoResponse> Steps { get; set; }
+
+		/// <summary>
+		/// Timing information for each label
+		/// </summary>
+		public List<GetLabelTimingInfoResponse> Labels { get; set; }
+
+		/// <summary>
+		/// Constructor
+		/// </summary>
+		/// <param name="job">The job object</param>
+		/// <param name="jobResponse">The job response</param>
+		/// <param name="steps">Timing info for each steps</param>
+		/// <param name="labels">Timing info for each label</param>
+		public GetJobTimingResponse(IJob? job, GetJobResponse? jobResponse, Dictionary<string, GetStepTimingInfoResponse> steps, List<GetLabelTimingInfoResponse> labels)
+		{
+			Job = job;
+			JobResponse = jobResponse;
+			Steps = steps;
+			Labels = labels;
+		}
+	}
+	
+	/// <summary>
+	/// The timing info for 
+	/// </summary>
+	public class FindJobTimingsResponse
+	{
+		/// <summary>
+		/// Timing info for each job
+		/// </summary>
+		public Dictionary<string, GetJobTimingResponse> Timings { get; set; }
+
+		/// <summary>
+		/// Constructor
+		/// </summary>
+		/// <param name="timings">Timing info for each job</param>
+		public FindJobTimingsResponse(Dictionary<string, GetJobTimingResponse> timings)
+		{
+			Timings = timings;
+		}
+	}
+
+	/// <summary>
+	/// Request used to update a jobstep
+	/// </summary>
+	public class UpdateStepRequest
+	{
+		/// <summary>
+		/// The new jobstep state
+		/// </summary>
+		public JobStepState State { get; set; } = JobStepState.Unspecified;
+
+		/// <summary>
+		/// Outcome from the jobstep
+		/// </summary>
+		public JobStepOutcome Outcome { get; set; } = JobStepOutcome.Unspecified;
+		
+		/// <summary>
+		/// If the step has been requested to abort
+		/// </summary>
+		public bool? AbortRequested { get; set; }
+
+		/// <summary>
+		/// Specifies the log file id for this step
+		/// </summary>
+		public string? LogId { get; set; }
+
+		/// <summary>
+		/// Whether the step should be re-run
+		/// </summary>
+		public bool? Retry { get; set; }
+
+		/// <summary>
+		/// New priority for this step
+		/// </summary>
+		public Priority? Priority { get; set; }
+
+		/// <summary>
+		/// Properties to set. Any entries with a null value will be removed.
+		/// </summary>
+		public Dictionary<string, string?>? Properties { get; set; }
+	}
+
+	/// <summary>
+	/// Response object when updating a jobstep
+	/// </summary>
+	public class UpdateStepResponse
+	{
+		/// <summary>
+		/// If a new step is created (due to specifying the retry flag), specifies the batch id
+		/// </summary>
+		public string? BatchId { get; set; }
+
+		/// <summary>
+		/// If a step is retried, includes the new step id
+		/// </summary>
+		public string? StepId { get; set; }
+	}
+
+	/// <summary>
+	/// Returns information about a jobstep
+	/// </summary>
+	public class GetStepResponse
+	{
+		/// <summary>
+		/// The unique id of the step
+		/// </summary>
+		public string Id { get; set; }
+
+		/// <summary>
+		/// Index of the node which this jobstep is to execute
+		/// </summary>
+		public int NodeIdx { get; set; }
+
+		/// <summary>
+		/// Current state of the job step. This is updated automatically when runs complete.
+		/// </summary>
+		public JobStepState State { get; set; }
+
+		/// <summary>
+		/// Current outcome of the jobstep
+		/// </summary>
+		public JobStepOutcome Outcome { get; set; }
+
+		/// <summary>
+		/// Error describing additional context for why a step failed to complete
+		/// </summary>
+		public JobStepError Error { get; set; }
+		
+		/// <summary>
+		/// If the step has been requested to abort
+		/// </summary>
+		public bool AbortRequested { get; set; }
+		
+		/// <summary>
+		/// Name of the user that requested the abort of this step [DEPRECATED]
+		/// </summary>
+		public string? AbortByUser { get; set; }
+
+		/// <summary>
+		/// The user that requested this step be run again 
+		/// </summary>
+		public GetThinUserInfoResponse? AbortedByUserInfo { get; set; }
+
+		/// <summary>
+		/// Name of the user that requested this step be run again [DEPRECATED]
+		/// </summary>
+		public string? RetryByUser { get; set; }
+
+		/// <summary>
+		/// The user that requested this step be run again 
+		/// </summary>
+		public GetThinUserInfoResponse? RetriedByUserInfo { get; set; }
+
+		/// <summary>
+		/// The log id for this step
+		/// </summary>
+		public string? LogId { get; set; }
+
+		/// <summary>
+		/// Time at which the batch was ready (UTC).
+		/// </summary>
+		public DateTimeOffset? ReadyTime { get; set; }
+
+		/// <summary>
+		/// Time at which the batch started (UTC).
+		/// </summary>
+		public DateTimeOffset? StartTime { get; set; }
+
+		/// <summary>
+		/// Time at which the batch finished (UTC)
+		/// </summary>
+		public DateTimeOffset? FinishTime { get; set; }
+
+		/// <summary>
+		/// List of reports
+		/// </summary>
+		public List<GetReportResponse>? Reports { get; set; }
+
+		/// <summary>
+		/// User-defined properties for this jobstep.
+		/// </summary>
+		public Dictionary<string, string>? Properties { get; set; }
+
+		/// <summary>
+		/// Constructor
+		/// </summary>
+		/// <param name="step">The step to construct from</param>
+		/// <param name="abortedByUserInfo">User that aborted this step</param>
+		/// <param name="retriedByUserInfo">User that retried this step</param>
+		public GetStepResponse(IJobStep step, GetThinUserInfoResponse? abortedByUserInfo, GetThinUserInfoResponse? retriedByUserInfo)
+		{
+			Id = step.Id.ToString();
+			NodeIdx = step.NodeIdx;
+			State = step.State;
+			Outcome = step.Outcome;
+			Error = step.Error;
+			AbortRequested = step.AbortRequested;
+			AbortByUser = abortedByUserInfo?.Login;
+			AbortedByUserInfo = abortedByUserInfo;
+			RetryByUser = retriedByUserInfo?.Login;
+			RetriedByUserInfo = retriedByUserInfo;
+			LogId = step.LogId?.ToString();
+			ReadyTime = step.ReadyTimeUtc;
+			StartTime = step.StartTimeUtc;
+			FinishTime = step.FinishTimeUtc;
+			Reports = step.Reports?.ConvertAll(x => new GetReportResponse(x));
+
+			if (step.Properties != null && step.Properties.Count > 0)
+			{
+				Properties = step.Properties;
+			}
+		}
+	}
+
+	/// <summary>
+	/// The state of a particular run
+	/// </summary>
+	[JsonConverter(typeof(JsonStringEnumConverter))]
+	public enum JobStepBatchState
+	{
+		/// <summary>
+		/// Waiting for dependencies of at least one jobstep to complete
+		/// </summary>
+		Waiting = 0,
+
+		/// <summary>
+		/// Ready to execute
+		/// </summary>
+		Ready = 1,
+
+		/// <summary>
+		/// Preparing to execute work
+		/// </summary>
+		Starting = 2,
+
+		/// <summary>
+		/// Executing work
+		/// </summary>
+		Running = 3,
+
+		/// <summary>
+		/// Preparing to stop
+		/// </summary>
+		Stopping = 4,
+
+		/// <summary>
+		/// All steps have finished executing
+		/// </summary>
+		Complete = 5
+	}
+
+	/// <summary>
+	/// Error code for a batch not being executed
+	/// </summary>
+	[JsonConverter(typeof(JsonStringEnumConverter))]
+	public enum JobStepBatchError
+	{
+		/// <summary>
+		/// No error
+		/// </summary>
+		None = 0,
+
+		/// <summary>
+		/// The stream for this job is unknown
+		/// </summary>
+		UnknownStream = 1,
+
+		/// <summary>
+		/// The given agent type for this batch was not valid for this stream
+		/// </summary>
+		UnknownAgentType = 2,
+
+		/// <summary>
+		/// The pool id referenced by the agent type was not found
+		/// </summary>
+		UnknownPool = 3,
+
+		/// <summary>
+		/// There are no agents in the given pool currently online
+		/// </summary>
+		NoAgentsInPool = 4,
+
+		/// <summary>
+		/// There are no agents in this pool that are onlinbe
+		/// </summary>
+		NoAgentsOnline = 5,
+
+		/// <summary>
+		/// Unknown workspace referenced by the agent type
+		/// </summary>
+		UnknownWorkspace = 6,
+
+		/// <summary>
+		/// Cancelled
+		/// </summary>
+		Cancelled = 7,
+
+		/// <summary>
+		/// Lost connection with the agent machine
+		/// </summary>
+		LostConnection = 8,
+
+		/// <summary>
+		/// Lease terminated prematurely but can be retried.
+		/// </summary>
+		Incomplete = 9,
+
+		/// <summary>
+		/// An error ocurred while executing the lease. Cannot be retried.
+		/// </summary>
+		ExecutionError = 10,
+	}
+
+	/// <summary>
+	/// Request to update a jobstep batch
+	/// </summary>
+	public class UpdateBatchRequest
+	{
+		/// <summary>
+		/// The new log file id
+		/// </summary>
+		public string? LogId { get; set; }
+
+		/// <summary>
+		/// The state of this batch
+		/// </summary>
+		public JobStepBatchState? State { get; set; }
+	}
+
+	/// <summary>
+	/// Information about a jobstep batch
+	/// </summary>
+	public class GetBatchResponse
+	{
+		/// <summary>
+		/// Unique id for this batch
+		/// </summary>
+		public string Id { get; set; }
+
+		/// <summary>
+		/// The unique log file id
+		/// </summary>
+		public string? LogId { get; set; }
+
+		/// <summary>
+		/// Index of the group being executed
+		/// </summary>
+		public int GroupIdx { get; set; }
+
+		/// <summary>
+		/// The state of this batch
+		/// </summary>
+		public JobStepBatchState State { get; set; }
+
+		/// <summary>
+		/// Error code for this batch
+		/// </summary>
+		public JobStepBatchError Error { get; set; }
+
+		/// <summary>
+		/// Steps within this run
+		/// </summary>
+		public List<GetStepResponse> Steps { get; set; }
+
+		/// <summary>
+		/// The agent assigned to execute this group
+		/// </summary>
+		public string? AgentId { get; set; }
+
+		/// <summary>
+		/// Rate for using this agent (per hour)
+		/// </summary>
+		public double? AgentRate { get; set; }
+
+		/// <summary>
+		/// The agent session holding this lease
+		/// </summary>
+		public string? SessionId { get; set; }
+
+		/// <summary>
+		/// The lease that's executing this group
+		/// </summary>
+		public string? LeaseId { get; set; }
+
+		/// <summary>
+		/// The priority of this batch
+		/// </summary>
+		public int WeightedPriority { get; set; }
+
+		/// <summary>
+		/// Time at which the group started (UTC).
+		/// </summary>
+		public DateTimeOffset? StartTime { get; set; }
+
+		/// <summary>
+		/// Time at which the group finished (UTC)
+		/// </summary>
+		public DateTimeOffset? FinishTime { get; set; }
+
+		/// <summary>
+		/// Converts this batch into a public response object
+		/// </summary>
+		/// <param name="batch">The batch to construct from</param>
+		/// <param name="steps">Steps in this batch</param>
+		/// <param name="agentRate">Rate for this agent</param>
+		/// <returns>Response instance</returns>
+		public GetBatchResponse(IJobStepBatch batch, List<GetStepResponse> steps, double? agentRate)
+		{
+			Id = batch.Id.ToString();
+			LogId = batch.LogId?.ToString();
+			GroupIdx = batch.GroupIdx;
+			State = batch.State;
+			Error = batch.Error;
+			Steps = steps;
+			AgentId = batch.AgentId?.ToString();
+			AgentRate = agentRate;
+			SessionId = batch.SessionId?.ToString();
+			LeaseId = batch.LeaseId?.ToString();
+			WeightedPriority = batch.SchedulePriority;
+			StartTime = batch.StartTimeUtc;
+			FinishTime = batch.FinishTimeUtc;
+		}
+	}
+
+	/// <summary>
+	/// State of an aggregate
+	/// </summary>
+	public enum LabelState
+	{
+		/// <summary>
+		/// Aggregate is not currently being built (no required nodes are present)
+		/// </summary>
+		Unspecified,
+
+		/// <summary>
+		/// Steps are still running
+		/// </summary>
+		Running,
+
+		/// <summary>
+		/// All steps are complete
+		/// </summary>
+		Complete
+	}
+
+	/// <summary>
+	/// Outcome of an aggregate
+	/// </summary>
+	public enum LabelOutcome
+	{
+		/// <summary>
+		/// Aggregate is not currently being built
+		/// </summary>
+		Unspecified,
+
+		/// <summary>
+		/// A step dependency failed
+		/// </summary>
+		Failure,
+
+		/// <summary>
+		/// A dependency finished with warnings
+		/// </summary>
+		Warnings,
+
+		/// <summary>
+		/// Successful
+		/// </summary>
+		Success,
+	}
+
+	/// <summary>
+	/// State of a label within a job
+	/// </summary>
+	public class GetLabelStateResponse
+	{
+		/// <summary>
+		/// State of the label
+		/// </summary>
+		public LabelState? State { get; set; }
+
+		/// <summary>
+		/// Outcome of the label
+		/// </summary>
+		public LabelOutcome? Outcome { get; set; }
+
+		/// <summary>
+		/// Constructor
+		/// </summary>
+		/// <param name="state">State of the label</param>
+		/// <param name="outcome">Outcome of the label</param>
+		public GetLabelStateResponse(LabelState state, LabelOutcome outcome)
+		{
+			State = state;
+			Outcome = outcome;
+		}
+	}
+
+	/// <summary>
+	/// Information about the default label (ie. with inlined list of nodes)
+	/// </summary>
+	public class GetDefaultLabelStateResponse : GetLabelStateResponse
+	{
+		/// <summary>
+		/// List of nodes covered by default label
+		/// </summary>
+		public List<string> Nodes { get; set; }
+
+		/// <summary>
+		/// Constructor
+		/// </summary>
+		/// <param name="state">State of the label</param>
+		/// <param name="outcome">Outcome of the label</param>
+		/// <param name="nodes">List of nodes that are covered by the default label</param>
+		public GetDefaultLabelStateResponse(LabelState state, LabelOutcome outcome, List<string> nodes)
+			: base(state, outcome)
+		{
+			Nodes = nodes;
+		}
+	}
+
+	/// <summary>
+	/// Information about the timing info for a particular target
+	/// </summary>
+	public class GetTimingInfoResponse
+	{
+		/// <summary>
+		/// Wait time on the critical path
+		/// </summary>
+		public float? TotalWaitTime { get; set; }
+
+		/// <summary>
+		/// Sync time on the critical path
+		/// </summary>
+		public float? TotalInitTime { get; set; }
+
+		/// <summary>
+		/// Duration to this point
+		/// </summary>
+		public float? TotalTimeToComplete { get; set; }
+
+		/// <summary>
+		/// Average wait time by the time the job reaches this point
+		/// </summary>
+		public float? AverageTotalWaitTime { get; set; }
+
+		/// <summary>
+		/// Average sync time to this point
+		/// </summary>
+		public float? AverageTotalInitTime { get; set; }
+
+		/// <summary>
+		/// Average duration to this point
+		/// </summary>
+		public float? AverageTotalTimeToComplete { get; set; }
+
+		/// <summary>
+		/// Constructor
+		/// </summary>
+		/// <param name="timingInfo">Timing info to construct from</param>
+		public GetTimingInfoResponse(TimingInfo timingInfo)
+		{
+			TotalWaitTime = (float?)timingInfo.TotalWaitTime?.TotalSeconds;
+			TotalInitTime = (float?)timingInfo.TotalInitTime?.TotalSeconds;
+			TotalTimeToComplete = (float?)timingInfo.TotalTimeToComplete?.TotalSeconds;
+
+			AverageTotalWaitTime = (float?)timingInfo.AverageTotalWaitTime?.TotalSeconds;
+			AverageTotalInitTime = (float?)timingInfo.AverageTotalInitTime?.TotalSeconds;
+			AverageTotalTimeToComplete = (float?)timingInfo.AverageTotalTimeToComplete?.TotalSeconds;
+		}
+	}
+
+	/// <summary>
+	/// Information about the timing info for a particular target
+	/// </summary>
+	public class GetStepTimingInfoResponse : GetTimingInfoResponse
+	{
+		/// <summary>
+		/// Name of this node
+		/// </summary>
+		public string? Name { get; set; }
+
+		/// <summary>
+		/// Average wait time for this step
+		/// </summary>
+		public float? AverageStepWaitTime { get; set; }
+
+		/// <summary>
+		/// Average init time for this step
+		/// </summary>
+		public float? AverageStepInitTime { get; set; }
+
+		/// <summary>
+		/// Average duration for this step
+		/// </summary>
+		public float? AverageStepDuration { get; set; }
+
+		/// <summary>
+		/// Constructor
+		/// </summary>
+		/// <param name="name">Name of the node</param>
+		/// <param name="jobTimingInfo">Timing info to construct from</param>
+		public GetStepTimingInfoResponse(string? name, TimingInfo jobTimingInfo)
+			: base(jobTimingInfo)
+		{
+			Name = name;
+			AverageStepWaitTime = jobTimingInfo.StepTiming?.AverageWaitTime;
+			AverageStepInitTime = jobTimingInfo.StepTiming?.AverageInitTime;
+			AverageStepDuration = jobTimingInfo.StepTiming?.AverageDuration;
+		}
+	}
+
+	/// <summary>
+	/// Information about the timing info for a label
+	/// </summary>
+	public class GetLabelTimingInfoResponse : GetTimingInfoResponse
+	{
+		/// <summary>
+		/// Name of the label
+		/// </summary>
+		[Obsolete("Use DashboardName instead")]
+		public string? Name => DashboardName;
+
+		/// <summary>
+		/// Category for the label
+		/// </summary>
+		[Obsolete("Use DashboardCategory instead")]
+		public string? Category => DashboardCategory;
+
+		/// <summary>
+		/// Name of the label
+		/// </summary>
+		public string? DashboardName { get; set; }
+
+		/// <summary>
+		/// Category for the label
+		/// </summary>
+		public string? DashboardCategory { get; set; }
+
+		/// <summary>
+		/// Name of the label
+		/// </summary>
+		public string? UgsName { get; set; }
+
+		/// <summary>
+		/// Category for the label
+		/// </summary>
+		public string? UgsProject { get; set; }
+
+		/// <summary>
+		/// Constructor
+		/// </summary>
+		/// <param name="label">The label to construct from</param>
+		/// <param name="timingInfo">Timing info to construct from</param>
+		public GetLabelTimingInfoResponse(ILabel label, TimingInfo timingInfo)
+			: base(timingInfo)
+		{
+			DashboardName = label.DashboardName;
+			DashboardCategory = label.DashboardCategory;
+			UgsName = label.UgsName;
+			UgsProject = label.UgsProject;
+		}
+	}
+
+	/// <summary>
+	/// Describes the history of a step
+	/// </summary>
+	public class GetJobStepRefResponse
+	{
+		/// <summary>
+		/// The job id
+		/// </summary>
+		public string JobId { get; set; }
+
+		/// <summary>
+		/// The batch containing the step
+		/// </summary>
+		public string BatchId { get; set; }
+
+		/// <summary>
+		/// The step identifier
+		/// </summary>
+		public string StepId { get; set; }
+
+		/// <summary>
+		/// The change number being built
+		/// </summary>
+		public int Change { get; set; }
+
+		/// <summary>
+		/// The step log id
+		/// </summary>
+		public string? LogId { get; set; }
+
+		/// <summary>
+		/// The pool id
+		/// </summary>
+		public string? PoolId { get; set; }
+
+		/// <summary>
+		/// The agent id
+		/// </summary>
+		public string? AgentId { get; set; }
+
+		/// <summary>
+		/// Outcome of the step, once complete.
+		/// </summary>
+		public JobStepOutcome? Outcome { get; set; }
+
+		/// <summary>
+		/// Time at which the step started.
+		/// </summary>
+		public DateTimeOffset StartTime { get; }
+
+		/// <summary>
+		/// Time at which the step finished.
+		/// </summary>
+		public DateTimeOffset? FinishTime { get; }
+
+		/// <summary>
+		/// Constructor
+		/// </summary>
+		/// <param name="jobStepRef">The jobstep ref to construct from</param>
+		public GetJobStepRefResponse(IJobStepRef jobStepRef)
+		{
+			JobId = jobStepRef.Id.JobId.ToString();
+			BatchId = jobStepRef.Id.BatchId.ToString();
+			StepId = jobStepRef.Id.StepId.ToString();
+			Change = jobStepRef.Change;
+			LogId = jobStepRef.LogId.ToString();
+			PoolId = jobStepRef.PoolId?.ToString();
+			AgentId = jobStepRef.AgentId?.ToString();
+			Outcome = jobStepRef.Outcome;
+			StartTime = jobStepRef.StartTimeUtc;
+			FinishTime = jobStepRef.FinishTimeUtc;
+		}
+	}
+}
diff --git a/Engine/Source/Programs/Horde/Horde.Build/Jobs/JobsController.Messages.cs b/Engine/Source/Programs/Horde/Horde.Build/Jobs/JobsController.Messages.cs
deleted file mode 100644
diff --git a/Engine/Source/Programs/Horde/Horde.Build/Projects/ProjectMessages.cs b/Engine/Source/Programs/Horde/Horde.Build/Projects/ProjectMessages.cs
new file mode 100644
--- /dev/null
+++ b/Engine/Source/Programs/Horde/Horde.Build/Projects/ProjectMessages.cs
@@ -0,0 +1,180 @@
+// Copyright Epic Games, Inc. All Rights Reserved.
+
+using System.Collections.Generic;
+using System.ComponentModel.DataAnnotations;
+using Horde.Build.Acls;
+using Horde.Build.Projects;
+
+namespace Horde.Build.Projects
+{
+	/// <summary>
+	/// Information about a category to display for a stream
+	/// </summary>
+	public class CreateProjectCategoryRequest
+	{
+		/// <summary>
+		/// Name of this category
+		/// </summary>
+		[Required]
+		public string Name { get; set; }
+
+		/// <summary>
+		/// Index of the row to display this category on
+		/// </summary>
+		public int Row { get; set; }
+
+		/// <summary>
+		/// Whether to show this category on the nav menu
+		/// </summary>
+		public bool ShowOnNavMenu { get; set; }
+
+		/// <summary>
+		/// Patterns for stream names to include
+		/// </summary>
+		public List<string> IncludePatterns { get; set; } = new List<string>();
+
+		/// <summary>
+		/// Patterns for stream names to exclude
+		/// </summary>
+		public List<string> ExcludePatterns { get; set; } = new List<string>();
+
+		/// <summary>
+		/// Constructor
+		/// </summary>
+		public CreateProjectCategoryRequest(string name)
+		{
+			Name = name;
+		}
+	}
+
+	/// <summary>
+	/// Information about a stream within a project
+	/// </summary>
+	public class GetProjectStreamResponse
+	{
+		/// <summary>
+		/// The stream id
+		/// </summary>
+		public string Id { get; set; }
+
+		/// <summary>
+		/// The stream name
+		/// </summary>
+		public string Name { get; set; }
+
+		/// <summary>
+		/// Constructor
+		/// </summary>
+		/// <param name="id">The unique stream id</param>
+		/// <param name="name">The stream name</param>
+		public GetProjectStreamResponse(string id, string name)
+		{
+			Id = id;
+			Name = name;
+		}
+	}
+
+	/// <summary>
+	/// Information about a category to display for a stream
+	/// </summary>
+	public class GetProjectCategoryResponse
+	{
+		/// <summary>
+		/// Heading for this column
+		/// </summary>
+		public string Name { get; set; }
+
+		/// <summary>
+		/// Index of the row to display this category on
+		/// </summary>
+		public int Row { get; set; }
+
+		/// <summary>
+		/// Whether to show this category on the nav menu
+		/// </summary>
+		public bool ShowOnNavMenu { get; set; }
+
+		/// <summary>
+		/// Patterns for stream names to include
+		/// </summary>
+		public List<string> IncludePatterns { get; set; } = new List<string>();
+
+		/// <summary>
+		/// Patterns for stream names to exclude
+		/// </summary>
+		public List<string> ExcludePatterns { get; set; } = new List<string>();
+
+		/// <summary>
+		/// Streams to include in this category
+		/// </summary>
+		public List<string> Streams { get; set; } = new List<string>();
+
+		/// <summary>
+		/// Constructor
+		/// </summary>
+		/// <param name="streamCategory">The category to construct from</param>
+		public GetProjectCategoryResponse(StreamCategory streamCategory)
+		{
+			Name = streamCategory.Name;
+			Row = streamCategory.Row;
+			ShowOnNavMenu = streamCategory.ShowOnNavMenu;
+			IncludePatterns = streamCategory.IncludePatterns;
+			ExcludePatterns = streamCategory.ExcludePatterns;
+		}
+	}
+
+	/// <summary>
+	/// Response describing a project
+	/// </summary>
+	public class GetProjectResponse
+	{
+		/// <summary>
+		/// Unique id of the project
+		/// </summary>
+		public string Id { get; set; }
+
+		/// <summary>
+		/// Name of the project
+		/// </summary>
+		public string Name { get; set; }
+
+		/// <summary>
+		/// Order to display this project on the dashboard
+		/// </summary>
+		public int Order { get; set; }
+
+		/// <summary>
+		/// List of streams that are in this project
+		/// </summary>
+		public List<GetProjectStreamResponse>? Streams { get; set; }
+
+		/// <summary>
+		/// List of stream categories to display
+		/// </summary>
+		public List<GetProjectCategoryResponse>? Categories { get; set; }
+
+		/// <summary>
+		/// Custom permissions for this object
+		/// </summary>
+		public GetAclResponse? Acl { get; set; }
+
+		/// <summary>
+		/// Constructor
+		/// </summary>
+		/// <param name="id">Unique id of the project</param>
+		/// <param name="name">Name of the project</param>
+		/// <param name="order">Order to show this project on the dashboard</param>
+		/// <param name="streams">List of streams to display</param>
+		/// <param name="categories">List of stream categories to display</param>
+		/// <param name="acl">Custom permissions for this object</param>
+		public GetProjectResponse(string id, string name, int order, List<GetProjectStreamResponse>? streams, List<GetProjectCategoryResponse>? categories, GetAclResponse? acl)
+		{
+			Id = id;
+			Name = name;
+			Order = order;
+			Streams = streams;
+			Categories = categories;
+			Acl = acl;
+		}
+	}
+}
diff --git a/Engine/Source/Programs/Horde/Horde.Build/Projects/ProjectsController.Messages.cs b/Engine/Source/Programs/Horde/Horde.Build/Projects/ProjectsController.Messages.cs
deleted file mode 100644
diff --git a/Engine/Source/Programs/Horde/Horde.Build/Secrets/CredentialMessages.cs b/Engine/Source/Programs/Horde/Horde.Build/Secrets/CredentialMessages.cs
new file mode 100644
--- /dev/null
+++ b/Engine/Source/Programs/Horde/Horde.Build/Secrets/CredentialMessages.cs
@@ -0,0 +1,106 @@
+// Copyright Epic Games, Inc. All Rights Reserved.
+
+using System.Collections.Generic;
+using System.ComponentModel.DataAnnotations;
+using Horde.Build.Acls;
+using Horde.Build.Secrets;
+
+namespace Horde.Build.Credentials
+{
+	/// <summary>
+	/// Parameters to create a new credential
+	/// </summary>
+	public class CreateCredentialRequest
+	{
+		/// <summary>
+		/// Name for the new credential
+		/// </summary>
+		[Required]
+		public string Name { get; set; } = null!;
+
+		/// <summary>
+		/// Properties for the new credential
+		/// </summary>
+		public Dictionary<string, string>? Properties { get; set; }
+	}
+
+	/// <summary>
+	/// Response from creating a new credential
+	/// </summary>
+	public class CreateCredentialResponse
+	{
+		/// <summary>
+		/// Unique id for the new credential
+		/// </summary>
+		public string Id { get; set; }
+
+		/// <summary>
+		/// Constructor
+		/// </summary>
+		/// <param name="id">Unique id for the new credential</param>
+		public CreateCredentialResponse(string id)
+		{
+			Id = id;
+		}
+	}
+
+	/// <summary>
+	/// Parameters to update a credential
+	/// </summary>
+	public class UpdateCredentialRequest
+	{
+		/// <summary>
+		/// Optional new name for the credential
+		/// </summary>
+		public string? Name { get; set; }
+
+		/// <summary>
+		/// Properties to update for the credential. Properties set to null will be removed.
+		/// </summary>
+		public Dictionary<string, string?>? Properties { get; set; }
+
+		/// <summary>
+		/// Custom permissions for this object
+		/// </summary>
+		public UpdateAclRequest? Acl { get; set; }
+	}
+
+	/// <summary>
+	/// Response describing a credential
+	/// </summary>
+	public class GetCredentialResponse
+	{
+		/// <summary>
+		/// Unique id of the credential
+		/// </summary>
+		public string Id { get; set; }
+
+		/// <summary>
+		/// Name of the credential
+		/// </summary>
+		public string Name { get; set; }
+
+		/// <summary>
+		/// Properties for the credential
+		/// </summary>
+		public Dictionary<string, string> Properties { get; set; }
+
+		/// <summary>
+		/// Custom permissions for this object
+		/// </summary>
+		public GetAclResponse? Acl { get; set; }
+
+		/// <summary>
+		/// Constructor
+		/// </summary>
+		/// <param name="credential">The credential to construct from</param>
+		/// <param name="bIncludeAcl">Whether to include the ACL in the response</param>
+		public GetCredentialResponse(Credential credential, bool bIncludeAcl)
+		{
+			Id = credential.Id.ToString();
+			Name = credential.Name;
+			Properties = credential.Properties;
+			Acl = (bIncludeAcl && credential.Acl != null)? new GetAclResponse(credential.Acl) : null;
+		}
+	}
+}
diff --git a/Engine/Source/Programs/Horde/Horde.Build/Secrets/CredentialsController.Messages.cs b/Engine/Source/Programs/Horde/Horde.Build/Secrets/CredentialsController.Messages.cs
deleted file mode 100644
diff --git a/Engine/Source/Programs/Horde/Horde.Build/Server/Notices/NoticeMessages.cs b/Engine/Source/Programs/Horde/Horde.Build/Server/Notices/NoticeMessages.cs
new file mode 100644
--- /dev/null
+++ b/Engine/Source/Programs/Horde/Horde.Build/Server/Notices/NoticeMessages.cs
@@ -0,0 +1,122 @@
+// Copyright Epic Games, Inc. All Rights Reserved.
+
+using System;
+using System.ComponentModel.DataAnnotations;
+using Horde.Build.Users;
+
+namespace Horde.Build.Server.Notices
+{
+	/// <summary>
+	/// Parameters required to create a notice
+	/// </summary>
+	public class CreateNoticeRequest
+	{
+		/// <summary>
+		/// Start time to display this message
+		/// </summary>
+		public DateTime? StartTime { get; set; }
+
+		/// <summary>
+		/// Finish time to display this message
+		/// </summary>
+		public DateTime? FinishTime { get; set; }
+
+		/// <summary>
+		/// Message to display
+		/// </summary>
+		[Required]
+		public string Message { get; set; }
+
+		/// <summary>
+		/// Constructor
+		/// </summary>
+		public CreateNoticeRequest(string message)
+		{
+			Message = message;
+		}
+	}
+
+	/// <summary>
+	/// Parameters required to update a notice
+	/// </summary>
+	public class UpdateNoticeRequest
+	{
+		/// <summary>
+		/// The id of the notice to update
+		/// </summary>
+		public string Id { get; set; }
+
+		/// <summary>
+		/// Start time to display this message
+		/// </summary>
+		public DateTime? StartTime { get; set; }
+
+		/// <summary>
+		/// Finish time to display this message
+		/// </summary>
+		public DateTime? FinishTime { get; set; }
+
+		/// <summary>
+		/// Message to display
+		/// </summary>
+		public string? Message { get; set; }
+
+		/// <summary>
+		/// Constructor
+		/// </summary>
+		public UpdateNoticeRequest(string id)
+		{
+			Id = id;
+		}
+	}
+
+	/// <summary>
+	/// Notice information
+	/// </summary>
+	public class GetNoticeResponse
+	{
+		/// <summary>
+		/// The id of the notice to update
+		/// </summary>
+		public string? Id { get; set; }
+
+		/// <summary>
+		/// Start time to display this message
+		/// </summary>
+		public DateTime? StartTime { get; set; }
+
+		/// <summary>
+		/// Finish time to display this message
+		/// </summary>
+		public DateTime? FinishTime { get; set; }
+
+		/// <summary>
+		/// Whether this notice is for scheduled downtime
+		/// </summary>
+		public bool ScheduledDowntime { get; set; } = false;
+
+		/// <summary>
+		/// Whether the notice is currently active
+		/// </summary>
+		public bool Active { get; set; }
+
+		/// <summary>
+		/// Message to display
+		/// </summary>
+		public string? Message { get; set; }
+
+		/// <summary>
+		/// User id who created the notice, otherwise null if a system message
+		/// </summary>
+		public GetThinUserInfoResponse? CreatedByUser { get; set; }
+
+		/// <summary>
+		/// Constructor
+		/// </summary>
+		public GetNoticeResponse()
+		{
+
+		}
+
+	}
+}
diff --git a/Engine/Source/Programs/Horde/Horde.Build/Server/Notices/NoticesController.Messages.cs b/Engine/Source/Programs/Horde/Horde.Build/Server/Notices/NoticesController.Messages.cs
deleted file mode 100644
diff --git a/Engine/Source/Programs/Horde/Horde.Build/Streams/StreamMessages.cs b/Engine/Source/Programs/Horde/Horde.Build/Streams/StreamMessages.cs
new file mode 100644
--- /dev/null
+++ b/Engine/Source/Programs/Horde/Horde.Build/Streams/StreamMessages.cs
@@ -0,0 +1,950 @@
+// Copyright Epic Games, Inc. All Rights Reserved.
+
+using System;
+using System.Collections.Generic;
+using System.ComponentModel.DataAnnotations;
+using EpicGames.Core;
+using Horde.Build.Acls;
+using Horde.Build.Jobs.Graphs;
+using Horde.Build.Jobs.Templates;
+using Horde.Build.Server;
+using Horde.Build.Jobs.Schedules;
+using Horde.Build.Users;
+using Horde.Build.Utilities;
+using HordeCommon;
+
+namespace Horde.Build.Streams
+{
+	using TemplateRefId = StringId<TemplateRef>;
+	using WorkflowId = StringId<WorkflowConfig>;
+
+	/// <summary>
+	/// Query selecting the base changelist to use
+	/// </summary>
+	public class ChangeQueryRequest
+	{
+		/// <summary>
+		/// The template id to query
+		/// </summary>
+		public string? TemplateId { get; set; }
+
+		/// <summary>
+		/// The target to query
+		/// </summary>
+		public string? Target { get; set; }
+
+		/// <summary>
+		/// Whether to match a job that produced warnings
+		/// </summary>
+		public List<JobStepOutcome>? Outcomes { get; set; }
+
+		/// <summary>
+		/// Convert to a model object
+		/// </summary>
+		/// <returns></returns>
+		public ChangeQuery ToModel()
+		{
+			ChangeQuery query = new ChangeQuery();
+			if (TemplateId != null)
+			{
+				query.TemplateRefId = new StringId<TemplateRef>(TemplateId);
+			}
+			query.Target = Target;
+			query.Outcomes = Outcomes;
+			return query;
+		}
+	}
+
+	/// <summary>
+	/// Specifies defaults for running a preflight
+	/// </summary>
+	public class DefaultPreflightRequest
+	{
+		/// <summary>
+		/// The template id to query
+		/// </summary>
+		public string? TemplateId { get; set; }
+
+		/// <summary>
+		/// The last successful job type to use for the base changelist
+		/// </summary>
+		[Obsolete("Use Change.TemplateId instead")]
+		public string? ChangeTemplateId { get; set; }
+
+		/// <summary>
+		/// Query for the change to use
+		/// </summary>
+		public ChangeQueryRequest? Change { get; set; }
+
+		/// <summary>
+		/// Convert to a model object
+		/// </summary>
+		/// <returns></returns>
+		public DefaultPreflight ToModel()
+		{
+#pragma warning disable CS0618 // Type or member is obsolete
+			if (ChangeTemplateId != null)
+			{
+				Change ??= new ChangeQueryRequest();
+				Change.TemplateId = ChangeTemplateId;
+			}
+			return new DefaultPreflight((TemplateId != null) ? (TemplateRefId?)new TemplateRefId(TemplateId) : null, Change?.ToModel());
+#pragma warning restore CS0618 // Type or member is obsolete
+		}
+	}
+
+	/// <summary>
+	/// Mapping from a BuildGraph agent type to a set of machines on the farm
+	/// </summary>
+	public class CreateAgentTypeRequest
+	{
+		/// <summary>
+		/// Pool of agents to use for this agent type
+		/// </summary>
+		[Required]
+		public string Pool { get; set; } = null!;
+
+		/// <summary>
+		/// Name of the workspace to sync
+		/// </summary>
+		public string? Workspace { get; set; }
+
+		/// <summary>
+		/// Path to the temporary storage dir
+		/// </summary>
+		public string? TempStorageDir { get; set; }
+
+		/// <summary>
+		/// Environment variables to be set when executing the job
+		/// </summary>
+		public Dictionary<string, string>? Environment { get; set; }
+	}
+
+	/// <summary>
+	/// Information about a workspace type
+	/// </summary>
+	public class CreateWorkspaceTypeRequest
+	{
+		/// <summary>
+		/// Name of the Perforce server cluster to use
+		/// </summary>
+		public string? Cluster { get; set; }
+
+		/// <summary>
+		/// The Perforce server and port (eg. perforce:1666)
+		/// </summary>
+		public string? ServerAndPort { get; set; }
+
+		/// <summary>
+		/// User to log into Perforce with (defaults to buildmachine)
+		/// </summary>
+		public string? UserName { get; set; }
+
+		/// <summary>
+		/// Password to use to log into the workspace
+		/// </summary>
+		public string? Password { get; set; }
+
+		/// <summary>
+		/// Identifier to distinguish this workspace from other workspaces. Defaults to the workspace type name.
+		/// </summary>
+		public string? Identifier { get; set; }
+
+		/// <summary>
+		/// Override for the stream to sync
+		/// </summary>
+		public string? Stream { get; set; }
+
+		/// <summary>
+		/// Custom view for the workspace
+		/// </summary>
+		public List<string>? View { get; set; }
+
+		/// <summary>
+		/// Whether to use an incrementally synced workspace
+		/// </summary>
+		public bool Incremental { get; set; }
+
+		/// <summary>
+		/// Whether to use the AutoSDK
+		/// </summary>
+		public bool UseAutoSdk { get; set; } = true;
+	}
+
+	/// <summary>
+	/// Trigger for another template
+	/// </summary>
+	public class CreateChainedJobTemplateRequest
+	{
+		/// <summary>
+		/// Name of the target that needs to complete before starting the other template
+		/// </summary>
+		[Required]
+		public string Trigger { get; set; } = String.Empty;
+
+		/// <summary>
+		/// Id of the template to trigger
+		/// </summary>
+		[Required]
+		public string TemplateId { get; set; } = String.Empty;
+	}
+
+	/// <summary>
+	/// Parameters to create a template within a stream
+	/// </summary>
+	public class TemplateRefConfig : TemplateConfig
+	{
+		TemplateRefId _id;
+
+		/// <summary>
+		/// Optional identifier for this ref. If not specified, an id will be generated from the name.
+		/// </summary>
+		public TemplateRefId Id
+		{
+			get => _id.IsEmpty ? TemplateRefId.Sanitize(Name) : _id;
+			set => _id = value;
+		}
+
+		/// <summary>
+		/// Whether to show badges in UGS for these jobs
+		/// </summary>
+		public bool ShowUgsBadges { get; set; }
+
+		/// <summary>
+		/// Whether to show alerts in UGS for these jobs
+		/// </summary>
+		public bool ShowUgsAlerts { get; set; }
+
+		/// <summary>
+		/// Notification channel for this template. Overrides the stream channel if set.
+		/// </summary>
+		public string? NotificationChannel { get; set; }
+
+		/// <summary>
+		/// Notification channel filter for this template. Can be Success|Failure|Warnings
+		/// </summary>
+		public string? NotificationChannelFilter { get; set; }
+
+		/// <summary>
+		/// Triage channel for this template. Overrides the stream channel if set.
+		/// </summary>
+		public string? TriageChannel { get; set; }
+
+		/// <summary>
+		/// Workflow to user for this stream
+		/// </summary>
+		public WorkflowId? WorkflowId
+		{
+			get => Annotations.WorkflowId;
+			set => Annotations.WorkflowId = value;
+		}
+
+		/// <summary>
+		/// Default annotations to apply to nodes in this template
+		/// </summary>
+		public NodeAnnotations Annotations { get; set; } = new NodeAnnotations();
+
+		/// <summary>
+		/// Schedule to execute this template
+		/// </summary>
+		public CreateScheduleRequest? Schedule { get; set; }
+
+		/// <summary>
+		/// List of chained job triggers
+		/// </summary>
+		public List<CreateChainedJobTemplateRequest>? ChainedJobs { get; set; }
+
+		/// <summary>
+		/// The ACL for this template
+		/// </summary>
+		public UpdateAclRequest? Acl { get; set; }
+	}
+
+	/// <summary>
+	/// Step state update request
+	/// </summary>
+	public class UpdateStepStateRequest
+	{
+		/// <summary>
+		/// Name of the step
+		/// </summary>
+		[Required]
+		public string Name { get; set; } = String.Empty;
+
+		/// <summary>
+		/// User who paused the step
+		/// </summary>
+		public string? PausedByUserId { get; set; }
+	}
+
+
+	/// <summary>
+	/// Updates an existing stream template ref
+	/// </summary>
+	public class UpdateTemplateRefRequest
+	{
+		
+		/// <summary>
+		/// Step states to update
+		/// </summary>
+		public List<UpdateStepStateRequest>? StepStates { get; set; }
+
+	}
+
+
+	/// <summary>
+	/// Mapping from a BuildGraph agent type to a set of machines on the farm
+	/// </summary>
+	public class GetAgentTypeResponse
+	{
+		/// <summary>
+		/// Pool of agents to use for this agent type
+		/// </summary>
+		public string Pool { get; set; }
+
+		/// <summary>
+		/// Name of the workspace to sync
+		/// </summary>
+		public string? Workspace { get; set; }
+
+		/// <summary>
+		/// Path to the temporary storage dir
+		/// </summary>
+		public string? TempStorageDir { get; set; }
+
+		/// <summary>
+		/// Environment variables to be set when executing the job
+		/// </summary>
+		public Dictionary<string, string>? Environment { get; set; }
+
+		/// <summary>
+		/// Constructor
+		/// </summary>
+		/// <param name="pool">Pool of agents to use for this agent type</param>
+		/// <param name="workspace">Name of the workspace to sync</param>
+		/// <param name="tempStorageDir">Path to the temp storage directory</param>
+		/// <param name="environment">Environment variables to be set when executing this job</param>
+		public GetAgentTypeResponse(string pool, string? workspace, string? tempStorageDir, Dictionary<string, string>? environment)
+		{
+			Pool = pool;
+			Workspace = workspace;
+			TempStorageDir = tempStorageDir;
+			Environment = environment;
+		}
+	}
+
+	/// <summary>
+	/// Information about a workspace type
+	/// </summary>
+	public class GetWorkspaceTypeResponse
+	{
+		/// <summary>
+		/// The Perforce server cluster
+		/// </summary>
+		public string? Cluster { get; set; }
+
+		/// <summary>
+		/// The Perforce server and port (eg. perforce:1666)
+		/// </summary>
+		public string? ServerAndPort { get; set; }
+
+		/// <summary>
+		/// User to log into Perforce with (defaults to buildmachine)
+		/// </summary>
+		public string? UserName { get; set; }
+
+		/// <summary>
+		/// Identifier to distinguish this workspace from other workspaces. Defaults to the workspace type name.
+		/// </summary>
+		public string? Identifier { get; set; }
+
+		/// <summary>
+		/// Override for the stream to sync
+		/// </summary>
+		public string? Stream { get; set; }
+
+		/// <summary>
+		/// Custom view for the workspace
+		/// </summary>
+		public List<string>? View { get; set; }
+
+		/// <summary>
+		/// Whether to use an incrementally synced workspace
+		/// </summary>
+		public bool Incremental { get; set; }
+
+		/// <summary>
+		/// Whether to use the AutoSDK
+		/// </summary>
+		public bool UseAutoSdk { get; set; }
+
+		/// <summary>
+		/// Default constructor
+		/// </summary>
+		public GetWorkspaceTypeResponse()
+		{
+		}
+
+		/// <summary>
+		/// Constructor
+		/// </summary>
+		/// <param name="cluster">The server cluster</param>
+		/// <param name="serverAndPort">The perforce server</param>
+		/// <param name="userName">The perforce user name</param>
+		/// <param name="identifier">Identifier to distinguish this workspace from other workspaces. Defaults to the workspace type name.</param>
+		/// <param name="stream">Override for the stream to sync</param>
+		/// <param name="view">Custom view for the workspace</param>
+		/// <param name="bIncremental">Whether to use an incrementally synced workspace</param>
+		/// <param name="bUseAutoSdk">Whether to use the AutoSDK</param>
+		public GetWorkspaceTypeResponse(string? cluster, string? serverAndPort, string? userName, string? identifier, string? stream, List<string>? view, bool bIncremental, bool bUseAutoSdk)
+		{
+			Cluster = cluster;
+			ServerAndPort = serverAndPort;
+			UserName = userName;
+			Identifier = identifier;
+			Stream = stream;
+			View = view;
+			Incremental = bIncremental;
+			UseAutoSdk = bUseAutoSdk;
+		}
+	}
+
+	/// <summary>
+	/// Trigger for another template
+	/// </summary>
+	public class GetChainedJobTemplateResponse
+	{
+		/// <summary>
+		/// Name of the target that needs to complete before starting the other template
+		/// </summary>
+		public string Trigger { get; set; }
+
+		/// <summary>
+		/// Id of the template to trigger
+		/// </summary>
+		public string TemplateId { get; set; }
+
+		/// <summary>
+		/// Constructor
+		/// </summary>
+		/// <param name="trigger">The trigger definition</param>
+		public GetChainedJobTemplateResponse(ChainedJobTemplate trigger)
+		{
+			Trigger = trigger.Trigger;
+			TemplateId = trigger.TemplateRefId.ToString();
+		}
+	}
+
+	/// <summary>
+	/// State information for a step in the stream
+	/// </summary>
+	public class GetTemplateStepStateResponse
+	{
+		/// <summary>
+		/// Name of the step
+		/// </summary>
+		public string Name { get; set; } = String.Empty;
+
+		/// <summary>
+		/// User who paused the step
+		/// </summary>
+		public GetThinUserInfoResponse? PausedByUserInfo { get; set; }
+
+		/// <summary>
+		/// The UTC time when the step was paused
+		/// </summary>
+		public DateTime? PauseTimeUtc { get; set; }
+
+		/// <summary>
+		/// Default constructor for serialization
+		/// </summary>
+		private GetTemplateStepStateResponse()
+		{
+		}
+
+		/// <summary>
+		/// Constructor
+		/// </summary>
+		public GetTemplateStepStateResponse(TemplateStepState state, GetThinUserInfoResponse? pausedByUserInfo)
+		{
+			Name = state.Name;			
+			PauseTimeUtc = state.PauseTimeUtc;
+			PausedByUserInfo = pausedByUserInfo;
+		}
+
+	}
+
+
+	/// <summary>
+	/// Information about a template in this stream
+	/// </summary>
+	public class GetTemplateRefResponse : GetTemplateResponseBase
+	{
+		/// <summary>
+		/// Id of the template ref
+		/// </summary>
+		public string Id { get; set; }
+
+		/// <summary>
+		/// Hash of the template definition
+		/// </summary>
+		public string Hash { get; set; }
+
+		/// <summary>
+		/// Whether to show badges in UGS for these jobs
+		/// </summary>
+		public bool ShowUgsBadges { get; set; }
+
+		/// <summary>
+		/// Whether to show alerts in UGS for these jobs
+		/// </summary>
+		public bool ShowUgsAlerts { get; set; }
+
+		/// <summary>
+		/// Notification channel for this template. Overrides the stream channel if set.
+		/// </summary>
+		public string? NotificationChannel { get; set; }
+
+		/// <summary>
+		/// Notification channel filter for this template. Can be Success|Failure|Warnings
+		/// </summary>
+		public string? NotificationChannelFilter { get; set; }
+
+		/// <summary>
+		/// Triage channel for this template. Overrides the stream channel if set.
+		/// </summary>
+		public string? TriageChannel { get; set; }
+
+		/// <summary>
+		/// The schedule for this ref
+		/// </summary>
+		public GetScheduleResponse? Schedule { get; set; }
+
+		/// <summary>
+		/// List of templates to trigger
+		/// </summary>
+		public List<GetChainedJobTemplateResponse>? ChainedJobs { get; set; }
+
+		/// <summary>
+		/// List of step states
+		/// </summary>
+		public List<GetTemplateStepStateResponse>? StepStates { get; set; }
+
+		/// <summary>
+		/// ACL for this template
+		/// </summary>
+		public GetAclResponse? Acl { get; set; }
+
+		/// <summary>
+		/// Constructor
+		/// </summary>
+		/// <param name="id">The template ref id</param>
+		/// <param name="templateRef">The template ref</param>
+		/// <param name="template">The template instance</param>
+		/// <param name="stepStates">The template step states</param>
+		/// <param name="bIncludeAcl">Whether to include the ACL in the response</param>
+		public GetTemplateRefResponse(TemplateRefId id, TemplateRef templateRef, ITemplate template, List<GetTemplateStepStateResponse>? stepStates, bool bIncludeAcl)
+			: base(template)
+		{
+			Id = id.ToString();
+			Hash = templateRef.Hash.ToString();
+			ShowUgsBadges = templateRef.ShowUgsBadges;
+			ShowUgsAlerts = templateRef.ShowUgsAlerts;
+			NotificationChannel = templateRef.NotificationChannel;
+			NotificationChannelFilter = templateRef.NotificationChannelFilter;
+			Schedule = (templateRef.Schedule != null) ? new GetScheduleResponse(templateRef.Schedule) : null;
+			ChainedJobs = (templateRef.ChainedJobs != null && templateRef.ChainedJobs.Count > 0) ? templateRef.ChainedJobs.ConvertAll(x => new GetChainedJobTemplateResponse(x)) : null;
+			StepStates = stepStates;
+			Acl = (bIncludeAcl && templateRef.Acl != null)? new GetAclResponse(templateRef.Acl) : null;
+		}
+	}
+
+	/// <summary>
+	/// Response describing a stream
+	/// </summary>
+	public class GetStreamResponse
+	{
+		/// <summary>
+		/// Unique id of the stream
+		/// </summary>
+		public string Id { get; set; }
+
+		/// <summary>
+		/// Unique id of the project containing this stream
+		/// </summary>
+		public string ProjectId { get; set; }
+
+		/// <summary>
+		/// Name of the stream
+		/// </summary>
+		public string Name { get; set; }
+
+		/// <summary>
+		/// The config file path on the server
+		/// </summary>
+		public string ConfigPath { get; set; } = String.Empty;
+
+		/// <summary>
+		/// Revision of the config file 
+		/// </summary>
+		public string ConfigRevision { get; set; } = String.Empty;
+
+		/// <summary>
+		/// Order to display in the list
+		/// </summary>
+		public int Order { get; set; }
+
+		/// <summary>
+		/// Notification channel for all jobs in this stream
+		/// </summary>
+		public string? NotificationChannel { get; set; }
+
+		/// <summary>
+		/// Notification channel filter for this template. Can be Success|Failure|Warnings
+		/// </summary>
+		public string? NotificationChannelFilter { get; set; }
+
+		/// <summary>
+		/// Channel to post issue triage notifications
+		/// </summary>
+		public string? TriageChannel { get; set; }
+
+		/// <summary>
+		/// Default template for running preflights
+		/// </summary>
+		public string? DefaultPreflightTemplate { get; set; }
+
+		/// <summary>
+		/// Default template to use for preflights
+		/// </summary>
+		public DefaultPreflightRequest? DefaultPreflight { get; set; }
+
+		/// <summary>
+		/// List of tabs to display for this stream
+		/// </summary>
+		public List<GetStreamTabResponse> Tabs { get; set; }
+
+		/// <summary>
+		/// Map of agent name to type
+		/// </summary>
+		public Dictionary<string, GetAgentTypeResponse> AgentTypes { get; set; }
+
+		/// <summary>
+		/// Map of workspace name to type
+		/// </summary>
+		public Dictionary<string, GetWorkspaceTypeResponse>? WorkspaceTypes { get; set; }
+
+		/// <summary>
+		/// Templates for jobs in this stream
+		/// </summary>
+		public List<GetTemplateRefResponse> Templates { get; set; }
+
+		/// <summary>
+		/// Custom permissions for this object
+		/// </summary>
+		public GetAclResponse? Acl { get; set; }
+		
+		/// <summary>
+		/// Stream paused for new builds until this date
+		/// </summary>
+		public DateTime? PausedUntil { get; set; }
+		
+		/// <summary>
+		/// Reason for stream being paused
+		/// </summary>
+		public string? PauseComment { get; set; }
+
+		/// <summary>
+		/// Workflows for this stream
+		/// </summary>
+		public List<WorkflowConfig> Workflows { get; set; }
+
+		/// <summary>
+		/// Constructor
+		/// </summary>
+		/// <param name="id">Unique id of the stream</param>
+		/// <param name="projectId">Unique id of the project containing the stream</param>
+		/// <param name="name">Name of the stream</param>
+		/// <param name="configRevision">The config file path on the server</param>
+		/// <param name="order">Order to display this stream</param>
+		/// <param name="notificationChannel"></param>
+		/// <param name="notificationChannelFilter"></param>
+		/// <param name="triageChannel"></param>
+		/// <param name="defaultPreflight">The default template to use for preflights</param>
+		/// <param name="tabs">List of tabs to display for this stream</param>
+		/// <param name="agentTypes">Map of agent type name to description</param>
+		/// <param name="workspaceTypes">Map of workspace name to description</param>
+		/// <param name="templates">Templates for this stream</param>
+		/// <param name="acl">Permissions for this object</param>
+		/// <param name="pausedUntil">Stream paused for new builds until this date</param>
+		/// <param name="pauseComment">Reason for stream being paused</param>
+		/// <param name="workflows">Workflows for this stream</param>
+		public GetStreamResponse(string id, string projectId, string name, string configRevision, int order, string? notificationChannel, string? notificationChannelFilter, string? triageChannel, DefaultPreflightRequest? defaultPreflight, List<GetStreamTabResponse> tabs, Dictionary<string, GetAgentTypeResponse> agentTypes, Dictionary<string, GetWorkspaceTypeResponse>? workspaceTypes, List<GetTemplateRefResponse> templates, GetAclResponse? acl, DateTime? pausedUntil, string? pauseComment, List<WorkflowConfig> workflows)
+		{
+			Id = id;
+			ProjectId = projectId;
+			Name = name;
+			ConfigRevision = configRevision;
+			Order = order;
+			NotificationChannel = notificationChannel;
+			NotificationChannelFilter = notificationChannelFilter;
+			TriageChannel = triageChannel;
+			DefaultPreflightTemplate = defaultPreflight?.TemplateId;
+			DefaultPreflight = defaultPreflight;
+			Tabs = tabs;
+			AgentTypes = agentTypes;
+			WorkspaceTypes = workspaceTypes;
+			Templates = templates;
+			Acl = acl;
+			PausedUntil = pausedUntil;
+			PauseComment = pauseComment;
+			Workflows = workflows;
+		}
+	}
+
+	/// <summary>
+	/// Information about a page to display in the dashboard for a stream
+	/// </summary>
+	[JsonKnownTypes(typeof(CreateJobsTabRequest))]
+	public abstract class CreateStreamTabRequest
+	{
+		/// <summary>
+		/// Title of this page
+		/// </summary>
+		[Required]
+		public string Title { get; set; } = null!;
+	}
+
+	/// <summary>
+	/// Type of a column in a jobs tab
+	/// </summary>
+	public enum JobsTabColumnType
+	{
+		/// <summary>
+		/// Contains labels
+		/// </summary>
+		Labels,
+
+		/// <summary>
+		/// Contains parameters
+		/// </summary>
+		Parameter
+	}
+
+	/// <summary>
+	/// Describes a column to display on the jobs page
+	/// </summary>
+	public class CreateJobsTabColumnRequest
+	{
+		/// <summary>
+		/// The type of column
+		/// </summary>
+		public JobsTabColumnType Type { get; set; } = JobsTabColumnType.Labels;
+
+		/// <summary>
+		/// Heading for this column
+		/// </summary>
+		[Required]
+		public string Heading { get; set; } = null!;
+
+		/// <summary>
+		/// Category of aggregates to display in this column. If null, includes any aggregate not matched by another column.
+		/// </summary>
+		public string? Category { get; set; }
+
+		/// <summary>
+		/// Parameter to show in this column
+		/// </summary>
+		public string? Parameter { get; set; }
+
+		/// <summary>
+		/// Relative width of this column.
+		/// </summary>
+		public int? RelativeWidth { get; set; }
+
+		/// <summary>
+		/// Construct a JobsTabColumn object from this request
+		/// </summary>
+		/// <returns>Column object</returns>
+		public JobsTabColumn ToModel()
+		{
+			switch (Type)
+			{
+				case JobsTabColumnType.Labels:
+					return new JobsTabLabelColumn(Heading, Category, RelativeWidth);
+				case JobsTabColumnType.Parameter:
+					return new JobsTabParameterColumn(Heading, Parameter ?? "Undefined", RelativeWidth);
+				default:
+					return new JobsTabLabelColumn(Heading, "Undefined", RelativeWidth);
+			}
+		}
+	}
+
+	/// <summary>
+	/// Describes a job page
+	/// </summary>
+	[JsonDiscriminator("Jobs")]
+	public class CreateJobsTabRequest : CreateStreamTabRequest
+	{
+		/// <summary>
+		/// Whether to show job names on this page
+		/// </summary>
+		public bool ShowNames { get; set; }
+
+		/// <summary>
+		/// Names of jobs to include on this page. If there is only one name specified, the name column does not need to be displayed.
+		/// </summary>
+		public List<string>? JobNames { get; set; }
+
+		/// <summary>
+		/// List of job template names to show on this page.
+		/// </summary>
+		public List<string>? Templates { get; set; }
+
+		/// <summary>
+		/// Columns to display for different types of aggregates
+		/// </summary>
+		public List<CreateJobsTabColumnRequest>? Columns { get; set; }
+	}
+
+	/// <summary>
+	/// Information about a page to display in the dashboard for a stream
+	/// </summary>
+	[JsonKnownTypes(typeof(GetJobsTabResponse))]
+	public abstract class GetStreamTabResponse
+	{
+		/// <summary>
+		/// Title of this page
+		/// </summary>
+		public string Title { get; set; }
+
+		/// <summary>
+		/// Constructor
+		/// </summary>
+		/// <param name="title">Title of this page</param>
+		protected GetStreamTabResponse(string title)//, StreamPageType Type)
+		{
+			Title = title;
+		}
+	}
+
+	/// <summary>
+	/// Describes a column to display on the jobs page
+	/// </summary>
+	[JsonKnownTypes(typeof(GetJobsTabLabelColumnResponse), typeof(GetJobsTabParameterColumnResponse))]
+	public abstract class GetJobsTabColumnResponse
+	{
+		/// <summary>
+		/// Heading for this column
+		/// </summary>
+		public string Heading { get; set; } = null!;
+
+		/// <summary>
+		/// Relative width of this column.
+		/// </summary>
+		public int? RelativeWidth { get; set; }
+
+		/// <summary>
+		/// Constructor
+		/// </summary>
+		/// <param name="heading">Heading for this column</param>
+		/// <param name="relativeWidth">Relative width of this column</param>
+		protected GetJobsTabColumnResponse(string heading, int? relativeWidth)
+		{
+			Heading = heading;
+			RelativeWidth = relativeWidth;
+		}
+	}
+
+	/// <summary>
+	/// Describes a label column to display on the jobs page
+	/// </summary>
+	[JsonDiscriminator(nameof(JobsTabColumnType.Labels))]
+	public class GetJobsTabLabelColumnResponse : GetJobsTabColumnResponse
+	{
+		/// <summary>
+		/// Category of aggregates to display in this column. If null, includes any aggregate not matched by another column.
+		/// </summary>
+		public string? Category { get; set; }
+
+		/// <summary>
+		/// Constructor
+		/// </summary>
+		/// <param name="heading">Heading for this column</param>
+		/// <param name="category">Category of aggregates to display in this column. If null, includes any aggregate not matched by another column.</param>
+		/// <param name="relativeWidth">Relative width of this column</param>
+		public GetJobsTabLabelColumnResponse(string heading, string? category, int? relativeWidth)
+			: base(heading, relativeWidth)
+		{
+			Category = category;
+		}
+	}
+
+	/// <summary>
+	/// Describes a parameter column to display on the jobs page
+	/// </summary>
+	[JsonDiscriminator(nameof(JobsTabColumnType.Parameter))]
+	public class GetJobsTabParameterColumnResponse : GetJobsTabColumnResponse
+	{
+		/// <summary>
+		/// Parameter to show in this column
+		/// </summary>
+		public string? Parameter { get; set; }
+
+		/// <summary>
+		/// Constructor
+		/// </summary>
+		/// <param name="heading">Heading for this column</param>
+		/// <param name="parameter">Name of the parameter to display</param>
+		/// <param name="relativeWidth">Relative width of this column</param>
+		public GetJobsTabParameterColumnResponse(string heading, string? parameter, int? relativeWidth)
+			: base(heading, relativeWidth)
+		{
+			Parameter = parameter;
+		}
+	}
+
+	/// <summary>
+	/// Describes a job page
+	/// </summary>
+	[JsonDiscriminator("Jobs")]
+	public class GetJobsTabResponse : GetStreamTabResponse
+	{
+		/// <summary>
+		/// Whether to show names on the page
+		/// </summary>
+		public bool ShowNames { get; set; }
+
+		/// <summary>
+		/// List of templates to show on the page
+		/// </summary>
+		public List<string>? Templates { get; set; }
+
+		/// <summary>
+		/// Names of jobs to include on this page. If there is only one name specified, the name column does not need to be displayed.
+		/// </summary>
+		public List<string>? JobNames { get; set; }
+
+		/// <summary>
+		/// Columns to display for different types of aggregates
+		/// </summary>
+		public List<GetJobsTabColumnResponse>? Columns { get; set; }
+
+		/// <summary>
+		/// Constructor
+		/// </summary>
+		/// <param name="title">Title for this page</param>
+		/// <param name="bShowNames">Whether to show names on the page</param>
+		/// <param name="templates">Templates to include on this page</param>
+		/// <param name="jobNames">List of job names to include on this page</param>
+		/// <param name="columns">List of columns to display</param>
+		public GetJobsTabResponse(string title, bool bShowNames, List<string>? templates, List<string>? jobNames, List<GetJobsTabColumnResponse>? columns)
+			: base(title)
+		{
+			ShowNames = bShowNames;
+			Templates = templates;
+			JobNames = jobNames;
+			Columns = columns;
+		}
+	}
+}
diff --git a/Engine/Source/Programs/Horde/Horde.Build/Streams/StreamsController.Messages.cs b/Engine/Source/Programs/Horde/Horde.Build/Streams/StreamsController.Messages.cs
deleted file mode 100644
diff --git a/Engine/Source/Programs/Horde/Horde.Build/Ugs/UgsController.Messages.cs b/Engine/Source/Programs/Horde/Horde.Build/Ugs/UgsController.Messages.cs
deleted file mode 100644
diff --git a/Engine/Source/Programs/Horde/Horde.Build/Ugs/UgsMessages.cs b/Engine/Source/Programs/Horde/Horde.Build/Ugs/UgsMessages.cs
new file mode 100644
--- /dev/null
+++ b/Engine/Source/Programs/Horde/Horde.Build/Ugs/UgsMessages.cs
@@ -0,0 +1,562 @@
+// Copyright Epic Games, Inc. All Rights Reserved.
+
+using System;
+using System.Collections.Generic;
+using System.ComponentModel.DataAnnotations;
+using System.Linq;
+using EpicGames.Core;
+using Horde.Build.Issues;
+using Horde.Build.Users;
+
+namespace Horde.Build.Ugs
+{
+	/// <summary>
+	/// Review by a user of a particular change
+	/// </summary>
+	public enum UgsUserVote
+	{
+		/// <summary>
+		/// No vote for the current change
+		/// </summary>
+		None,
+
+		/// <summary>
+		/// Succesfully compiled the change
+		/// </summary>
+		CompileSuccess,
+
+		/// <summary>
+		/// Failed to compile the change
+		/// </summary>
+		CompileFailure,
+
+		/// <summary>
+		/// Manually marked the change as good
+		/// </summary>
+		Good,
+
+		/// <summary>
+		/// Manually marked the change as bad
+		/// </summary>
+		Bad
+	}
+
+	/// <summary>
+	/// State of a badge
+	/// </summary>
+	public enum UgsBadgeState
+	{
+		/// <summary>
+		/// Starting work on this badge, outcome currently unknown
+		/// </summary>
+		Starting,
+
+		/// <summary>
+		/// Badge failed
+		/// </summary>
+		Failure,
+
+		/// <summary>
+		/// Badge produced a warning
+		/// </summary>
+		Warning,
+
+		/// <summary>
+		/// Badge succeeded
+		/// </summary>
+		Success,
+
+		/// <summary>
+		/// Badge was skipped
+		/// </summary>
+		Skipped,
+	}
+
+	/// <summary>
+	/// Adds a new badge to the change
+	/// </summary>
+	public class AddUgsBadgeRequest
+	{
+		/// <summary>
+		/// Name of the badge
+		/// </summary>
+		[Required]
+		public string Name { get; set; } = String.Empty;
+
+		/// <summary>
+		/// Url for the badge
+		/// </summary>
+		public Uri? Url { get; set; }
+
+		/// <summary>
+		/// Current status of the badge
+		/// </summary>
+		public UgsBadgeState State { get; set; }
+	}
+
+	/// <summary>
+	/// Request object for adding new metadata to the server
+	/// </summary>
+	public class AddUgsMetadataRequest
+	{
+		/// <summary>
+		/// The stream name
+		/// </summary>
+		[Required]
+		public string Stream { get; set; } = null!;
+
+		/// <summary>
+		/// The changelist number
+		/// </summary>
+		[Required]
+		public int Change { get; set; }
+
+		/// <summary>
+		/// The project name
+		/// </summary>
+		public string? Project { get; set; }
+
+		/// <summary>
+		/// Name of the current user
+		/// </summary>
+		public string? UserName { get; set; }
+
+		/// <summary>
+		/// Whether this changelist has been synced by the user
+		/// </summary>
+		public bool? Synced { get; set; }
+
+		/// <summary>
+		/// State of the user
+		/// </summary>
+		public UgsUserVote? Vote { get; set; }
+
+		/// <summary>
+		/// New starred state for the issue
+		/// </summary>
+		public bool? Starred { get; set; }
+
+		/// <summary>
+		/// Whether the user is investigating
+		/// </summary>
+		public bool? Investigating { get; set; }
+
+		/// <summary>
+		/// Comment for this change
+		/// </summary>
+		public string? Comment { get; set; }
+
+		/// <summary>
+		/// List of badges to add
+		/// </summary>
+		public List<AddUgsBadgeRequest>? Badges { get; set; }
+	}
+
+	/// <summary>
+	/// Information about a user synced to a change
+	/// </summary>
+	public class GetUgsUserResponse
+	{
+		/// <summary>
+		/// Name of the user
+		/// </summary>
+		public string User { get; set; }
+
+		/// <summary>
+		/// Time that the change was synced
+		/// </summary>
+		public long? SyncTime { get; set; }
+
+		/// <summary>
+		/// State of the user
+		/// </summary>
+		public UgsUserVote? Vote { get; set; }
+
+		/// <summary>
+		/// Comment by this user
+		/// </summary>
+		public string? Comment { get; set; }
+
+		/// <summary>
+		/// Whether the user is investigating this change
+		/// </summary>
+		public bool? Investigating { get; set; }
+
+		/// <summary>
+		/// Whether this changelist is starred
+		/// </summary>
+		public bool? Starred { get; set; }
+
+		/// <summary>
+		/// Constructor
+		/// </summary>
+		/// <param name="userData">The user data object</param>
+		internal GetUgsUserResponse(IUgsUserData userData)
+		{
+			User = userData.User;
+			SyncTime = userData.SyncTime;
+			Vote = (userData.Vote == UgsUserVote.None) ? (UgsUserVote?)null : userData.Vote;
+			Comment = userData.Comment;
+			Investigating = userData.Investigating;
+			Starred = userData.Starred;
+		}
+	}
+
+	/// <summary>
+	/// Information about a badge
+	/// </summary>
+	public class GetUgsBadgeResponse
+	{
+		/// <summary>
+		/// Name of the badge
+		/// </summary>
+		public string Name { get; set; }
+
+		/// <summary>
+		/// Url for the badge
+		/// </summary>
+		public Uri? Url { get; set; }
+
+		/// <summary>
+		/// Current status of the badge
+		/// </summary>
+		public UgsBadgeState State { get; set; }
+
+		/// <summary>
+		/// Constructor
+		/// </summary>
+		/// <param name="badgeData">The badge data object</param>
+		internal GetUgsBadgeResponse(IUgsBadgeData badgeData)
+		{
+			Name = badgeData.Name;
+			Url = badgeData.Url;
+			State = badgeData.State;
+		}
+	}
+
+	/// <summary>
+	/// Response object for querying metadata
+	/// </summary>
+	public class GetUgsMetadataResponse
+	{
+		/// <summary>
+		/// Number of the changelist
+		/// </summary>
+		public int Change { get; set; }
+
+		/// <summary>
+		/// The project name
+		/// </summary>
+		public string? Project { get; set; }
+
+		/// <summary>
+		/// Information about a user synced to this change
+		/// </summary>
+		public List<GetUgsUserResponse>? Users { get; set; }
+
+		/// <summary>
+		/// Badges for this change
+		/// </summary>
+		public List<GetUgsBadgeResponse>? Badges { get; set; }
+
+		/// <summary>
+		/// Constructor
+		/// </summary>
+		/// <param name="metadata">Metadata object</param>
+		internal GetUgsMetadataResponse(IUgsMetadata metadata)
+		{
+			Change = metadata.Change;
+			Project = metadata.Project;
+			Users = metadata.Users?.ConvertAll(x => new GetUgsUserResponse(x));
+			Badges = metadata.Badges?.ConvertAll(x => new GetUgsBadgeResponse(x));
+		}
+	}
+
+	/// <summary>
+	/// Response object for querying metadata
+	/// </summary>
+	public class GetUgsMetadataListResponse
+	{
+		/// <summary>
+		/// Last time that the metadata was modified
+		/// </summary>
+		public long SequenceNumber { get; set; }
+
+		/// <summary>
+		/// List of changes matching the requested criteria
+		/// </summary>
+		public List<GetUgsMetadataResponse> Items { get; set; } = new List<GetUgsMetadataResponse>();
+	}
+
+	/// <summary>
+	/// Outcome of a particular build
+	/// </summary>
+	public enum UgsIssueBuildOutcome
+	{
+		/// <summary>
+		/// Unknown outcome
+		/// </summary>
+		Unknown,
+
+		/// <summary>
+		/// Build succeeded
+		/// </summary>
+		Success,
+
+		/// <summary>
+		/// Build failed
+		/// </summary>
+		Error,
+
+		/// <summary>
+		/// Build finished with warnings
+		/// </summary>
+		Warning,
+	}
+
+	/// <summary>
+	/// Legacy response describing a build
+	/// </summary>
+	public class GetUgsIssueBuildResponse
+	{
+		/// <summary>
+		/// Identifier for this build
+		/// </summary>
+		public long Id { get; set; }
+
+		/// <summary>
+		/// Path to the stream containing this build
+		/// </summary>
+		public string Stream { get; set; }
+
+		/// <summary>
+		/// The changelist number for this build
+		/// </summary>
+		public int Change { get; set; }
+
+		/// <summary>
+		/// Name of this job
+		/// </summary>
+		public string? JobName { get; set; }
+
+		/// <summary>
+		/// Link to the job
+		/// </summary>
+		public Uri? JobUrl { get; set; }
+
+		/// <summary>
+		/// Name of the job step
+		/// </summary>
+		public string? JobStepName { get; set; }
+
+		/// <summary>
+		/// Link to the job step
+		/// </summary>
+		public Uri? JobStepUrl { get; set; }
+
+		/// <summary>
+		/// Url for this particular error
+		/// </summary>
+		public Uri? ErrorUrl { get; set; }
+
+		/// <summary>
+		/// Outcome of this build (see <see cref="IssueBuildOutcome"/>)
+		/// </summary>
+		public int Outcome { get; set; }
+
+		/// <summary>
+		/// 
+		/// </summary>
+		/// <param name="stream"></param>
+		/// <param name="change"></param>
+		/// <param name="outcome"></param>
+		public GetUgsIssueBuildResponse(string stream, int change, IssueBuildOutcome outcome)
+		{
+			Stream = stream;
+			Change = change;
+			Outcome = (int)outcome;
+		}
+	}
+
+	/// <summary>
+	/// Information about a diagnostic
+	/// </summary>
+	public class GetUgsIssueDiagnosticResponse
+	{
+		/// <summary>
+		/// The corresponding build id
+		/// </summary>
+		public long? BuildId { get; set; }
+
+		/// <summary>
+		/// Message for the diagnostic
+		/// </summary>
+		public string Message { get; set; }
+
+		/// <summary>
+		/// Link to the error
+		/// </summary>
+		public Uri Url { get; set; }
+
+		/// <summary>
+		/// Constructor
+		/// </summary>
+		/// <param name="buildId">The corresponding build id</param>
+		/// <param name="message">Message for the diagnostic</param>
+		/// <param name="url">Link to the diagnostic</param>
+		public GetUgsIssueDiagnosticResponse(long? buildId, string message, Uri url)
+		{
+			BuildId = buildId;
+			Message = message;
+			Url = url;
+		}
+	}
+
+	/// <summary>
+	/// Stores information about a build health issue
+	/// </summary>
+	public class GetUgsIssueResponse
+	{
+		/// <summary>
+		/// Version number for this response
+		/// </summary>
+		public int Version { get; set; } = 1;
+
+		/// <summary>
+		/// The unique object id
+		/// </summary>
+		public int Id { get; set; }
+
+		/// <summary>
+		/// Time at which the issue was created
+		/// </summary>
+		public DateTime CreatedAt { get; set; }
+
+		/// <summary>
+		/// Time at which the issue was retrieved
+		/// </summary>
+		public DateTime RetrievedAt { get; set; }
+
+		/// <summary>
+		/// The associated project for the issue
+		/// </summary>
+		public string? Project { get; set; }
+
+		/// <summary>
+		/// The summary text for this issue
+		/// </summary>
+		public string Summary { get; set; }
+
+		/// <summary>
+		/// Owner of the issue
+		/// </summary>
+		public string? Owner { get; set; }
+
+		/// <summary>
+		/// User that nominated the current owner
+		/// </summary>
+		public string? NominatedBy { get; set; }
+
+		/// <summary>
+		/// Time that the issue was acknowledged
+		/// </summary>
+		public DateTime? AcknowledgedAt { get; set; }
+
+		/// <summary>
+		/// Changelist that fixed this issue
+		/// </summary>
+		public int? FixChange { get; set; }
+
+		/// <summary>
+		/// Time at which the issue was resolved
+		/// </summary>
+		public DateTime? ResolvedAt { get; set; }
+
+		/// <summary>
+		/// Whether to notify the user about this issue
+		/// </summary>
+		public bool BNotify { get; set; }
+
+		/// <summary>
+		/// Whether this issue just contains warnings
+		/// </summary>
+		public bool BWarning { get; set; }
+
+		/// <summary>
+		/// Link to the last build
+		/// </summary>
+		public Uri? BuildUrl { get; set; }
+
+		/// <summary>
+		/// List of streams affected by this issue
+		/// </summary>
+		public List<string> Streams { get; set; }
+
+		/// <summary>
+		/// Constructs a new issue
+		/// </summary>
+		/// <param name="details">Issue to construct from</param>
+		/// <param name="owner"></param>
+		/// <param name="nominatedBy"></param>
+		/// <param name="bNotify">Whether to notify the user about this issue</param>
+		/// <param name="buildUrl">Link to the last build for this issue</param>
+		public GetUgsIssueResponse(IIssueDetails details, IUser? owner, IUser? nominatedBy, bool bNotify, Uri? buildUrl)
+		{
+			IIssue issue = details.Issue;
+			Id = issue.Id;
+			CreatedAt = issue.CreatedAt;
+			RetrievedAt = DateTime.UtcNow;
+			Summary = issue.Summary;
+			Owner = owner?.Login;
+			NominatedBy = nominatedBy?.Login;
+			AcknowledgedAt = issue.AcknowledgedAt;
+			FixChange = issue.FixChange;
+			ResolvedAt = issue.ResolvedAt;
+			BNotify = bNotify;
+			BWarning = issue.Severity == IssueSeverity.Warning;
+			BuildUrl = buildUrl;
+			Streams = details.Spans.Select(x => x.StreamName ?? x.StreamId.ToString()).Distinct().ToList();
+		}
+	}
+
+	/// <summary>
+	/// Request an issue to be updated
+	/// </summary>
+	public class UpdateUgsIssueRequest
+	{
+		/// <summary>
+		/// New owner of the issue
+		/// </summary>
+		public string? Owner { get; set; }
+
+		/// <summary>
+		/// User than nominates the new owner
+		/// </summary>
+		public string? NominatedBy { get; set; }
+
+		/// <summary>
+		/// Whether the issue has been acknowledged
+		/// </summary>
+		public bool? Acknowledged { get; set; }
+
+		/// <summary>
+		/// Name of the user that declines the issue
+		/// </summary>
+		public string? DeclinedBy { get; set; }
+
+		/// <summary>
+		/// The change at which the issue is claimed fixed. 0 = not fixed, -1 = systemic issue.
+		/// </summary>
+		public int? FixChange { get; set; }
+
+		/// <summary>
+		/// Whether the issue should be marked as resolved
+		/// </summary>
+		public bool? Resolved { get; set; }
+
+		/// <summary>
+		/// Name of the user that resolved the issue
+		/// </summary>
+		public string? ResolvedBy { get; set; }
+	}
+}
\ No newline at end of file
