Horde: Threads for new issues now have "New Issue XXX" in the title, rather than just "Issue XXX".

diff --git a/Engine/Source/Programs/Horde/Horde.Build/Notifications/Sinks/SlackNotificationSink.cs b/Engine/Source/Programs/Horde/Horde.Build/Notifications/Sinks/SlackNotificationSink.cs
--- a/Engine/Source/Programs/Horde/Horde.Build/Notifications/Sinks/SlackNotificationSink.cs
+++ b/Engine/Source/Programs/Horde/Horde.Build/Notifications/Sinks/SlackNotificationSink.cs
@@ -744,7 +744,7 @@ namespace Horde.Build.Notifications.Sinks
 
 				string eventId = GetTriageThreadEventId(issue.Id);
 
-				string text = $"{workflow.TriagePrefix}*Issue <{issueUrl}|{issue.Id}>*: {issue.Summary}{workflow.TriageSuffix}";
+				string text = $"{workflow.TriagePrefix}*New Issue <{issueUrl}|{issue.Id}>*: {issue.Summary}{workflow.TriageSuffix}";
 				if (!spans.Any(x => x.NextSuccess == null))
 				{
 					text = $"~{text}~";
