Horde: Store the description of the change being preflighted on the job document.

diff --git a/Engine/Source/Programs/Horde/Horde.Build.Tests/Stubs/Services/PerforceServiceStub.cs b/Engine/Source/Programs/Horde/Horde.Build.Tests/Stubs/Services/PerforceServiceStub.cs
--- a/Engine/Source/Programs/Horde/Horde.Build.Tests/Stubs/Services/PerforceServiceStub.cs
+++ b/Engine/Source/Programs/Horde/Horde.Build.Tests/Stubs/Services/PerforceServiceStub.cs
@@ -108,7 +108,7 @@ namespace Horde.Build.Tests.Stubs.Services
 			return Task.FromResult(results);
 		}
 
-		public Task<CheckShelfResult> CheckShelfAsync(string clusterName, string streamName, int changeNumber, string? impersonateUser)
+		public Task<(CheckShelfResult, string?)> CheckShelfAsync(string clusterName, string streamName, int changeNumber, string? impersonateUser)
 		{
 			throw new NotImplementedException();
 		}
diff --git a/Engine/Source/Programs/Horde/Horde.Build/Jobs/IJob.cs b/Engine/Source/Programs/Horde/Horde.Build/Jobs/IJob.cs
--- a/Engine/Source/Programs/Horde/Horde.Build/Jobs/IJob.cs
+++ b/Engine/Source/Programs/Horde/Horde.Build/Jobs/IJob.cs
@@ -606,6 +606,11 @@ namespace Horde.Build.Jobs
 		/// </summary>
 		public int ClonedPreflightChange { get; }
 
+		/// <summary>
+		/// Description for the shelved change if running a preflight
+		/// </summary>
+		public string? PreflightDescription { get; }
+
 		/// <summary>
 		/// Priority of this job
 		/// </summary>
diff --git a/Engine/Source/Programs/Horde/Horde.Build/Jobs/IJobCollection.cs b/Engine/Source/Programs/Horde/Horde.Build/Jobs/IJobCollection.cs
--- a/Engine/Source/Programs/Horde/Horde.Build/Jobs/IJobCollection.cs
+++ b/Engine/Source/Programs/Horde/Horde.Build/Jobs/IJobCollection.cs
@@ -45,6 +45,7 @@ namespace Horde.Build.Jobs
 		/// <param name="codeChange">The corresponding code changelist number</param>
 		/// <param name="preflightChange">Optional changelist to preflight</param>
 		/// <param name="clonedPreflightChange">Optional cloned preflight changelist</param>
+		/// <param name="preflightDescription">Description for the change being preflighted</param>
 		/// <param name="startedByUserId">User that started the job</param>
 		/// <param name="priority">Priority of the job</param>
 		/// <param name="autoSubmit">Whether to automatically submit the preflighted change on completion</param>
@@ -57,7 +58,7 @@ namespace Horde.Build.Jobs
 		/// <param name="notificationChannelFilter">Notification channel filter for this job</param>
 		/// <param name="arguments">Arguments for the job</param>
 		/// <returns>The new job document</returns>
-		Task<IJob> AddAsync(JobId jobId, StreamId streamId, TemplateRefId templateRefId, ContentHash templateHash, IGraph graph, string name, int change, int codeChange, int? preflightChange, int? clonedPreflightChange, UserId? startedByUserId, Priority? priority, bool? autoSubmit, bool? updateIssues, bool? promoteIssuesByDefault, List<ChainedJobTemplate>? jobTriggers, bool showUgsBadges, bool showUgsAlerts, string? notificationChannel, string? notificationChannelFilter, List<string>? arguments);
+		Task<IJob> AddAsync(JobId jobId, StreamId streamId, TemplateRefId templateRefId, ContentHash templateHash, IGraph graph, string name, int change, int codeChange, int? preflightChange, int? clonedPreflightChange, string? preflightDescription, UserId? startedByUserId, Priority? priority, bool? autoSubmit, bool? updateIssues, bool? promoteIssuesByDefault, List<ChainedJobTemplate>? jobTriggers, bool showUgsBadges, bool showUgsAlerts, string? notificationChannel, string? notificationChannelFilter, List<string>? arguments);
 
 		/// <summary>
 		/// Gets a job with the given unique id
diff --git a/Engine/Source/Programs/Horde/Horde.Build/Jobs/JobCollection.cs b/Engine/Source/Programs/Horde/Horde.Build/Jobs/JobCollection.cs
--- a/Engine/Source/Programs/Horde/Horde.Build/Jobs/JobCollection.cs
+++ b/Engine/Source/Programs/Horde/Horde.Build/Jobs/JobCollection.cs
@@ -231,6 +231,7 @@ namespace Horde.Build.Jobs
 			public int CodeChange { get; set; }
 			public int PreflightChange { get; set; }
 			public int ClonedPreflightChange { get; set; }
+			public string? PreflightDescription { get; set; }
 			public Priority Priority { get; set; }
 
 			[BsonIgnoreIfDefault]
@@ -296,7 +297,7 @@ namespace Horde.Build.Jobs
 				GraphHash = null!;
 			}
 
-			public JobDocument(JobId id, StreamId streamId, TemplateRefId templateId, ContentHash templateHash, ContentHash graphHash, string name, int change, int codeChange, int preflightChange, int clonedPreflightChange, UserId? startedByUserId, Priority? priority, bool? autoSubmit, bool? updateIssues, bool? promoteIssuesByDefault, DateTime createTimeUtc, List<ChainedJobDocument> chainedJobs, bool showUgsBadges, bool showUgsAlerts, string? notificationChannel, string? notificationChannelFilter, List<string>? arguments)
+			public JobDocument(JobId id, StreamId streamId, TemplateRefId templateId, ContentHash templateHash, ContentHash graphHash, string name, int change, int codeChange, int preflightChange, int clonedPreflightChange, string? preflightDescription, UserId? startedByUserId, Priority? priority, bool? autoSubmit, bool? updateIssues, bool? promoteIssuesByDefault, DateTime createTimeUtc, List<ChainedJobDocument> chainedJobs, bool showUgsBadges, bool showUgsAlerts, string? notificationChannel, string? notificationChannelFilter, List<string>? arguments)
 			{
 				Id = id;
 				StreamId = streamId;
@@ -308,6 +309,7 @@ namespace Horde.Build.Jobs
 				CodeChange = codeChange;
 				PreflightChange = preflightChange;
 				ClonedPreflightChange = clonedPreflightChange;
+				PreflightDescription = preflightDescription;
 				StartedByUserId = startedByUserId;
 				Priority = priority ?? HordeCommon.Priority.Normal;
 				AutoSubmit = autoSubmit ?? false;
@@ -394,7 +396,7 @@ namespace Horde.Build.Jobs
 		}
 
 		/// <inheritdoc/>
-		public async Task<IJob> AddAsync(JobId jobId, StreamId streamId, TemplateRefId templateRefId, ContentHash templateHash, IGraph graph, string name, int change, int codeChange, int? preflightChange, int? clonedPreflightChange, UserId? startedByUserId, Priority? priority, bool? autoSubmit, bool? updateIssues, bool? promoteIssuesByDefault, List<ChainedJobTemplate>? chainedJobs, bool showUgsBadges, bool showUgsAlerts, string? notificationChannel, string? notificationChannelFilter, List<string>? arguments)
+		public async Task<IJob> AddAsync(JobId jobId, StreamId streamId, TemplateRefId templateRefId, ContentHash templateHash, IGraph graph, string name, int change, int codeChange, int? preflightChange, int? clonedPreflightChange, string? preflightDescription, UserId? startedByUserId, Priority? priority, bool? autoSubmit, bool? updateIssues, bool? promoteIssuesByDefault, List<ChainedJobTemplate>? chainedJobs, bool showUgsBadges, bool showUgsAlerts, string? notificationChannel, string? notificationChannelFilter, List<string>? arguments)
 		{
 			List<ChainedJobDocument> jobTriggers = new List<ChainedJobDocument>();
 			if (chainedJobs == null)
@@ -406,7 +408,7 @@ namespace Horde.Build.Jobs
 				jobTriggers = chainedJobs.ConvertAll(x => new ChainedJobDocument(x));
 			}
 
-			JobDocument newJob = new JobDocument(jobId, streamId, templateRefId, templateHash, graph.Id, name, change, codeChange, preflightChange ?? 0, clonedPreflightChange ?? 0, startedByUserId, priority, autoSubmit, updateIssues, promoteIssuesByDefault, DateTime.UtcNow, jobTriggers, showUgsBadges, showUgsAlerts, notificationChannel, notificationChannelFilter, arguments);
+			JobDocument newJob = new JobDocument(jobId, streamId, templateRefId, templateHash, graph.Id, name, change, codeChange, preflightChange ?? 0, clonedPreflightChange ?? 0, preflightDescription, startedByUserId, priority, autoSubmit, updateIssues, promoteIssuesByDefault, DateTime.UtcNow, jobTriggers, showUgsBadges, showUgsAlerts, notificationChannel, notificationChannelFilter, arguments);
 			CreateBatches(newJob, graph, _logger);
 
 			await _jobs.InsertOneAsync(newJob);
diff --git a/Engine/Source/Programs/Horde/Horde.Build/Jobs/JobMessages.cs b/Engine/Source/Programs/Horde/Horde.Build/Jobs/JobMessages.cs
--- a/Engine/Source/Programs/Horde/Horde.Build/Jobs/JobMessages.cs
+++ b/Engine/Source/Programs/Horde/Horde.Build/Jobs/JobMessages.cs
@@ -231,6 +231,11 @@ namespace Horde.Build.Jobs
 		/// </summary>
 		public int? ClonedPreflightChange { get; set; }
 
+		/// <summary>
+		/// Description of the preflight
+		/// </summary>
+		public string? PreflightDescription { get; set; }
+
 		/// <summary>
 		/// The template type
 		/// </summary>
@@ -357,6 +362,7 @@ namespace Horde.Build.Jobs
 			CodeChange = (job.CodeChange != 0) ? (int?)job.CodeChange : null;
 			PreflightChange = (job.PreflightChange != 0) ? (int?)job.PreflightChange : null;
 			ClonedPreflightChange = (job.ClonedPreflightChange != 0) ? (int?)job.ClonedPreflightChange : null;
+			PreflightDescription = job.PreflightDescription;
 			TemplateId = job.TemplateId.ToString();
 			TemplateHash = job.TemplateHash?.ToString() ?? String.Empty;
 			GraphHash = job.GraphHash.ToString();
diff --git a/Engine/Source/Programs/Horde/Horde.Build/Jobs/JobService.cs b/Engine/Source/Programs/Horde/Horde.Build/Jobs/JobService.cs
--- a/Engine/Source/Programs/Horde/Horde.Build/Jobs/JobService.cs
+++ b/Engine/Source/Programs/Horde/Horde.Build/Jobs/JobService.cs
@@ -175,6 +175,7 @@ namespace Horde.Build.Jobs
 		/// <param name="codeChange">The corresponding code changelist</param>
 		/// <param name="preflightChange">Optional changelist to preflight</param>
 		/// <param name="clonedPreflightChange">Duplicated preflight change</param>
+		/// <param name="preflightDescription">Description for the preflight</param>
 		/// <param name="startedByUserId">Id of the user that started the job</param>
 		/// <param name="priority">Priority of the job</param>
 		/// <param name="autoSubmit">Whether to automatically submit the preflighted change on completion</param>
@@ -187,7 +188,7 @@ namespace Horde.Build.Jobs
 		/// <param name="notificationChannelFilter">Notification Channel filter for this job</param>
 		/// <param name="arguments">Arguments for the job</param>
 		/// <returns>Unique id representing the job</returns>
-		public async Task<IJob> CreateJobAsync(JobId? jobId, IStream stream, TemplateRefId templateRefId, ContentHash templateHash, IGraph graph, string name, int change, int codeChange, int? preflightChange, int? clonedPreflightChange, UserId? startedByUserId, Priority? priority, bool? autoSubmit, bool? updateIssues, bool? promoteIssuesByDefault, List<ChainedJobTemplate>? jobTriggers, bool showUgsBadges, bool showUgsAlerts, string? notificationChannel, string? notificationChannelFilter, IReadOnlyList<string> arguments)
+		public async Task<IJob> CreateJobAsync(JobId? jobId, IStream stream, TemplateRefId templateRefId, ContentHash templateHash, IGraph graph, string name, int change, int codeChange, int? preflightChange, int? clonedPreflightChange, string? preflightDescription, UserId? startedByUserId, Priority? priority, bool? autoSubmit, bool? updateIssues, bool? promoteIssuesByDefault, List<ChainedJobTemplate>? jobTriggers, bool showUgsBadges, bool showUgsAlerts, string? notificationChannel, string? notificationChannelFilter, IReadOnlyList<string> arguments)
 		{
 			using IScope traceScope = GlobalTracer.Instance.BuildSpan("JobService.CreateJobAsync").StartActive();
 			traceScope.Span.SetTag("JobId", jobId);
@@ -251,7 +252,7 @@ namespace Horde.Build.Jobs
 
 			name = StringUtils.ExpandProperties(name, properties);
 
-			IJob newJob = await _jobs.AddAsync(jobIdValue, stream.Id, templateRefId, templateHash, graph, name, change, codeChange, preflightChange, clonedPreflightChange, startedByUserId, priority, autoSubmit, updateIssues, promoteIssuesByDefault, jobTriggers, showUgsBadges, showUgsAlerts, notificationChannel, notificationChannelFilter, expandedArguments);
+			IJob newJob = await _jobs.AddAsync(jobIdValue, stream.Id, templateRefId, templateHash, graph, name, change, codeChange, preflightChange, clonedPreflightChange, preflightDescription, startedByUserId, priority, autoSubmit, updateIssues, promoteIssuesByDefault, jobTriggers, showUgsBadges, showUgsAlerts, notificationChannel, notificationChannelFilter, expandedArguments);
 			_jobTaskSource.UpdateQueuedJob(newJob, graph);
 
 			await _jobTaskSource.UpdateUgsBadges(newJob, graph, new List<(LabelState, LabelOutcome)>());
@@ -1242,7 +1243,7 @@ namespace Horde.Build.Jobs
 					IGraph triggerGraph = await _graphs.AddAsync(template);
 					_logger.LogInformation("Creating downstream job {ChainedJobId} from job {JobId}", chainedJobId, newJob.Id);
 
-					await CreateJobAsync(chainedJobId, stream, jobTrigger.TemplateRefId, templateRef.Hash, triggerGraph, templateRef.Name, newJob.Change, newJob.CodeChange, newJob.PreflightChange, newJob.ClonedPreflightChange, newJob.StartedByUserId, template.Priority, null, newJob.UpdateIssues, newJob.PromoteIssuesByDefault, templateRef.ChainedJobs, false, false, templateRef.NotificationChannel, templateRef.NotificationChannelFilter, template.Arguments);
+					await CreateJobAsync(chainedJobId, stream, jobTrigger.TemplateRefId, templateRef.Hash, triggerGraph, templateRef.Name, newJob.Change, newJob.CodeChange, newJob.PreflightChange, newJob.ClonedPreflightChange, newJob.PreflightDescription, newJob.StartedByUserId, template.Priority, null, newJob.UpdateIssues, newJob.PromoteIssuesByDefault, templateRef.ChainedJobs, false, false, templateRef.NotificationChannel, templateRef.NotificationChannelFilter, template.Arguments);
 					return newJob;
 				}
 
diff --git a/Engine/Source/Programs/Horde/Horde.Build/Jobs/JobsController.cs b/Engine/Source/Programs/Horde/Horde.Build/Jobs/JobsController.cs
--- a/Engine/Source/Programs/Horde/Horde.Build/Jobs/JobsController.cs
+++ b/Engine/Source/Programs/Horde/Horde.Build/Jobs/JobsController.cs
@@ -159,9 +159,10 @@ namespace Horde.Build.Jobs
 			List<string> arguments = create.Arguments ?? template.GetDefaultArguments();
 
 			// Check the preflight change is valid
+			string? preflightDescription = null;
 			if (create.PreflightChange != null)
 			{
-				CheckShelfResult result = await _perforce.CheckShelfAsync(stream.ClusterName, stream.Name, create.PreflightChange.Value, null);
+				(CheckShelfResult result, preflightDescription) = await _perforce.CheckShelfAsync(stream.ClusterName, stream.Name, create.PreflightChange.Value, null);
 				switch (result)
 				{
 					case CheckShelfResult.Ok:
@@ -190,7 +191,7 @@ namespace Horde.Build.Jobs
 			}
 
 			// Create the job
-			IJob job = await _jobService.CreateJobAsync(null, stream, templateRefId, template.Id, graph, name, change, codeChange, create.PreflightChange, null, User.GetUserId(), priority, create.AutoSubmit, updateIssues, false, templateRef.ChainedJobs, templateRef.ShowUgsBadges, templateRef.ShowUgsAlerts, templateRef.NotificationChannel, templateRef.NotificationChannelFilter, arguments);
+			IJob job = await _jobService.CreateJobAsync(null, stream, templateRefId, template.Id, graph, name, change, codeChange, create.PreflightChange, null, preflightDescription, User.GetUserId(), priority, create.AutoSubmit, updateIssues, false, templateRef.ChainedJobs, templateRef.ShowUgsBadges, templateRef.ShowUgsAlerts, templateRef.NotificationChannel, templateRef.NotificationChannelFilter, arguments);
 			await UpdateNotificationsAsync(job.Id, new UpdateNotificationsRequest { Slack = true });
 			return new CreateJobResponse(job.Id.ToString());
 		}
diff --git a/Engine/Source/Programs/Horde/Horde.Build/Jobs/Schedules/ScheduleService.cs b/Engine/Source/Programs/Horde/Horde.Build/Jobs/Schedules/ScheduleService.cs
--- a/Engine/Source/Programs/Horde/Horde.Build/Jobs/Schedules/ScheduleService.cs
+++ b/Engine/Source/Programs/Horde/Horde.Build/Jobs/Schedules/ScheduleService.cs
@@ -521,7 +521,7 @@ namespace Horde.Build.Jobs.Schedules
 			{
 				cancellationToken.ThrowIfCancellationRequested();
 				List<string> defaultArguments = template.GetDefaultArguments();
-				IJob newJob = await _jobService.CreateJobAsync(null, stream, templateId, template.Id, graph, template.Name, change, codeChange, null, null, null, template.Priority, null, null, template.PromoteIssuesByDefault, templateRef.ChainedJobs, templateRef.ShowUgsBadges, templateRef.ShowUgsAlerts, templateRef.NotificationChannel, templateRef.NotificationChannelFilter, defaultArguments);
+				IJob newJob = await _jobService.CreateJobAsync(null, stream, templateId, template.Id, graph, template.Name, change, codeChange, null, null, null, null, template.Priority, null, null, template.PromoteIssuesByDefault, templateRef.ChainedJobs, templateRef.ShowUgsBadges, templateRef.ShowUgsAlerts, templateRef.NotificationChannel, templateRef.NotificationChannelFilter, defaultArguments);
 				_logger.LogInformation("Started new job for {StreamId} template {TemplateId} at CL {Change} (Code CL {CodeChange}): {JobId}", stream.Id, templateId, change, codeChange, newJob.Id);
 				await _streamService.UpdateScheduleTriggerAsync(stream, templateId, utcNow, change, new List<JobId> { newJob.Id }, new List<JobId>());
 			}
diff --git a/Engine/Source/Programs/Horde/Horde.Build/Perforce/IPerforceService.cs b/Engine/Source/Programs/Horde/Horde.Build/Perforce/IPerforceService.cs
--- a/Engine/Source/Programs/Horde/Horde.Build/Perforce/IPerforceService.cs
+++ b/Engine/Source/Programs/Horde/Horde.Build/Perforce/IPerforceService.cs
@@ -162,8 +162,8 @@ namespace Horde.Build.Perforce
 		/// <param name="streamName">The stream to query</param>
 		/// <param name="changeNumber">Shelved changelist number</param>
 		/// <param name="impersonateUser">Name of the user to impersonate</param>
-		/// <returns></returns>
-		public Task<CheckShelfResult> CheckShelfAsync(string clusterName, string streamName, int changeNumber, string? impersonateUser);
+		/// <returns>Whether the shelf is valid, plus the description for it</returns>
+		public Task<(CheckShelfResult, string?)> CheckShelfAsync(string clusterName, string streamName, int changeNumber, string? impersonateUser);
 
 		/// <summary>
 		/// Gets the latest change for a particular stream
diff --git a/Engine/Source/Programs/Horde/Horde.Build/Perforce/PerforceService.cs b/Engine/Source/Programs/Horde/Horde.Build/Perforce/PerforceService.cs
--- a/Engine/Source/Programs/Horde/Horde.Build/Perforce/PerforceService.cs
+++ b/Engine/Source/Programs/Horde/Horde.Build/Perforce/PerforceService.cs
@@ -943,7 +943,7 @@ namespace Horde.Build.Perforce
 		}
 
 		/// <inheritdoc/>
-		public async Task<CheckShelfResult> CheckShelfAsync(string clusterName, string streamName, int changeNumber, string? impersonateUser)
+		public async Task<(CheckShelfResult, string?)> CheckShelfAsync(string clusterName, string streamName, int changeNumber, string? impersonateUser)
 		{
 			using IScope scope = GlobalTracer.Instance.BuildSpan("PerforceService.CheckPreflightAsync").StartActive();
 			scope.Span.SetTag("ClusterName", clusterName);
@@ -957,11 +957,11 @@ namespace Horde.Build.Perforce
 				P4.Changelist change = repository.GetChangelist(changeNumber, new P4.DescribeCmdOptions(P4.DescribeChangelistCmdFlags.Omit | P4.DescribeChangelistCmdFlags.Shelved, 0, 0));
 				if(change == null)
 				{
-					return CheckShelfResult.NoChange;
+					return (CheckShelfResult.NoChange, null);
 				}
 				if (change.ShelvedFiles == null || change.ShelvedFiles.Count == 0)
 				{
-					return CheckShelfResult.NoShelvedFiles;
+					return (CheckShelfResult.NoShelvedFiles, null);
 				}
 
 				P4.Stream stream = repository.GetStream(streamName, new P4.StreamCmdOptions(P4.StreamCmdFlags.View, null, null));
@@ -983,10 +983,11 @@ namespace Horde.Build.Perforce
 
 				if (bHasUnmappedFile)
 				{
-					return bHasMappedFile ? CheckShelfResult.MixedStream : CheckShelfResult.WrongStream;
+					return (bHasMappedFile ? CheckShelfResult.MixedStream : CheckShelfResult.WrongStream, null);
 				}
+
+				return (CheckShelfResult.Ok, change.Description);
 			}
-			return CheckShelfResult.Ok;
 		}
 
 		/// <inheritdoc/>
