Horde: Fix tests.

diff --git a/Engine/Source/Programs/Horde/Horde.Build.Tests/Fixture.cs b/Engine/Source/Programs/Horde/Horde.Build.Tests/Fixture.cs
--- a/Engine/Source/Programs/Horde/Horde.Build.Tests/Fixture.cs
+++ b/Engine/Source/Programs/Horde/Horde.Build.Tests/Fixture.cs
@@ -99,6 +99,7 @@ namespace Horde.Build.Tests
 				codeChange: 1000002,
 				preflightChange: 1001,
 				clonedPreflightChange: null,
+				preflightDescription: null,
 				startedByUserId: null,
 				priority: Priority.Normal,
 				null,
@@ -124,6 +125,7 @@ namespace Horde.Build.Tests
 				codeChange: 2000002,
 				preflightChange: null,
 				clonedPreflightChange: null,
+				preflightDescription: null,
 				startedByUserId: null,
 				priority: Priority.Normal,
 				null,
diff --git a/Engine/Source/Programs/Horde/Horde.Build.Tests/Fleet/JobQueueStrategyTest.cs b/Engine/Source/Programs/Horde/Horde.Build.Tests/Fleet/JobQueueStrategyTest.cs
--- a/Engine/Source/Programs/Horde/Horde.Build.Tests/Fleet/JobQueueStrategyTest.cs
+++ b/Engine/Source/Programs/Horde/Horde.Build.Tests/Fleet/JobQueueStrategyTest.cs
@@ -150,7 +150,7 @@ namespace Horde.Build.Tests.Fleet
 		{
 			IJob job = await JobCollection.AddAsync(ObjectId<IJob>.GenerateNewId(), streamId,
 				new StringId<TemplateRef>("bogusTemplateRefId"), ContentHash.Empty, graph, "bogusJobName",
-				1000, 1000, null, null, null, null, null, null, null, null, false,
+				1000, 1000, null, null, null, null, null, null, null, null, null, false,
 				false, null, null, new List<string> { "-Target=" + nodeNameToExecute });
 
 			return job;
diff --git a/Engine/Source/Programs/Horde/Horde.Build.Tests/JobCollectionTests.cs b/Engine/Source/Programs/Horde/Horde.Build.Tests/JobCollectionTests.cs
--- a/Engine/Source/Programs/Horde/Horde.Build.Tests/JobCollectionTests.cs
+++ b/Engine/Source/Programs/Horde/Horde.Build.Tests/JobCollectionTests.cs
@@ -81,7 +81,7 @@ namespace Horde.Build.Tests
 			arguments.Add("-Target=Publish Client");
 			arguments.Add("-Target=Post-Publish Client");
 
-			IJob job = await JobCollection.AddAsync(JobId.GenerateNewId(), new StreamId("ue4-main"), new TemplateRefId("test-build"), ContentHash.SHA1("hello"), baseGraph, "Test job", 123, 123, null, null, null, null, null, null, null, null, false, false, null, null, arguments);
+			IJob job = await JobCollection.AddAsync(JobId.GenerateNewId(), new StreamId("ue4-main"), new TemplateRefId("test-build"), ContentHash.SHA1("hello"), baseGraph, "Test job", 123, 123, null, null, null, null, null, null, null, null, null, false, false, null, null, arguments);
 
 			job = await StartBatch(job, baseGraph, 0);
 			job = await RunStep(job, baseGraph, 0, 0, JobStepOutcome.Success); // Setup Build
@@ -156,7 +156,7 @@ namespace Horde.Build.Tests
 			arguments.Add("-Target=Step 1");
 			arguments.Add("-Target=Step 3");
 
-			IJob job = await JobCollection.AddAsync(JobId.GenerateNewId(), new StreamId("ue4-main"), new TemplateRefId("test-build"), ContentHash.SHA1("hello"), baseGraph, "Test job", 123, 123, null, null, null, null, null, null, null, null, false, false, null, null, arguments);
+			IJob job = await JobCollection.AddAsync(JobId.GenerateNewId(), new StreamId("ue4-main"), new TemplateRefId("test-build"), ContentHash.SHA1("hello"), baseGraph, "Test job", 123, 123, null, null, null, null, null, null, null, null, null, false, false, null, null, arguments);
 
 			job = await StartBatch(job, baseGraph, 0);
 			job = await RunStep(job, baseGraph, 0, 0, JobStepOutcome.Success); // Setup Build
@@ -212,7 +212,7 @@ namespace Horde.Build.Tests
 			arguments.Add("-Target=Step 1");
 			arguments.Add("-Target=Step 3");
 
-			IJob job = await JobCollection.AddAsync(JobId.GenerateNewId(), new StreamId("ue4-main"), new TemplateRefId("test-build"), ContentHash.SHA1("hello"), baseGraph, "Test job", 123, 123, null, null, null, null, null, null, null, null, false, false, null, null, arguments);
+			IJob job = await JobCollection.AddAsync(JobId.GenerateNewId(), new StreamId("ue4-main"), new TemplateRefId("test-build"), ContentHash.SHA1("hello"), baseGraph, "Test job", 123, 123, null, null, null, null, null, null, null, null, null, false, false, null, null, arguments);
 			Assert.AreEqual(1, job.Batches.Count);
 
 			job = await StartBatch(job, baseGraph, 0);
diff --git a/Engine/Source/Programs/Horde/Horde.Build.Tests/JobServiceTests.cs b/Engine/Source/Programs/Horde/Horde.Build.Tests/JobServiceTests.cs
--- a/Engine/Source/Programs/Horde/Horde.Build.Tests/JobServiceTests.cs
+++ b/Engine/Source/Programs/Horde/Horde.Build.Tests/JobServiceTests.cs
@@ -52,7 +52,7 @@ namespace Horde.Build.Tests
 			IStream? stream = await StreamService.GetStreamAsync(streamId);
 			stream = await CreateOrReplaceStreamAsync(new StreamId("ue5-main"), stream, projectId, streamConfig);
 
-			IJob job = await JobService.CreateJobAsync(null, stream!, templateRefId1, template.Id, graph, "Hello", 1234, 1233, 999, null, null, null, null, null, null, stream!.Templates[templateRefId1].ChainedJobs, true, true, null, null, new List<string>());
+			IJob job = await JobService.CreateJobAsync(null, stream!, templateRefId1, template.Id, graph, "Hello", 1234, 1233, 999, null, null, null, null, null, null, null, stream!.Templates[templateRefId1].ChainedJobs, true, true, null, null, new List<string>());
 			Assert.AreEqual(1, job.ChainedJobs.Count);
 
 			job = Deref(await JobService.UpdateBatchAsync(job, job.Batches[0].Id, LogId.GenerateNewId(), JobStepBatchState.Running));
@@ -113,6 +113,7 @@ namespace Horde.Build.Tests
 				codeChange: 1000002,
 				preflightChange: preflightChange,
 				clonedPreflightChange: null,
+				preflightDescription: null,
 				startedByUserId: user.Id,
 				priority: Priority.Normal,
 				null,
@@ -177,7 +178,7 @@ namespace Horde.Build.Tests
 
 			graph = await GraphCollection.AppendAsync(graph, new List<NewGroup> { groupA, groupB });
 
-			IJob job = await JobService.CreateJobAsync(null, stream!, new TemplateRefId("temp"), template.Id, graph, "Hello", 1234, 1233, 999, null, null, null, null, null, null, null, true, true, null, null, new List<string> { "-Target=Pak" });
+			IJob job = await JobService.CreateJobAsync(null, stream!, new TemplateRefId("temp"), template.Id, graph, "Hello", 1234, 1233, 999, null, null, null, null, null, null, null, null, true, true, null, null, new List<string> { "-Target=Pak" });
 
 			job = Deref(await JobService.UpdateBatchAsync(job, job.Batches[0].Id, LogId.GenerateNewId(), JobStepBatchState.Running));
 			job = Deref(await JobService.UpdateStepAsync(job, job.Batches[0].Id, job.Batches[0].Steps[0].Id, JobStepState.Running));
diff --git a/Engine/Source/Programs/Horde/Horde.Build.Tests/JobTaskSourceTests.cs b/Engine/Source/Programs/Horde/Horde.Build.Tests/JobTaskSourceTests.cs
--- a/Engine/Source/Programs/Horde/Horde.Build.Tests/JobTaskSourceTests.cs
+++ b/Engine/Source/Programs/Horde/Horde.Build.Tests/JobTaskSourceTests.cs
@@ -134,7 +134,7 @@ namespace Horde.Build.Tests
 			// create a new job
 			IJob job = await JobCollection.AddAsync(JobId.GenerateNewId(), Stream.Id,
 				fixture.TemplateRefId1, fixture.Template.Id, graph, "Test Paused Step Job",
-				1000, 1000, null, null, null, Priority.Highest, null, null, null, null, false,
+				1000, 1000, null, null, null, null, Priority.Highest, null, null, null, null, false,
 				false, null, null, new List<string> { "-Target=" + "Step That Depends on Paused Step;Step That Depends on Update Version Files" });
 
 			// validate
diff --git a/Engine/Source/Programs/Horde/Horde.Build.Tests/SchedulerTests.cs b/Engine/Source/Programs/Horde/Horde.Build.Tests/SchedulerTests.cs
--- a/Engine/Source/Programs/Horde/Horde.Build.Tests/SchedulerTests.cs
+++ b/Engine/Source/Programs/Horde/Horde.Build.Tests/SchedulerTests.cs
@@ -442,7 +442,7 @@ namespace Horde.Build.Tests
 			Assert.AreEqual(0, jobs2.Count);
 
 			// Create a job and fail it
-			IJob job1 = await JobService.CreateJobAsync(null, stream, newTemplateRefId1, _template.Id, graphA, "Hello", 1234, 1233, 999, null, null, null, null, null, null, null, true, true, null, null, new List<string> { "-Target=TriggerNext" });
+			IJob job1 = await JobService.CreateJobAsync(null, stream, newTemplateRefId1, _template.Id, graphA, "Hello", 1234, 1233, 999, null, null, null, null, null, null, null, null, true, true, null, null, new List<string> { "-Target=TriggerNext" });
 			SubResourceId batchId1 = job1.Batches[0].Id;
 			SubResourceId stepId1 = job1.Batches[0].Steps[0].Id;
 			job1 = Deref(await JobService.UpdateBatchAsync(job1, batchId1, LogId.GenerateNewId(), JobStepBatchState.Running));
@@ -458,7 +458,7 @@ namespace Horde.Build.Tests
 			Assert.AreEqual(0, jobs3.Count);
 
 			// Create a job and make it succeed
-			IJob job2 = await JobService.CreateJobAsync(null, stream, newTemplateRefId1, _template.Id, graphA, "Hello", 1234, 1233, 999, null, null, null, null, null, null, null, true, true, null, null, new List<string> { "-Target=TriggerNext" });
+			IJob job2 = await JobService.CreateJobAsync(null, stream, newTemplateRefId1, _template.Id, graphA, "Hello", 1234, 1233, 999, null, null, null, null, null, null, null, null, true, true, null, null, new List<string> { "-Target=TriggerNext" });
 			SubResourceId batchId2 = job2.Batches[0].Id;
 			SubResourceId stepId2 = job2.Batches[0].Steps[0].Id;
 			job2 = Deref(await JobService.UpdateBatchAsync(job2, batchId2, LogId.GenerateNewId(), JobStepBatchState.Running));
@@ -528,7 +528,7 @@ namespace Horde.Build.Tests
 			{
 				int codeChange = (change < 1233) ? 1230 : 1233;
 
-				IJob job1 = await JobService.CreateJobAsync(null, stream, newTemplateRefId1, _template.Id, graphA, "Hello", change, codeChange, null, null, null, null, null, false, null, null, true, true, null, null, new List<string> { "-Target=TriggerNext" });
+				IJob job1 = await JobService.CreateJobAsync(null, stream, newTemplateRefId1, _template.Id, graphA, "Hello", change, codeChange, null, null, null, null, null, null, false, null, null, true, true, null, null, new List<string> { "-Target=TriggerNext" });
 				for (int batchIdx = 0; batchIdx < job1.Batches.Count; batchIdx++)
 				{
 					SubResourceId batchId1 = job1.Batches[batchIdx].Id;
