Horde: Add preflight CL to job overview summary
#skipci
#rnx

diff --git a/Engine/Source/Programs/Horde/HordeDashboard/src/backend/Api.ts b/Engine/Source/Programs/Horde/HordeDashboard/src/backend/Api.ts
--- a/Engine/Source/Programs/Horde/HordeDashboard/src/backend/Api.ts
+++ b/Engine/Source/Programs/Horde/HordeDashboard/src/backend/Api.ts
@@ -1063,6 +1063,9 @@ export type GetJobResponse = {
 	/** The preflight changelist number */
 	preflightChange?: number;
 
+	/** Description of the preflight */
+	preflightDescription?: string;
+
 	/** The template type */
 	templateId?: string;
 
