Horde: Add fake HordeRpcServer for testing

- The FakeHordeRpcServer will mock a HordeRpcClient that can be used by WorkerService during tests whenever it needs to interact with the server.
- A basic integration test testing shutdown of WorkerService was added as a start.
- Adding more meaningful tests likely require some refactoring of WorkerService since it's quite opaque from testing point of view.
- To make HordeRpcClient usable by Moq, the gRPC/protobuf generated files had to be set to "public" rather than "internal"
- Since we generate Protobuf files in both Horde.Agent and Horde.Server this caused a name resolution clash, so an import alias was added in Horde.Build.Tests when importing Horde.Agent.

diff --git a/Engine/Source/Programs/Horde/Horde.Agent.Tests/WorkerServiceTest.cs b/Engine/Source/Programs/Horde/Horde.Agent.Tests/WorkerServiceTest.cs
--- a/Engine/Source/Programs/Horde/Horde.Agent.Tests/WorkerServiceTest.cs
+++ b/Engine/Source/Programs/Horde/Horde.Agent.Tests/WorkerServiceTest.cs
@@ -1,10 +1,14 @@
 // Copyright Epic Games, Inc. All Rights Reserved.
 
 using System;
+using System.Collections.Generic;
 using System.IO;
 using System.Threading;
+using System.Threading.Channels;
 using System.Threading.Tasks;
+using Google.Protobuf.WellKnownTypes;
 using Grpc.Core;
+using Grpc.Net.Client;
 using Horde.Agent.Execution.Interfaces;
 using Horde.Agent.Services;
 using Horde.Agent.Utility;
@@ -13,9 +17,9 @@ using HordeCommon.Rpc;
 using HordeCommon.Rpc.Messages;
 using HordeCommon.Rpc.Tasks;
 using Microsoft.Extensions.Logging;
-using Microsoft.Extensions.Logging.Abstractions;
 using Microsoft.Extensions.Options;
 using Microsoft.VisualStudio.TestTools.UnitTesting;
+using Moq;
 
 namespace Horde.Agent.Tests
 {
@@ -23,14 +27,24 @@ namespace Horde.Agent.Tests
 	public class WorkerServiceTest
 	{
 		private readonly ILogger<WorkerService> _workerLogger;
+		private readonly ILogger<GrpcService> _grpcServiceLogger;
+		private readonly ILogger<FakeHordeRpcServer> _hordeRpcServerLogger;
 
 		public WorkerServiceTest()
 		{
-			_workerLogger = NullLogger<WorkerService>.Instance;
+			using ILoggerFactory loggerFactory = LoggerFactory.Create(builder =>
+			{
+				builder.AddConsole(consoleLoggerOptions => consoleLoggerOptions.TimestampFormat = "[HH:mm:ss] ");
+			});
+
+			_workerLogger = loggerFactory.CreateLogger<WorkerService>();
+			_grpcServiceLogger = loggerFactory.CreateLogger<GrpcService>();
+			_hordeRpcServerLogger = loggerFactory.CreateLogger<FakeHordeRpcServer>();
 		}
 
 		private WorkerService GetWorkerService(
-			Func<IRpcConnection, ExecuteJobTask, BeginBatchResponse, IExecutor>? createExecutor = null)
+			Func<IRpcConnection, ExecuteJobTask, BeginBatchResponse, IExecutor>? createExecutor = null,
+			Func<GrpcChannel, HordeRpc.HordeRpcClient>? createHordeRpcClient = null)
 		{
 			AgentSettings settings = new AgentSettings();
 			ServerProfile profile = new ServerProfile();
@@ -44,7 +58,10 @@ namespace Horde.Agent.Tests
 			settings.Executor = ExecutorType.Test; // Not really used since the executor is overridden in the tests
 			IOptions<AgentSettings> settingsMonitor = new TestOptionsMonitor<AgentSettings>(settings);
 
-			return new WorkerService(_workerLogger, settingsMonitor, null!, null!, createExecutor);
+			GrpcService grpcService = new (settingsMonitor, _grpcServiceLogger);
+			WorkerService ws = new (_workerLogger, settingsMonitor, grpcService, null!, createExecutor, createHordeRpcClient);
+			ws._rpcConnectionRetryDelay = TimeSpan.FromSeconds(0.1);
+			return ws;
 		}
 
 		[TestMethod]
@@ -167,5 +184,242 @@ namespace Horde.Agent.Tests
 			await ws.PollForStepAbort(rpcConnection, "jobId1", "batchId1", "logId1", stepCancelSource, stepFinishedSource.Task, stepPollCancelSource.Token);
 			Assert.IsTrue(stepCancelSource.IsCancellationRequested);
 		}
+
+		[TestMethod]
+		public async Task Shutdown()
+		{
+			IExecutor executor = new SimpleTestExecutor(async (step, logger, cancellationToken) =>
+			{
+				await Task.Delay(50, cancellationToken);
+				return JobStepOutcome.Success;
+			});
+			
+			using CancellationTokenSource cts = new ();
+			cts.CancelAfter(20000);
+
+			FakeHordeRpcServer fakeServer = new("bogusServerName", _hordeRpcServerLogger, cts.Token);
+			using WorkerService ws = GetWorkerService((a, b, c) => executor,(c) => fakeServer.GetClient());
+
+			Task handleSessionTask = ws.HandleSessionAsync(cts.Token);
+			await fakeServer.CreateSessionReceived.Task.WaitAsync(cts.Token);
+			await fakeServer.UpdateSessionReceived.Task.WaitAsync(cts.Token);
+			await ws.ShutdownAsync(_workerLogger);
+			await handleSessionTask; // Ensure it runs to completion and no exceptions are raised
+		}
+	}
+
+	/// <summary>
+	/// Fake implementation of a HordeRpc gRPC server.
+	/// Provides a corresponding gRPC client class that can be used with the WorkerService
+	/// to test client-server interactions.
+	/// </summary>
+	internal class FakeHordeRpcServer
+	{
+		private readonly string _serverName;
+		private bool _isStopping = false;
+		private List<Lease> _leases = new();
+		private readonly Mock<HordeRpc.HordeRpcClient> _mockClient;
+		private readonly ILogger<FakeHordeRpcServer> _logger;
+		public readonly TaskCompletionSource<bool> CreateSessionReceived = new();
+		public readonly TaskCompletionSource<bool> UpdateSessionReceived = new();
+
+		public FakeHordeRpcServer(string serverName,ILogger<FakeHordeRpcServer> logger, CancellationToken cancellationToken)
+		{
+			_serverName = serverName;
+			_mockClient = new (MockBehavior.Strict);
+			_logger = logger;
+
+			_mockClient
+				.Setup(m => m.CreateSessionAsync(It.IsAny<CreateSessionRequest>(), null, null, It.IsAny<CancellationToken>()))
+				.Returns<CreateSessionRequest, Metadata, DateTime?, CancellationToken>((request, metadata, expireTime, cancellationToken) => CreateAsyncUnaryCall(OnCreateSessionRequest(request)));
+			
+			_mockClient
+				.Setup(m => m.QueryServerStateV2(null, null, It.IsAny<CancellationToken>()))
+				.Returns(() => GetQueryServerStateCall(cancellationToken));
+			
+			_mockClient
+				.Setup(m => m.UpdateSession(null, It.IsAny<DateTime>(), It.IsAny<CancellationToken>()))
+				.Returns(() => GetUpdateSessionCall(cancellationToken));
+		}
+
+		public HordeRpc.HordeRpcClient GetClient()
+		{
+			return _mockClient.Object;
+		}
+
+		public CreateSessionResponse OnCreateSessionRequest(CreateSessionRequest request)
+		{
+			CreateSessionReceived.TrySetResult(true);
+			_logger.LogInformation("OnCreateSessionRequest: {Name} {Status}", request.Name, request.Status);
+			CreateSessionResponse response = new()
+			{
+				AgentId = "bogusAgentId",
+				Token = "bogusToken",
+				SessionId = "bogusSessionId",
+				ExpiryTime = Timestamp.FromDateTime(DateTime.UtcNow.AddHours(3)),
+			};
+
+			return response;
+		}
+
+		public AsyncDuplexStreamingCall<QueryServerStateRequest, QueryServerStateResponse> GetQueryServerStateCall(CancellationToken cancellationToken)
+		{
+			FakeAsyncStreamReader<QueryServerStateResponse> responseStream = new(cancellationToken);
+			FakeClientStreamWriter<QueryServerStateRequest> requestStream = new(onComplete: () =>
+			{
+				responseStream.Complete();
+				return Task.CompletedTask;
+			});
+
+			responseStream.Write(new QueryServerStateResponse { Name = _serverName, Stopping = _isStopping });
+			
+			return new (
+				requestStream,
+				responseStream,
+				Task.FromResult(new Metadata()),
+				() => Status.DefaultSuccess,
+				() => new Metadata(),
+				() => { /*isDisposed = true;*/ });
+		}
+		
+		public AsyncDuplexStreamingCall<UpdateSessionRequest, UpdateSessionResponse> GetUpdateSessionCall(CancellationToken cancellationToken)
+		{
+			FakeAsyncStreamReader<UpdateSessionResponse> responseStream = new(cancellationToken);
+			
+			async Task OnRequest(UpdateSessionRequest request)
+			{
+				UpdateSessionReceived.TrySetResult(true);
+				
+				// TODO: Handle updates of leases sent in request and verify session/agent IDs.
+				
+				_logger.LogInformation("OnUpdateSessionRequest: {AgentId} {SessionId} {Status}", request.AgentId, request.SessionId, request.Status);
+				await Task.Delay(1000, cancellationToken);
+				UpdateSessionResponse response = new () { ExpiryTime = Timestamp.FromDateTime(DateTime.UtcNow + TimeSpan.FromMinutes(120)) };
+				response.Leases.AddRange(_leases);
+				await responseStream.Write(response);
+			}
+			
+			FakeClientStreamWriter<UpdateSessionRequest> requestStream = new(OnRequest, () => {
+				responseStream.Complete();
+				return Task.CompletedTask;
+			});
+			
+			return new (
+				requestStream,
+				responseStream,
+				Task.FromResult(new Metadata()),
+				() => Status.DefaultSuccess,
+				() => new Metadata(),
+				() => { });
+		}
+		
+		public static AsyncUnaryCall<TResponse> CreateAsyncUnaryCall<TResponse>(TResponse response)
+		{
+			return new AsyncUnaryCall<TResponse>(
+				Task.FromResult(response),
+				Task.FromResult(new Metadata()),
+				() => Status.DefaultSuccess,
+				() => new Metadata(),
+				() => { });
+		}
+	}
+
+	/// <summary>
+	/// Fake stream reader used for testing gRPC clients
+	/// </summary>
+	/// <typeparam name="T">Message type reader will handle</typeparam>
+	internal class FakeAsyncStreamReader<T> : IAsyncStreamReader<T> where T : class
+	{
+		private readonly Channel<T> _channel = System.Threading.Channels.Channel.CreateUnbounded<T>();
+		private T? _current;
+		private CancellationToken? _cancellationTokenOverride;
+
+		public FakeAsyncStreamReader(CancellationToken? cancellationTokenOverride = null)
+		{
+			_cancellationTokenOverride = cancellationTokenOverride;
+		}
+
+		public Task Write(T message)
+		{
+			if (!_channel.Writer.TryWrite(message))
+			{
+				throw new InvalidOperationException("Unable to write message.");
+			}
+			
+			return Task.CompletedTask;
+		}
+
+		public void Complete()
+		{
+			_channel.Writer.Complete();
+		}
+		
+		/// <inheritdoc/>
+		public async Task<bool> MoveNext(CancellationToken cancellationToken)
+		{
+			if (_cancellationTokenOverride != null)
+			{
+				cancellationToken = _cancellationTokenOverride.Value;
+			}
+				
+			if (await _channel.Reader.WaitToReadAsync(cancellationToken))
+			{
+				if (_channel.Reader.TryRead(out T? message))
+				{
+					_current = message;
+					return true;
+				}
+			}
+
+			_current = null!;
+			return false;
+		}
+		
+		/// <inheritdoc/>
+		public T Current
+		{
+			get
+			{
+				if (_current == null)
+				{
+					throw new InvalidOperationException("No current element is available.");
+				}
+				return _current;
+			}
+		}
+	}
+	
+	/// <summary>
+	/// Fake stream writer used for testing gRPC clients
+	/// </summary>
+	/// <typeparam name="T">Message type writer will handle</typeparam>
+	internal class FakeClientStreamWriter<T> : IClientStreamWriter<T> where T : class
+	{
+		private readonly Func<T, Task>? _onWrite;
+		private readonly Func<Task>? _onComplete;
+		private bool _isCompleted;
+
+		public FakeClientStreamWriter(Func<T, Task>? onWrite = null, Func<Task>? onComplete = null)
+		{
+			_onWrite = onWrite;
+			_onComplete = onComplete;
+		}
+
+		/// <inheritdoc/>
+		public async Task WriteAsync(T message)
+		{
+			if (_isCompleted) throw new InvalidOperationException("Stream is marked as complete");
+			if (_onWrite != null) await _onWrite(message);
+		}
+
+		/// <inheritdoc/>
+		public WriteOptions? WriteOptions { get; set; }
+		
+		/// <inheritdoc/>
+		public async Task CompleteAsync()
+		{
+			_isCompleted = true;
+			if (_onComplete != null) await _onComplete();
+		}
 	}
-}
\ No newline at end of file
+}
diff --git a/Engine/Source/Programs/Horde/Horde.Agent/Horde.Agent.csproj b/Engine/Source/Programs/Horde/Horde.Agent/Horde.Agent.csproj
--- a/Engine/Source/Programs/Horde/Horde.Agent/Horde.Agent.csproj
+++ b/Engine/Source/Programs/Horde/Horde.Agent/Horde.Agent.csproj
@@ -108,32 +108,32 @@
   </ItemGroup>
 
   <ItemGroup>
-    <Protobuf Include="..\Horde.Build\Protos\horde\horde_rpc.proto" ProtoRoot="..\Horde.Build\Protos" GrpcServices="Client" Access="Internal" />
-    <Protobuf Include="..\Horde.Build\Protos\horde\horde_rpc_messages.proto" ProtoRoot="..\Horde.Build\Protos" Access="Internal" />
-    <Protobuf Include="..\Horde.Build\Protos\horde\common\agent_status.proto" ProtoRoot="..\Horde.Build\Protos" Access="Internal" />
-    <Protobuf Include="..\Horde.Build\Protos\horde\common\bucket_id.proto" ProtoRoot="..\Horde.Build\Protos" Access="Internal" />
-    <Protobuf Include="..\Horde.Build\Protos\horde\common\event_severity.proto" ProtoRoot="..\Horde.Build\Protos" Access="Internal" />
-    <Protobuf Include="..\Horde.Build\Protos\horde\common\io_hash.proto" ProtoRoot="..\Horde.Build\Protos" Access="Internal" />
-    <Protobuf Include="..\Horde.Build\Protos\horde\common\jobstep_outcome.proto" ProtoRoot="..\Horde.Build\Protos" Access="Internal" />
-    <Protobuf Include="..\Horde.Build\Protos\horde\common\jobstep_state.proto" ProtoRoot="..\Horde.Build\Protos" Access="Internal" />
-    <Protobuf Include="..\Horde.Build\Protos\horde\common\label_change.proto" ProtoRoot="..\Horde.Build\Protos" Access="Internal" />
-    <Protobuf Include="..\Horde.Build\Protos\horde\common\lease_outcome.proto" ProtoRoot="..\Horde.Build\Protos" Access="Internal" />
-    <Protobuf Include="..\Horde.Build\Protos\horde\common\lease_state.proto" ProtoRoot="..\Horde.Build\Protos" Access="Internal" />
-    <Protobuf Include="..\Horde.Build\Protos\horde\common\namespace_id.proto" ProtoRoot="..\Horde.Build\Protos" Access="Internal" />
-    <Protobuf Include="..\Horde.Build\Protos\horde\common\priority.proto" ProtoRoot="..\Horde.Build\Protos" Access="Internal" />
-    <Protobuf Include="..\Horde.Build\Protos\horde\common\qualified_ref_id.proto" ProtoRoot="..\Horde.Build\Protos" Access="Internal" />
-    <Protobuf Include="..\Horde.Build\Protos\horde\common\ref_id.proto" ProtoRoot="..\Horde.Build\Protos" Access="Internal" />
-    <Protobuf Include="..\Horde.Build\Protos\horde\messages\agent_capabilities.proto" ProtoRoot="..\Horde.Build\Protos" Access="Internal" />
-    <Protobuf Include="..\Horde.Build\Protos\horde\messages\agent_requirements.proto" ProtoRoot="..\Horde.Build\Protos" Access="Internal" />
-    <Protobuf Include="..\Horde.Build\Protos\horde\messages\agent_workspace.proto" ProtoRoot="..\Horde.Build\Protos" Access="Internal" />
-    <Protobuf Include="..\Horde.Build\Protos\horde\messages\lease.proto" ProtoRoot="..\Horde.Build\Protos" Access="Internal" />
-    <Protobuf Include="..\Horde.Build\Protos\horde\tasks\aws_lambda_compute_task.proto" ProtoRoot="..\Horde.Build\Protos" Access="Internal" />
-    <Protobuf Include="..\Horde.Build\Protos\horde\tasks\compute_task.proto" ProtoRoot="..\Horde.Build\Protos" Access="Internal" />
-    <Protobuf Include="..\Horde.Build\Protos\horde\tasks\conform_task.proto" ProtoRoot="..\Horde.Build\Protos" Access="Internal" />
-    <Protobuf Include="..\Horde.Build\Protos\horde\tasks\job_task.proto" ProtoRoot="..\Horde.Build\Protos" Access="Internal" />
-    <Protobuf Include="..\Horde.Build\Protos\horde\tasks\restart_task.proto" ProtoRoot="..\Horde.Build\Protos" Access="Internal" />
-    <Protobuf Include="..\Horde.Build\Protos\horde\tasks\shutdown_task.proto" ProtoRoot="..\Horde.Build\Protos" Access="Internal" />
-    <Protobuf Include="..\Horde.Build\Protos\horde\tasks\upgrade_task.proto" ProtoRoot="..\Horde.Build\Protos" Access="Internal" />
+    <Protobuf Include="..\Horde.Build\Protos\horde\horde_rpc.proto" ProtoRoot="..\Horde.Build\Protos" GrpcServices="Client" Access="Public" />
+    <Protobuf Include="..\Horde.Build\Protos\horde\horde_rpc_messages.proto" ProtoRoot="..\Horde.Build\Protos" Access="Public" />
+    <Protobuf Include="..\Horde.Build\Protos\horde\common\agent_status.proto" ProtoRoot="..\Horde.Build\Protos" Access="Public" />
+    <Protobuf Include="..\Horde.Build\Protos\horde\common\bucket_id.proto" ProtoRoot="..\Horde.Build\Protos" Access="Public" />
+    <Protobuf Include="..\Horde.Build\Protos\horde\common\event_severity.proto" ProtoRoot="..\Horde.Build\Protos" Access="Public" />
+    <Protobuf Include="..\Horde.Build\Protos\horde\common\io_hash.proto" ProtoRoot="..\Horde.Build\Protos" Access="Public" />
+    <Protobuf Include="..\Horde.Build\Protos\horde\common\jobstep_outcome.proto" ProtoRoot="..\Horde.Build\Protos" Access="Public" />
+    <Protobuf Include="..\Horde.Build\Protos\horde\common\jobstep_state.proto" ProtoRoot="..\Horde.Build\Protos" Access="Public" />
+    <Protobuf Include="..\Horde.Build\Protos\horde\common\label_change.proto" ProtoRoot="..\Horde.Build\Protos" Access="Public" />
+    <Protobuf Include="..\Horde.Build\Protos\horde\common\lease_outcome.proto" ProtoRoot="..\Horde.Build\Protos" Access="Public" />
+    <Protobuf Include="..\Horde.Build\Protos\horde\common\lease_state.proto" ProtoRoot="..\Horde.Build\Protos" Access="Public" />
+    <Protobuf Include="..\Horde.Build\Protos\horde\common\namespace_id.proto" ProtoRoot="..\Horde.Build\Protos" Access="Public" />
+    <Protobuf Include="..\Horde.Build\Protos\horde\common\priority.proto" ProtoRoot="..\Horde.Build\Protos" Access="Public" />
+    <Protobuf Include="..\Horde.Build\Protos\horde\common\qualified_ref_id.proto" ProtoRoot="..\Horde.Build\Protos" Access="Public" />
+    <Protobuf Include="..\Horde.Build\Protos\horde\common\ref_id.proto" ProtoRoot="..\Horde.Build\Protos" Access="Public" />
+    <Protobuf Include="..\Horde.Build\Protos\horde\messages\agent_capabilities.proto" ProtoRoot="..\Horde.Build\Protos" Access="Public" />
+    <Protobuf Include="..\Horde.Build\Protos\horde\messages\agent_requirements.proto" ProtoRoot="..\Horde.Build\Protos" Access="Public" />
+    <Protobuf Include="..\Horde.Build\Protos\horde\messages\agent_workspace.proto" ProtoRoot="..\Horde.Build\Protos" Access="Public" />
+    <Protobuf Include="..\Horde.Build\Protos\horde\messages\lease.proto" ProtoRoot="..\Horde.Build\Protos" Access="Public" />
+    <Protobuf Include="..\Horde.Build\Protos\horde\tasks\aws_lambda_compute_task.proto" ProtoRoot="..\Horde.Build\Protos" Access="Public" />
+    <Protobuf Include="..\Horde.Build\Protos\horde\tasks\compute_task.proto" ProtoRoot="..\Horde.Build\Protos" Access="Public" />
+    <Protobuf Include="..\Horde.Build\Protos\horde\tasks\conform_task.proto" ProtoRoot="..\Horde.Build\Protos" Access="Public" />
+    <Protobuf Include="..\Horde.Build\Protos\horde\tasks\job_task.proto" ProtoRoot="..\Horde.Build\Protos" Access="Public" />
+    <Protobuf Include="..\Horde.Build\Protos\horde\tasks\restart_task.proto" ProtoRoot="..\Horde.Build\Protos" Access="Public" />
+    <Protobuf Include="..\Horde.Build\Protos\horde\tasks\shutdown_task.proto" ProtoRoot="..\Horde.Build\Protos" Access="Public" />
+    <Protobuf Include="..\Horde.Build\Protos\horde\tasks\upgrade_task.proto" ProtoRoot="..\Horde.Build\Protos" Access="Public" />
   </ItemGroup>
 
   <ItemGroup>
diff --git a/Engine/Source/Programs/Horde/Horde.Agent/Services/WorkerService.cs b/Engine/Source/Programs/Horde/Horde.Agent/Services/WorkerService.cs
--- a/Engine/Source/Programs/Horde/Horde.Agent/Services/WorkerService.cs
+++ b/Engine/Source/Programs/Horde/Horde.Agent/Services/WorkerService.cs
@@ -209,15 +209,27 @@ namespace Horde.Agent.Services
 		int _updateSessionFailures;
 
 		/// <summary>
-		/// Function for creating a new executor. Primarily done this way to aid testing. 
+		/// Function for creating a new executor. Primarily to aid testing. 
 		/// </summary>
 		private readonly Func<IRpcConnection, ExecuteJobTask, BeginBatchResponse, IExecutor> _createExecutor;
 
+		/// <summary>
+		/// Function for creating a new HordeRpcClient. Primarily to aid testing. 
+		/// </summary>
+		private readonly Func<GrpcChannel, HordeRpc.HordeRpcClient> _createHordeRpcClient;
+
 		/// <summary>
 		/// How often to poll the server checking if a step has been aborted
-		/// Exposed as internal to ease testing. 
+		/// Exposed as internal to ease testing.
 		/// </summary>
 		internal TimeSpan _stepAbortPollInterval = TimeSpan.FromSeconds(5);
+		
+		
+		/// <summary>
+		/// How long to wait before trying to reacquire a new connection
+		/// Exposed as internal to ease testing. Using a lower delay can speed up tests. 
+		/// </summary>
+		internal TimeSpan _rpcConnectionRetryDelay = TimeSpan.FromSeconds(5);
 
 		/// <summary>
 		/// Constructor. Registers with the server and starts accepting connections.
@@ -226,14 +238,18 @@ namespace Horde.Agent.Services
 		/// <param name="options">The current settings</param>
 		/// <param name="grpcService">Instance of the Grpc service</param>
 		/// <param name="storageClient">Instance of the storage client</param>
-		/// <param name="createExecutor"></param>
-		public WorkerService(ILogger<WorkerService> logger, IOptions<AgentSettings> options, GrpcService grpcService, IStorageClient storageClient, Func<IRpcConnection, ExecuteJobTask, BeginBatchResponse, IExecutor>? createExecutor = null)
+		/// <param name="createExecutor">Optional factory method for creating an executor</param>
+		/// <param name="createHordeRpcClient">Optional factory method for creating a HordeRpcClient</param>
+		public WorkerService(ILogger<WorkerService> logger, IOptions<AgentSettings> options, GrpcService grpcService, IStorageClient storageClient,
+			Func<IRpcConnection, ExecuteJobTask, BeginBatchResponse, IExecutor>? createExecutor = null,
+			Func<GrpcChannel, HordeRpc.HordeRpcClient>? createHordeRpcClient = null)
 		{
 			_logger = logger;
 			_settings = options.Value;
 			_serverProfile = _settings.GetCurrentServerProfile();
 			_grpcService = grpcService;
 			_storageClient = storageClient;
+			_createHordeRpcClient = createHordeRpcClient ?? (channel => new HordeRpc.HordeRpcClient(channel));
 
 			if (_settings.WorkingDir == null)
 			{
@@ -290,7 +306,7 @@ namespace Horde.Agent.Services
 		/// Background task to cycle access tokens and update the state of the agent with the server.
 		/// </summary>
 		/// <param name="stoppingToken">Indicates that the service is trying to stop</param>
-		async Task ExecuteInnerAsync(CancellationToken stoppingToken)
+		internal async Task ExecuteInnerAsync(CancellationToken stoppingToken)
 		{
 			// Print the server info
 			_logger.LogInformation("Server: {Server}", _serverProfile.Url);
@@ -424,7 +440,7 @@ namespace Horde.Agent.Services
 		/// </summary>
 		/// <param name="stoppingToken">Indicates that the service is trying to stop</param>
 		/// <returns>Async task</returns>
-		async Task HandleSessionAsync(CancellationToken stoppingToken)
+		internal async Task HandleSessionAsync(CancellationToken stoppingToken)
 		{
 			// Make sure there's only one instance of the agent running
 			using Mutex singleInstanceMutex = new Mutex(false, "Global\\HordeAgent-DB828ACB-0AA5-4D32-A62A-21D4429B1014");
@@ -469,7 +485,7 @@ namespace Horde.Agent.Services
 			CreateSessionResponse createSessionResponse;
 			using (GrpcChannel channel = _grpcService.CreateGrpcChannel(_serverProfile.Token))
 			{
-				HordeRpc.HordeRpcClient rpcClient = new HordeRpc.HordeRpcClient(channel);
+				HordeRpc.HordeRpcClient rpcClient = _createHordeRpcClient(channel);
 
 				// Create the session information
 				CreateSessionRequest sessionRequest = new CreateSessionRequest();
@@ -483,8 +499,10 @@ namespace Horde.Agent.Services
 				_logger.LogInformation("Session started ({SessionId})", createSessionResponse.SessionId);
 			}
 
+			Func<GrpcChannel> createGrpcChannel = () => _grpcService.CreateGrpcChannel(createSessionResponse.Token);
+			
 			// Open a connection to the server
-			await using (IRpcConnection rpcCon = RpcConnection.Create(() => _grpcService.CreateGrpcChannel(createSessionResponse.Token), _logger))
+			await using (IRpcConnection rpcCon = RpcConnection.Create(createGrpcChannel, _createHordeRpcClient, _logger))
 			{
 				// Track how many updates we get in 10 seconds. We'll start rate limiting this if it looks like we've got a problem that's causing us to spam the server.
 				Stopwatch updateTimer = Stopwatch.StartNew();
@@ -554,7 +572,8 @@ namespace Horde.Agent.Services
 						{
 							if (rpcClientRef == null)
 							{
-								await Task.WhenAny(Task.Delay(TimeSpan.FromSeconds(5.0), stoppingToken), waitTask);
+								// An RpcConnection has not yet been established, wait a period of time and try again
+								await Task.WhenAny(Task.Delay(_rpcConnectionRetryDelay, stoppingToken), waitTask);
 							}
 							else
 							{
@@ -1459,7 +1478,7 @@ namespace Horde.Agent.Services
 		/// </summary>
 		/// <param name="logger">Logging device</param>
 		/// <returns>Outcome of this operation</returns>
-		Task<LeaseResult> ShutdownAsync(ILogger logger)
+		internal Task<LeaseResult> ShutdownAsync(ILogger logger)
 		{
 			logger.LogInformation("Setting shutdown flag");
 			_requestShutdown = true;
diff --git a/Engine/Source/Programs/Horde/Horde.Agent/Utilities/HordeRpc.cs b/Engine/Source/Programs/Horde/Horde.Agent/Utilities/HordeRpc.cs
--- a/Engine/Source/Programs/Horde/Horde.Agent/Utilities/HordeRpc.cs
+++ b/Engine/Source/Programs/Horde/Horde.Agent/Utilities/HordeRpc.cs
@@ -6,6 +6,10 @@ using System.Linq;
 using Google.Protobuf;
 using Google.Protobuf.Collections;
 
+// These partial classes/extensions operate on generated gRPC and Protobuf code.
+// Warnings below are disabled to avoid documenting every class touched.
+#pragma warning disable CS1591
+#pragma warning disable CA1716
 namespace HordeCommon.Rpc
 {
 	partial class Property
@@ -181,3 +185,5 @@ namespace HordeCommon.Rpc
 		}
 	}
 }
+#pragma warning restore CA1716
+#pragma warning restore CS1591
\ No newline at end of file
diff --git a/Engine/Source/Programs/Horde/Horde.Agent/Utility/RpcConnection.cs b/Engine/Source/Programs/Horde/Horde.Agent/Utility/RpcConnection.cs
--- a/Engine/Source/Programs/Horde/Horde.Agent/Utility/RpcConnection.cs
+++ b/Engine/Source/Programs/Horde/Horde.Agent/Utility/RpcConnection.cs
@@ -342,6 +342,7 @@ namespace Horde.Agent.Utility
 		};
 
 		private readonly Func<GrpcChannel> _createGrpcChannel;
+		private readonly Func<GrpcChannel, HordeRpc.HordeRpcClient> _createHordeRpcClient;
 		private readonly TaskCompletionSource<bool> _stoppingTaskSource = new TaskCompletionSource<bool>();
 		private TaskCompletionSource<RpcSubConnection> _subConnectionTaskSource = new TaskCompletionSource<RpcSubConnection>();
 		private Task? _backgroundTask;
@@ -351,10 +352,12 @@ namespace Horde.Agent.Utility
 		/// Constructor
 		/// </summary>
 		/// <param name="createGrpcChannel">Factory method for creating new GRPC channels</param>
+		/// <param name="createHordeRpcClient">Factory method for creating new Horde.Rpc clients</param>
 		/// <param name="logger">Logger instance</param>
-		public RpcConnection(Func<GrpcChannel> createGrpcChannel, ILogger logger)
+		public RpcConnection(Func<GrpcChannel> createGrpcChannel, Func<GrpcChannel, HordeRpc.HordeRpcClient> createHordeRpcClient, ILogger logger)
 		{
 			_createGrpcChannel = createGrpcChannel;
+			_createHordeRpcClient = createHordeRpcClient;
 			_logger = logger;
 
 			_backgroundTask = Task.Run(() => ExecuteAsync());
@@ -624,7 +627,7 @@ namespace Horde.Agent.Utility
 		{
 			using (GrpcChannel channel = _createGrpcChannel())
 			{
-				HordeRpc.HordeRpcClient client = new HordeRpc.HordeRpcClient(channel);//.Intercept(new ClientTracingInterceptor(GlobalTracer.Instance)));
+				HordeRpc.HordeRpcClient client = _createHordeRpcClient(channel);
 
 				RpcSubConnection? subConnection = null;
 				try
@@ -737,9 +740,9 @@ namespace Horde.Agent.Utility
 			}
 		}
 
-		public static IRpcConnection Create(Func<GrpcChannel> createGrpcChannel, ILogger logger)
+		public static IRpcConnection Create(Func<GrpcChannel> createGrpcChannel, Func<GrpcChannel, HordeRpc.HordeRpcClient> createHordeRpcClient, ILogger logger)
 		{
-			return new RpcConnection(createGrpcChannel, logger);
+			return new RpcConnection(createGrpcChannel, createHordeRpcClient, logger);
 		}
 	}
 }
diff --git a/Engine/Source/Programs/Horde/Horde.Build.Tests/Horde.Build.Tests.csproj b/Engine/Source/Programs/Horde/Horde.Build.Tests/Horde.Build.Tests.csproj
--- a/Engine/Source/Programs/Horde/Horde.Build.Tests/Horde.Build.Tests.csproj
+++ b/Engine/Source/Programs/Horde/Horde.Build.Tests/Horde.Build.Tests.csproj
@@ -29,7 +29,9 @@
 
   <ItemGroup>
     <ProjectReference Include="..\..\Shared\EpicGames.Serialization\EpicGames.Serialization.csproj"/>
-    <ProjectReference Include="..\Horde.Agent\Horde.Agent.csproj"/>
+    <ProjectReference Include="..\Horde.Agent\Horde.Agent.csproj">
+      <Aliases>HordeAgent</Aliases>
+    </ProjectReference>
     <ProjectReference Include="..\Horde.Build\Horde.Build.csproj"/>
   </ItemGroup>
 
diff --git a/Engine/Source/Programs/Horde/Horde.Build.Tests/IssueServiceTests.cs b/Engine/Source/Programs/Horde/Horde.Build.Tests/IssueServiceTests.cs
--- a/Engine/Source/Programs/Horde/Horde.Build.Tests/IssueServiceTests.cs
+++ b/Engine/Source/Programs/Horde/Horde.Build.Tests/IssueServiceTests.cs
@@ -1,5 +1,7 @@
 // Copyright Epic Games, Inc. All Rights Reserved.
 
+extern alias HordeAgent;
+
 using System;
 using System.Collections.Generic;
 using System.Linq;
@@ -7,8 +9,8 @@ using System.Runtime.InteropServices;
 using System.Text.Json;
 using System.Threading.Tasks;
 using EpicGames.Core;
-using Horde.Agent.Parser;
-using Horde.Agent.Utility;
+using HordeAgent.Horde.Agent.Parser;
+using HordeAgent.Horde.Agent.Utility;
 using Horde.Build.Jobs.Graphs;
 using Horde.Build.Jobs;
 using Horde.Build.Issues;
