Missing file.

diff --git a/Engine/Source/Programs/Horde/Horde.Agent.Tests/LogParserTests.cs b/Engine/Source/Programs/Horde/Horde.Agent.Tests/LogParserTests.cs
--- a/Engine/Source/Programs/Horde/Horde.Agent.Tests/LogParserTests.cs
+++ b/Engine/Source/Programs/Horde/Horde.Agent.Tests/LogParserTests.cs
@@ -720,19 +720,39 @@ namespace Horde.Agent.Tests
 		{
 			string[] lines =
 			{
-				@"[2022.05.31-04.25.40:235][  0]LogGatherTextFromSourceCommandlet: Warning: Plugins/Enterprise/VariantManager/Source/VariantManager/Private/SVariantManager.cpp(3717): LOCTEXT macro has an empty identifier and cannot be gathered."
+				@"[2022.05.31-04.25.40:235][  0]LogGatherTextFromSourceCommandlet: Warning: Engine/Plugins/Enterprise/VariantManager/Source/VariantManager/Private/SVariantManager.cpp(3717): LOCTEXT macro has an empty identifier and cannot be gathered.",
+				@"[2022.06.01-04.29.09:630][  0]LogLocTextHelper: Warning: Engine/Plugins/Experimental/UVEditor/Source/UVEditor/Private/UVEditorCommands.cpp(39): Text conflict from UI_COMMAND macro for namespace ""UICommands.FUVEditorCommands"" and key ""SplitAction_ToolTip"". The conflicting sources are ""Given an edge selection, split those edges. Given a vertex selection, split any selected bowtie vertices. Given a triangle selection, split along selection boundaries."" and ""Given an edge selection, split those edges. Given a vertex selection, split any selected bowtie vertices."".",
+				@"[2022.06.01-04.29.09:630][  0]LogLocTextHelper: Warning: Engine/Content/Localization/Engine/Engine.manifest: See conflicting location.",
 			};
 
 			{
 				List<LogEvent> logEvents = Parse(lines);
-				Assert.AreEqual(1, logEvents.Count);
-				CheckEventGroup(logEvents, 0, 1, LogLevel.Warning, KnownLogEvents.Engine_Localization);
+				Assert.AreEqual(3, logEvents.Count);
+				CheckEventGroup(logEvents.Slice(0, 1), 0, 1, LogLevel.Warning, KnownLogEvents.Engine_Localization);
+				CheckEventGroup(logEvents.Slice(1, 2), 1, 2, LogLevel.Warning, KnownLogEvents.Engine_Localization);
 
 				LogEvent logEvent = logEvents[0];
 				Assert.AreEqual("LogGatherTextFromSourceCommandlet", logEvent.GetProperty("channel").ToString());
 				Assert.AreEqual("Warning", logEvent.GetProperty("severity").ToString());
 				Assert.AreEqual("Engine/Plugins/Enterprise/VariantManager/Source/VariantManager/Private/SVariantManager.cpp", logEvent.GetProperty<LogValue>("file")!.Properties!["relativePath"].ToString());
 				Assert.AreEqual("3717", logEvent.GetProperty("line").ToString());
+				Assert.AreEqual(0, logEvent.LineIndex);
+				Assert.AreEqual(1, logEvent.LineCount);
+
+				logEvent = logEvents[1];
+				Assert.AreEqual("LogLocTextHelper", logEvent.GetProperty("channel").ToString());
+				Assert.AreEqual("Warning", logEvent.GetProperty("severity").ToString());
+				Assert.AreEqual("Engine/Plugins/Experimental/UVEditor/Source/UVEditor/Private/UVEditorCommands.cpp", logEvent.GetProperty<LogValue>("file")!.Properties!["relativePath"].ToString());
+				Assert.AreEqual("39", logEvent.GetProperty("line").ToString());
+				Assert.AreEqual(0, logEvent.LineIndex);
+				Assert.AreEqual(2, logEvent.LineCount);
+
+				logEvent = logEvents[2];
+				Assert.AreEqual("LogLocTextHelper", logEvent.GetProperty("channel").ToString());
+				Assert.AreEqual("Warning", logEvent.GetProperty("severity").ToString());
+				Assert.AreEqual("Engine/Content/Localization/Engine/Engine.manifest", logEvent.GetProperty<LogValue>("file")!.Properties!["relativePath"].ToString());
+				Assert.AreEqual(1, logEvent.LineIndex);
+				Assert.AreEqual(2, logEvent.LineCount);
 			}
 		}
 
@@ -751,19 +771,26 @@ namespace Horde.Agent.Tests
 			LogBuffer buffer = new LogBuffer(10);
 			buffer.AddLines(lines);
 
-			Assert.AreEqual(buffer.Hanging()[0], "first line");
-			Assert.AreEqual(buffer.Hanging()[1], null);
-			buffer.MoveNext();
-
-			Assert.AreEqual(buffer.Hanging()[0], "first line in multi-line message");
-			Assert.AreEqual(buffer.Hanging()[1], "  this is a hanging indent");
-			Assert.AreEqual(buffer.Hanging()[2], "   this is also hanging");
-			Assert.AreEqual(buffer.Hanging()[3], null);
-			buffer.Advance(3);
-
-			Assert.AreEqual(buffer.Hanging()[0], "this is a separate item");
-			Assert.AreEqual(buffer.Hanging()[1], null);
-			buffer.MoveNext();
+			Assert.IsTrue(buffer.IsAligned(0, buffer.CurrentLine));
+			Assert.IsTrue(buffer.IsAligned(1, buffer.CurrentLine));
+			Assert.IsTrue(buffer.IsAligned(2, buffer.CurrentLine));
+			Assert.IsTrue(buffer.IsAligned(3, buffer.CurrentLine));
+			Assert.IsTrue(buffer.IsAligned(4, buffer.CurrentLine));
+			Assert.IsFalse(buffer.IsAligned(5, buffer.CurrentLine));
+
+			string? firstIndentedLine = buffer[2];
+			Assert.IsTrue(buffer.IsAligned(2, firstIndentedLine));
+			Assert.IsTrue(buffer.IsAligned(3, firstIndentedLine));
+			Assert.IsFalse(buffer.IsAligned(4, firstIndentedLine));
+
+			Assert.IsFalse(buffer.IsHanging(1, buffer.CurrentLine));
+			Assert.IsTrue(buffer.IsHanging(2, buffer.CurrentLine));
+			Assert.IsTrue(buffer.IsHanging(3, buffer.CurrentLine));
+			Assert.IsFalse(buffer.IsHanging(4, buffer.CurrentLine));
+
+			Assert.IsFalse(buffer.IsHanging(2, firstIndentedLine));
+			Assert.IsTrue(buffer.IsHanging(3, firstIndentedLine));
+			Assert.IsFalse(buffer.IsHanging(4, firstIndentedLine));
 		}
 
 		static List<LogEvent> Parse(IEnumerable<string> lines)
