Horde: Allow issue matchers to scan entire event list before setting fingerprints.

This allows handlers to determine which events to ignore after a positive match, and also grouping together events that may otherwise be reported separately.

diff --git a/Engine/Source/Programs/Horde/Horde.Build.Tests/IssueServiceTests.cs b/Engine/Source/Programs/Horde/Horde.Build.Tests/IssueServiceTests.cs
--- a/Engine/Source/Programs/Horde/Horde.Build.Tests/IssueServiceTests.cs
+++ b/Engine/Source/Programs/Horde/Horde.Build.Tests/IssueServiceTests.cs
@@ -1027,7 +1027,7 @@ namespace Horde.Build.Tests
 		}
 
 		[TestMethod]
-		public async Task LinkerIssueTest()
+		public async Task SymbolIssueTest3()
 		{
 			// #1
 			// Scenario: Job step completes successfully at CL 105
@@ -1059,9 +1059,9 @@ namespace Horde.Build.Tests
 
 				IIssue issue = issues[0];
 				List<IIssueSpan> spans = await IssueCollection.FindSpansAsync(issue.Id);
-				Assert.AreEqual(spans.Count, 1);
-				Assert.AreEqual(issue.Fingerprints.Count, 1);
-				Assert.AreEqual(issue.Fingerprints[0].Type, "Symbol");
+				Assert.AreEqual(1, spans.Count);
+				Assert.AreEqual(1, issue.Fingerprints.Count);
+				Assert.AreEqual("Symbol", issue.Fingerprints[0].Type);
 
 				IIssueSpan span = spans[0];
 				Assert.AreEqual(105, span.LastSuccess?.Change);
@@ -1086,9 +1086,9 @@ namespace Horde.Build.Tests
 
 				IIssue issue = issues[0];
 				List<IIssueSpan> spans = await IssueCollection.FindSpansAsync(issue.Id);
-				Assert.AreEqual(spans.Count, 1);
-				Assert.AreEqual(issue.Fingerprints.Count, 1);
-				Assert.AreEqual(issue.Fingerprints[0].Type, "Symbol");
+				Assert.AreEqual(1, spans.Count);
+				Assert.AreEqual(1, issue.Fingerprints.Count);
+				Assert.AreEqual("Default", issue.Fingerprints[0].Type);
 
 				IIssueSpan span = spans[0];
 				Assert.AreEqual(120, span.LastSuccess?.Change);
@@ -1113,7 +1113,7 @@ namespace Horde.Build.Tests
 			Assert.AreEqual(1, issues.Count);
 
 			IIssue issue = issues[0];
-			Assert.AreEqual("Linker errors in Update Version Files", issue.Summary);
+			Assert.AreEqual("Errors in Update Version Files", issue.Summary);
 		}
 
 		[TestMethod]
diff --git a/Engine/Source/Programs/Horde/Horde.Build/Issues/Handlers/CompileIssueHandler.cs b/Engine/Source/Programs/Horde/Horde.Build/Issues/Handlers/CompileIssueHandler.cs
--- a/Engine/Source/Programs/Horde/Horde.Build/Issues/Handlers/CompileIssueHandler.cs
+++ b/Engine/Source/Programs/Horde/Horde.Build/Issues/Handlers/CompileIssueHandler.cs
@@ -40,28 +40,27 @@ namespace Horde.Build.Issues.Handlers
 		}
 
 		/// <inheritdoc/>
-		public override bool TryGetFingerprint(IJob job, INode node, IReadOnlyNodeAnnotations annotations, ILogEventData eventData, [NotNullWhen(true)] out NewIssueFingerprint? fingerprint)
+		public override void TagEvents(IJob job, INode node, IReadOnlyNodeAnnotations annotations, IReadOnlyList<IssueEvent> stepEvents)
 		{
-			if (!IsMatchingEventId(eventData.EventId))
+			foreach (IssueEvent stepEvent in stepEvents)
 			{
-				fingerprint = null;
-				return false;
-			}
+				if (stepEvent.EventId.HasValue && IsMatchingEventId(stepEvent.EventId))
+				{
+					List<string> newFileNames = new List<string>();
+					GetSourceFiles(stepEvent.EventData, newFileNames);
 
-			List<string> newFileNames = new List<string>();
-			GetSourceFiles(eventData, newFileNames);
+					string? compileType;
+					if (!annotations.TryGetValue(CompileTypeAnnotation, out compileType))
+					{
+						compileType = "Compile";
+					}
 
-			string? compileType;
-			if (!annotations.TryGetValue(CompileTypeAnnotation, out compileType))
-			{
-				compileType = "Compile";
-			}
+					List<string> newMetadata = new List<string>();
+					newMetadata.Add($"{CompileTypeAnnotation}={compileType}");
 
-			List<string> newMetadata = new List<string>();
-			newMetadata.Add($"{CompileTypeAnnotation}={compileType}");
-
-			fingerprint = new NewIssueFingerprint(Type, newFileNames, null, newMetadata);
-			return true;
+					stepEvent.Fingerprint = new NewIssueFingerprint(Type, newFileNames, null, newMetadata);
+				}
+			}
 		}
 
 		/// <inheritdoc/>
diff --git a/Engine/Source/Programs/Horde/Horde.Build/Issues/Handlers/ContentIssueHandler.cs b/Engine/Source/Programs/Horde/Horde.Build/Issues/Handlers/ContentIssueHandler.cs
--- a/Engine/Source/Programs/Horde/Horde.Build/Issues/Handlers/ContentIssueHandler.cs
+++ b/Engine/Source/Programs/Horde/Horde.Build/Issues/Handlers/ContentIssueHandler.cs
@@ -17,20 +17,20 @@ namespace Horde.Build.Issues.Handlers
 	/// <summary>
 	/// Instance of a particular compile error
 	/// </summary>
-	class ContentIssueHandler : IIssueHandler
+	class ContentIssueHandler : IssueHandler
 	{
 		/// <inheritdoc/>
-		public string Type => "Content";
+		public override string Type => "Content";
 
 		/// <inheritdoc/>
-		public int Priority => 10;
+		public override int Priority => 10;
 
 		/// <summary>
 		/// Determines if the given event id matches
 		/// </summary>
 		/// <param name="eventId">The event id to compare</param>
 		/// <returns>True if the given event id matches</returns>
-		public static bool IsMatchingEventId(EventId? eventId)
+		public static bool IsMatchingEventId(EventId eventId)
 		{
 			return eventId == KnownLogEvents.Engine_AssetLog;
 		}
@@ -55,22 +55,22 @@ namespace Horde.Build.Issues.Handlers
 		}
 
 		/// <inheritdoc/>
-		public bool TryGetFingerprint(IJob job, INode node, IReadOnlyNodeAnnotations annotations, ILogEventData eventData, [NotNullWhen(true)] out NewIssueFingerprint? fingerprint)
+		public override void TagEvents(IJob job, INode node, IReadOnlyNodeAnnotations annotations, IReadOnlyList<IssueEvent> stepEvents)
 		{
-			if (!IsMatchingEventId(eventData.EventId))
+			foreach (IssueEvent stepEvent in stepEvents)
 			{
-				fingerprint = null;
-				return false;
-			}
+				if (stepEvent.EventId != null && IsMatchingEventId(stepEvent.EventId.Value))
+				{
+					HashSet<string> newAssetNames = new HashSet<string>();
+					GetAssetNames(stepEvent.EventData, newAssetNames);
 
-			HashSet<string> newAssetNames = new HashSet<string>();
-			GetAssetNames(eventData, newAssetNames);
-			fingerprint = new NewIssueFingerprint(Type, newAssetNames, null, null);
-			return true;
+					stepEvent.Fingerprint = new NewIssueFingerprint(Type, newAssetNames, null, null);
+				}
+			}
 		}
 
 		/// <inheritdoc/>
-		public string GetSummary(IIssueFingerprint fingerprint, IssueSeverity severity)
+		public override string GetSummary(IIssueFingerprint fingerprint, IssueSeverity severity)
 		{
 			string type = (severity == IssueSeverity.Warning) ? "Warnings" : "Errors";
 			string list = StringUtils.FormatList(fingerprint.Keys.ToArray(), 2);
@@ -78,7 +78,7 @@ namespace Horde.Build.Issues.Handlers
 		}
 
 		/// <inheritdoc/>
-		public void RankSuspects(IIssueFingerprint fingerprint, List<SuspectChange> suspects)
+		public override void RankSuspects(IIssueFingerprint fingerprint, List<SuspectChange> suspects)
 		{
 			foreach (SuspectChange suspect in suspects)
 			{
diff --git a/Engine/Source/Programs/Horde/Horde.Build/Issues/Handlers/CopyrightIssueHandler.cs b/Engine/Source/Programs/Horde/Horde.Build/Issues/Handlers/CopyrightIssueHandler.cs
--- a/Engine/Source/Programs/Horde/Horde.Build/Issues/Handlers/CopyrightIssueHandler.cs
+++ b/Engine/Source/Programs/Horde/Horde.Build/Issues/Handlers/CopyrightIssueHandler.cs
@@ -17,84 +17,40 @@ namespace Horde.Build.Issues.Handlers
 	/// <summary>
 	/// Instance of a particular compile error
 	/// </summary>
-	class CopyrightIssueHandler : IIssueHandler
+	class CopyrightIssueHandler : SourceFileIssueHandler
 	{
 		/// <inheritdoc/>
-		public string Type => "Copyright";
+		public override string Type => "Copyright";
 
 		/// <inheritdoc/>
-		public int Priority => 10;
-
-		/// <summary>
-		/// Filenames containing errors or warnings
-		/// </summary>
-		public HashSet<string> FileNames { get; set; } = new HashSet<string>();
+		public override int Priority => 10;
 
 		/// <summary>
 		/// Determines if the given event id matches
 		/// </summary>
 		/// <param name="eventId">The event id to compare</param>
 		/// <returns>True if the given event id matches</returns>
-		public static bool IsMatchingEventId(EventId? eventId)
+		public static bool IsMatchingEventId(EventId eventId)
 		{
 			return eventId == KnownLogEvents.AutomationTool_MissingCopyright;
 		}
 
-		/// <summary>
-		/// Extracts a list of source files from an event
-		/// </summary>
-		/// <param name="logEventData">The event data</param>
-		/// <param name="sourceFiles">List of source files</param>
-		public static void GetSourceFiles(ILogEventData logEventData, HashSet<string> sourceFiles)
-		{
-			foreach (ILogEventLine line in logEventData.Lines)
-			{
-				string? relativePath;
-				if (line.Data.TryGetNestedProperty("properties.file.relativePath", out relativePath) || line.Data.TryGetNestedProperty("properties.file", out relativePath))
-				{
-					int endIdx = relativePath.LastIndexOfAny(new char[] { '/', '\\' }) + 1;
-					string fileName = relativePath.Substring(endIdx);
-					sourceFiles.Add(fileName);
-				}
-			}
-		}
-
-		/// <inheritdoc/>
-		public bool TryGetFingerprint(IJob job, INode node, IReadOnlyNodeAnnotations annotations, ILogEventData eventData, [NotNullWhen(true)] out NewIssueFingerprint? fingerprint)
-		{
-			if (!IsMatchingEventId(eventData.EventId))
-			{
-				fingerprint = null;
-				return false;
-			}
-
-			HashSet<string> newFileNames = new HashSet<string>(StringComparer.OrdinalIgnoreCase);
-			GetSourceFiles(eventData, newFileNames);
-			fingerprint = new NewIssueFingerprint(Type, newFileNames, null, null);
-			return true;
-		}
-
 		/// <inheritdoc/>
-		public void RankSuspects(IIssueFingerprint fingerprint, List<SuspectChange> suspects)
+		public override void TagEvents(IJob job, INode node, IReadOnlyNodeAnnotations annotations, IReadOnlyList<IssueEvent> stepEvents)
 		{
-			foreach (SuspectChange suspect in suspects)
+			foreach (IssueEvent stepEvent in stepEvents)
 			{
-				if (suspect.ContainsCode)
+				if (stepEvent.EventId != null && IsMatchingEventId(stepEvent.EventId.Value))
 				{
-					if (FileNames.Any(x => suspect.ModifiesFile(x)))
-					{
-						suspect.Rank += 20;
-					}
-					else
-					{
-						suspect.Rank += 10;
-					}
+					List<string> newFileNames = new List<string>();
+					GetSourceFiles(stepEvent.EventData, newFileNames);
+					stepEvent.Fingerprint = new NewIssueFingerprint(Type, newFileNames, null, null);
 				}
 			}
 		}
 
 		/// <inheritdoc/>
-		public string GetSummary(IIssueFingerprint fingerprint, IssueSeverity severity)
+		public override string GetSummary(IIssueFingerprint fingerprint, IssueSeverity severity)
 		{
 			return $"Missing copyright notice in {StringUtils.FormatList(fingerprint.Keys.ToArray(), 2)}";
 		}
diff --git a/Engine/Source/Programs/Horde/Horde.Build/Issues/Handlers/DefaultIssueHandler.cs b/Engine/Source/Programs/Horde/Horde.Build/Issues/Handlers/DefaultIssueHandler.cs
--- a/Engine/Source/Programs/Horde/Horde.Build/Issues/Handlers/DefaultIssueHandler.cs
+++ b/Engine/Source/Programs/Horde/Horde.Build/Issues/Handlers/DefaultIssueHandler.cs
@@ -1,44 +1,45 @@
 // Copyright Epic Games, Inc. All Rights Reserved.
 
 using System.Collections.Generic;
-using System.Diagnostics.CodeAnalysis;
 using System.Linq;
 using Horde.Build.Jobs;
 using Horde.Build.Jobs.Graphs;
-using Horde.Build.Logs;
 
 namespace Horde.Build.Issues.Handlers
 {
 	/// <summary>
 	/// Instance of a particular compile error
 	/// </summary>
-	class DefaultIssueHandler : IIssueHandler
+	class DefaultIssueHandler : IssueHandler
 	{
 		/// <summary>
 		/// Name of the handler
 		/// </summary>
-		public const string Type = "Default";
+		public const string TypeConst = "Default";
 
 		/// <inheritdoc/>
-		string IIssueHandler.Type => Type;
+		public override string Type => TypeConst;
 
 		/// <inheritdoc/>
-		public int Priority => 0;
+		public override int Priority => 0;
 
 		/// <inheritdoc/>
-		public bool TryGetFingerprint(IJob job, INode node, IReadOnlyNodeAnnotations annotations, ILogEventData eventData, [NotNullWhen(true)] out NewIssueFingerprint? fingerprint)
+		public override void TagEvents(IJob job, INode node, IReadOnlyNodeAnnotations annotations, IReadOnlyList<IssueEvent> stepEvents)
 		{
-			fingerprint = new NewIssueFingerprint(Type, new[] { node.Name }, null, null);
-			return true;
+			NewIssueFingerprint fingerprint = new NewIssueFingerprint(TypeConst, new[] { node.Name }, null, null);
+			foreach (IssueEvent stepEvent in stepEvents)
+			{
+				stepEvent.Fingerprint = fingerprint;
+			}
 		}
 
 		/// <inheritdoc/>
-		public void RankSuspects(IIssueFingerprint fingerprint, List<SuspectChange> suspects)
+		public override void RankSuspects(IIssueFingerprint fingerprint, List<SuspectChange> suspects)
 		{
 		}
 
 		/// <inheritdoc/>
-		public string GetSummary(IIssueFingerprint fingerprint, IssueSeverity severity)
+		public override string GetSummary(IIssueFingerprint fingerprint, IssueSeverity severity)
 		{
 			string nodeName = fingerprint.Keys.FirstOrDefault() ?? "(unknown)";
 			if(severity == IssueSeverity.Warning)
diff --git a/Engine/Source/Programs/Horde/Horde.Build/Issues/Handlers/GauntletIssueHandler.cs b/Engine/Source/Programs/Horde/Horde.Build/Issues/Handlers/GauntletIssueHandler.cs
--- a/Engine/Source/Programs/Horde/Horde.Build/Issues/Handlers/GauntletIssueHandler.cs
+++ b/Engine/Source/Programs/Horde/Horde.Build/Issues/Handlers/GauntletIssueHandler.cs
@@ -17,7 +17,7 @@ namespace Horde.Build.Issues.Handlers
 	/// <summary>
 	/// Instance of a particular Gauntlet error
 	/// </summary>
-	class GauntletIssueHandler : IIssueHandler
+	class GauntletIssueHandler : IssueHandler
 	{
 		/// <summary>
 		/// Prefix for unit test keys
@@ -30,17 +30,17 @@ namespace Horde.Build.Issues.Handlers
 		const string ScreenshotTestPrefix = "Screenshot:";
 
 		/// <inheritdoc/>
-		public string Type => "Gauntlet";
+		public override string Type => "Gauntlet";
 
 		/// <inheritdoc/>
-		public int Priority => 10;
+		public override int Priority => 10;
 
 		/// <summary>
 		/// Determines if the given event id matches
 		/// </summary>
 		/// <param name="eventId">The event id to compare</param>
 		/// <returns>True if the given event id matches</returns>
-		public static bool IsMatchingEventId(EventId? eventId)
+		public static bool IsMatchingEventId(EventId eventId)
 		{
 			return eventId == KnownLogEvents.Gauntlet_UnitTest || eventId == KnownLogEvents.Gauntlet_ScreenshotTest;
 		}
@@ -98,30 +98,30 @@ namespace Horde.Build.Issues.Handlers
 		}
 
 		/// <inheritdoc/>
-		public bool TryGetFingerprint(IJob job, INode node, IReadOnlyNodeAnnotations annotations, ILogEventData eventData, [NotNullWhen(true)] out NewIssueFingerprint? fingerprint)
+		public override void TagEvents(IJob job, INode node, IReadOnlyNodeAnnotations annotations, IReadOnlyList<IssueEvent> stepEvents)
 		{
-			if(!IsMatchingEventId(eventData.EventId))
+			foreach (IssueEvent stepEvent in stepEvents)
 			{
-				fingerprint = null;
-				return false;
-			}
-
-			HashSet<string> keys = new HashSet<string>();
-			GetUnitTestNames(eventData, keys);
-			GetScreenshotTestNames(eventData, keys);
+				if (stepEvent.EventId != null && IsMatchingEventId(stepEvent.EventId.Value))
+				{
+					HashSet<string> keys = new HashSet<string>();
+					GetUnitTestNames(stepEvent.EventData, keys);
+					GetScreenshotTestNames(stepEvent.EventData, keys);
 
-			if (keys.Count == 0)
-			{
-				fingerprint = null;
-				return false;
+					if (keys.Count == 0)
+					{
+						stepEvent.Ignored = true;
+					}
+					else
+					{
+						stepEvent.Fingerprint = new NewIssueFingerprint(Type, keys, null, null);
+					}
+				}
 			}
-
-			fingerprint = new NewIssueFingerprint(Type, keys, null, null);
-			return true;
 		}
 
 		/// <inheritdoc/>
-		public void RankSuspects(IIssueFingerprint fingerprint, List<SuspectChange> changes)
+		public override void RankSuspects(IIssueFingerprint fingerprint, List<SuspectChange> changes)
 		{
 			foreach (SuspectChange change in changes)
 			{
@@ -133,7 +133,7 @@ namespace Horde.Build.Issues.Handlers
 		}
 
 		/// <inheritdoc/>
-		public string GetSummary(IIssueFingerprint fingerprint, IssueSeverity severity)
+		public override string GetSummary(IIssueFingerprint fingerprint, IssueSeverity severity)
 		{
 			List<string> unitTestNames = fingerprint.Keys.Where(x => x.StartsWith(UnitTestPrefix, StringComparison.Ordinal)).Select(x => x.Substring(UnitTestPrefix.Length + 1)).ToList();
 			if (unitTestNames.Count > 0)
diff --git a/Engine/Source/Programs/Horde/Horde.Build/Issues/Handlers/LocalizationIssueHandler.cs b/Engine/Source/Programs/Horde/Horde.Build/Issues/Handlers/LocalizationIssueHandler.cs
--- a/Engine/Source/Programs/Horde/Horde.Build/Issues/Handlers/LocalizationIssueHandler.cs
+++ b/Engine/Source/Programs/Horde/Horde.Build/Issues/Handlers/LocalizationIssueHandler.cs
@@ -17,18 +17,13 @@ namespace Horde.Build.Issues.Handlers
 	/// <summary>
 	/// Instance of a localization error
 	/// </summary>
-	class LocalizationIssueHandler : IIssueHandler
+	class LocalizationIssueHandler : IssueHandler
 	{
 		/// <inheritdoc/>
-		public string Type => "Localization";
+		public override string Type => "Localization";
 
 		/// <inheritdoc/>
-		public int Priority => 10;
-
-		/// <summary>
-		/// Filenames containing errors or warnings
-		/// </summary>
-		public HashSet<string> FileNames { get; set; } = new HashSet<string>();
+		public override int Priority => 10;
 
 		/// <summary>
 		/// Determines if the given event id matches
@@ -63,28 +58,35 @@ namespace Horde.Build.Issues.Handlers
 		}
 
 		/// <inheritdoc/>
-		public bool TryGetFingerprint(IJob job, INode node, IReadOnlyNodeAnnotations annotations, ILogEventData eventData, [NotNullWhen(true)] out NewIssueFingerprint? fingerprint)
+		public override void TagEvents(IJob job, INode node, IReadOnlyNodeAnnotations annotations, IReadOnlyList<IssueEvent> stepEvents)
 		{
-			if (!IsMatchingEventId(eventData.EventId))
+			foreach (IssueEvent stepEvent in stepEvents)
 			{
-				fingerprint = null;
-				return false;
-			}
+				if (stepEvent.EventId != null && IsMatchingEventId(stepEvent.EventId.Value))
+				{
+					HashSet<string> newFileNames = new HashSet<string>(StringComparer.OrdinalIgnoreCase);
+					GetSourceFiles(stepEvent.EventData, newFileNames);
 
-			HashSet<string> newFileNames = new HashSet<string>(StringComparer.OrdinalIgnoreCase);
-			GetSourceFiles(eventData, newFileNames);
-			fingerprint = new NewIssueFingerprint(Type, newFileNames, null, null);
-			return true;
+					if (newFileNames.Count == 0)
+					{
+						stepEvent.Ignored = true;
+					}
+					else
+					{
+						stepEvent.Fingerprint = new NewIssueFingerprint(Type, newFileNames, null, null);
+					}
+				}
+			}
 		}
 
 		/// <inheritdoc/>
-		public void RankSuspects(IIssueFingerprint fingerprint, List<SuspectChange> suspects)
+		public override void RankSuspects(IIssueFingerprint fingerprint, List<SuspectChange> suspects)
 		{
 			foreach (SuspectChange suspect in suspects)
 			{
 				if (suspect.ContainsCode)
 				{
-					if (FileNames.Any(x => suspect.ModifiesFile(x)))
+					if (fingerprint.Keys.Any(x => suspect.ModifiesFile(x)))
 					{
 						suspect.Rank += 20;
 					}
@@ -97,7 +99,7 @@ namespace Horde.Build.Issues.Handlers
 		}
 
 		/// <inheritdoc/>
-		public string GetSummary(IIssueFingerprint fingerprint, IssueSeverity severity)
+		public override string GetSummary(IIssueFingerprint fingerprint, IssueSeverity severity)
 		{
 			string type = (severity == IssueSeverity.Warning)? "warnings" : "errors";
 			return $"Localization {type} in {StringUtils.FormatList(fingerprint.Keys.ToArray(), 2)}";
diff --git a/Engine/Source/Programs/Horde/Horde.Build/Issues/Handlers/PerforceCaseIssueHandler.cs b/Engine/Source/Programs/Horde/Horde.Build/Issues/Handlers/PerforceCaseIssueHandler.cs
--- a/Engine/Source/Programs/Horde/Horde.Build/Issues/Handlers/PerforceCaseIssueHandler.cs
+++ b/Engine/Source/Programs/Horde/Horde.Build/Issues/Handlers/PerforceCaseIssueHandler.cs
@@ -15,26 +15,26 @@ namespace Horde.Build.Issues.Handlers
 	/// <summary>
 	/// Instance of a Perforce case mismatch error
 	/// </summary>
-	class PerforceCaseIssueHandler : IIssueHandler
+	class PerforceCaseIssueHandler : IssueHandler
 	{
 		/// <inheritdoc/>
-		public string Type => "PerforceCase";
+		public override string Type => "PerforceCase";
 
 		/// <inheritdoc/>
-		public int Priority => 10;
+		public override int Priority => 10;
 
 		/// <summary>
 		/// Determines if the given event id matches
 		/// </summary>
 		/// <param name="eventId">The event id to compare</param>
 		/// <returns>True if the given event id matches</returns>
-		public static bool IsMatchingEventId(EventId? eventId)
+		public static bool IsMatchingEventId(EventId eventId)
 		{
 			return eventId == KnownLogEvents.AutomationTool_PerforceCase;
 		}
 
 		/// <inheritdoc/>
-		public void RankSuspects(IIssueFingerprint fingerprint, List<SuspectChange> changes)
+		public override void RankSuspects(IIssueFingerprint fingerprint, List<SuspectChange> changes)
 		{
 			foreach (SuspectChange change in changes)
 			{
@@ -46,7 +46,7 @@ namespace Horde.Build.Issues.Handlers
 		}
 
 		/// <inheritdoc/>
-		public string GetSummary(IIssueFingerprint fingerprint, IssueSeverity severity)
+		public override string GetSummary(IIssueFingerprint fingerprint, IssueSeverity severity)
 		{
 			return $"Inconsistent case for {StringUtils.FormatList(fingerprint.Keys.Select(x => x.Substring(x.LastIndexOf('/') + 1)).ToArray(), 3)}";
 		}
@@ -75,19 +75,18 @@ namespace Horde.Build.Issues.Handlers
 		}
 
 		/// <inheritdoc/>
-		public bool TryGetFingerprint(IJob job, INode node, IReadOnlyNodeAnnotations annotations, ILogEventData eventData, [NotNullWhen(true)] out NewIssueFingerprint? fingerprint)
+		public override void TagEvents(IJob job, INode node, IReadOnlyNodeAnnotations annotations, IReadOnlyList<IssueEvent> stepEvents)
 		{
-			if (!IsMatchingEventId(eventData.EventId))
+			foreach (IssueEvent stepEvent in stepEvents)
 			{
-				fingerprint = null;
-				return false;
-			}
-
-			HashSet<string> newFileNames = new HashSet<string>();
-			GetSourceFiles(eventData, newFileNames);
+				if (stepEvent.EventId != null && IsMatchingEventId(stepEvent.EventId.Value))
+				{
+					HashSet<string> newFileNames = new HashSet<string>();
+					GetSourceFiles(stepEvent.EventData, newFileNames);
 
-			fingerprint = new NewIssueFingerprint(Type, newFileNames, null, null);
-			return true;
+					stepEvent.Fingerprint = new NewIssueFingerprint(Type, newFileNames, null, null);
+				}
+			}
 		}
 	}
 }
diff --git a/Engine/Source/Programs/Horde/Horde.Build/Issues/Handlers/SourceFileIssueHandler.cs b/Engine/Source/Programs/Horde/Horde.Build/Issues/Handlers/SourceFileIssueHandler.cs
--- a/Engine/Source/Programs/Horde/Horde.Build/Issues/Handlers/SourceFileIssueHandler.cs
+++ b/Engine/Source/Programs/Horde/Horde.Build/Issues/Handlers/SourceFileIssueHandler.cs
@@ -16,19 +16,13 @@ namespace Horde.Build.Issues.Handlers
 	/// <summary>
 	/// Instance of a particular compile error
 	/// </summary>
-	abstract class SourceFileIssueHandler : IIssueHandler
+	abstract class SourceFileIssueHandler : IssueHandler
 	{
 		/// <summary>
 		/// Prefix used to identify files that may match against modified files, but which are not the files failing to compile
 		/// </summary>
 		protected const string NotePrefix = "note:";
 
-		/// <inheritdoc/>
-		public abstract string Type { get; }
-
-		/// <inheritdoc/>
-		public abstract int Priority { get; }
-
 		/// <summary>
 		/// Extracts a list of source files from an event
 		/// </summary>
@@ -87,13 +81,7 @@ namespace Horde.Build.Issues.Handlers
 		}
 
 		/// <inheritdoc/>
-		public abstract bool TryGetFingerprint(IJob job, INode node, IReadOnlyNodeAnnotations annotations, ILogEventData eventData, [NotNullWhen(true)] out NewIssueFingerprint? fingerprint);
-
-		/// <inheritdoc/>
-		public abstract string GetSummary(IIssueFingerprint fingerprint, IssueSeverity severity);
-
-		/// <inheritdoc/>
-		public void RankSuspects(IIssueFingerprint fingerprint, List<SuspectChange> suspects)
+		public override void RankSuspects(IIssueFingerprint fingerprint, List<SuspectChange> suspects)
 		{
 			List<string> fileNames = new List<string>();
 			foreach (string key in fingerprint.Keys)
diff --git a/Engine/Source/Programs/Horde/Horde.Build/Issues/Handlers/SymbolIssueHandler.cs b/Engine/Source/Programs/Horde/Horde.Build/Issues/Handlers/SymbolIssueHandler.cs
--- a/Engine/Source/Programs/Horde/Horde.Build/Issues/Handlers/SymbolIssueHandler.cs
+++ b/Engine/Source/Programs/Horde/Horde.Build/Issues/Handlers/SymbolIssueHandler.cs
@@ -17,7 +17,7 @@ namespace Horde.Build.Issues.Handlers
 	/// <summary>
 	/// Instance of a particular compile error
 	/// </summary>
-	class SymbolIssueHandler : IIssueHandler
+	class SymbolIssueHandler : IssueHandler
 	{
 		const string NodeName = "Node";
 		const string EventIdName = "EventId";
@@ -25,17 +25,17 @@ namespace Horde.Build.Issues.Handlers
 		static readonly string DuplicateEventIdMetadata = $"{EventIdName}={KnownLogEvents.Linker_DuplicateSymbol.Id}";
 
 		/// <inheritdoc/>
-		public string Type => "Symbol";
+		public override string Type => "Symbol";
 
 		/// <inheritdoc/>
-		public int Priority => 10;
+		public override int Priority => 10;
 
 		/// <summary>
 		/// Determines if the given event id matches
 		/// </summary>
 		/// <param name="eventId">The event id to compare</param>
 		/// <returns>True if the given event id matches</returns>
-		public static bool IsMatchingEventId(EventId? eventId)
+		public static bool IsMatchingEventId(EventId eventId)
 		{
 			return eventId == KnownLogEvents.Linker_UndefinedSymbol || eventId == KnownLogEvents.Linker_DuplicateSymbol || eventId == KnownLogEvents.Linker;
 		}
@@ -58,7 +58,7 @@ namespace Horde.Build.Issues.Handlers
 		}
 
 		/// <inheritdoc/>
-		public void RankSuspects(IIssueFingerprint fingerprint, List<SuspectChange> changes)
+		public override void RankSuspects(IIssueFingerprint fingerprint, List<SuspectChange> changes)
 		{
 			HashSet<string> names = new HashSet<string>();
 			foreach (string name in fingerprint.Keys)
@@ -77,7 +77,7 @@ namespace Horde.Build.Issues.Handlers
 		}
 
 		/// <inheritdoc/>
-		public string GetSummary(IIssueFingerprint fingerprint, IssueSeverity severity)
+		public override string GetSummary(IIssueFingerprint fingerprint, IssueSeverity severity)
 		{
 			HashSet<string> symbols = fingerprint.Keys;
 			if (symbols.Count == 0)
@@ -107,26 +107,39 @@ namespace Horde.Build.Issues.Handlers
 			}
 		}
 
-		public bool TryGetFingerprint(IJob job, INode node, IReadOnlyNodeAnnotations annotations, ILogEventData eventData, [NotNullWhen(true)] out NewIssueFingerprint? fingerprint)
+		public override void TagEvents(IJob job, INode node, IReadOnlyNodeAnnotations annotations, IReadOnlyList<IssueEvent> stepEvents)
 		{
-			if (!IsMatchingEventId(eventData.EventId))
+			bool hasMatches = false;
+			foreach (IssueEvent stepEvent in stepEvents)
 			{
-				fingerprint = null;
-				return false;
-			}
-
-			SortedSet<string> symbolNames = new SortedSet<string>();
-			GetSymbolNames(eventData, symbolNames);
-
-			List<string> metadata = new List<string>();
-			metadata.Add($"{NodeName}={node.Name}");
-			if (eventData.EventId != null)
-			{
-				metadata.Add($"{EventIdName}={eventData.EventId.Value.Id}");
+				if (stepEvent.EventId != null)
+				{
+					EventId eventId = stepEvent.EventId.Value;
+					if (IsMatchingEventId(eventId))
+					{
+						SortedSet<string> symbolNames = new SortedSet<string>();
+						GetSymbolNames(stepEvent.EventData, symbolNames);
+
+						if (symbolNames.Count > 0)
+						{
+							List<string> metadata = new List<string>();
+							metadata.Add($"{NodeName}={node.Name}");
+							metadata.Add($"{EventIdName}={eventId.Id}");
+
+							stepEvent.Fingerprint = new NewIssueFingerprint(Type, symbolNames, null, metadata);
+							hasMatches = true;
+						}
+						else if (hasMatches)
+						{
+							stepEvent.Ignored = true;
+						}
+					}
+					else if(eventId == KnownLogEvents.Generic || eventId == KnownLogEvents.AutomationTool)
+					{
+						stepEvent.Ignored = true;
+					}
+				}
 			}
-
-			fingerprint = new NewIssueFingerprint(Type, symbolNames, null, metadata);
-			return true;
 		}
 	}
 }
diff --git a/Engine/Source/Programs/Horde/Horde.Build/Issues/Handlers/UnacceptableWordsIssueHandler.cs b/Engine/Source/Programs/Horde/Horde.Build/Issues/Handlers/UnacceptableWordsIssueHandler.cs
--- a/Engine/Source/Programs/Horde/Horde.Build/Issues/Handlers/UnacceptableWordsIssueHandler.cs
+++ b/Engine/Source/Programs/Horde/Horde.Build/Issues/Handlers/UnacceptableWordsIssueHandler.cs
@@ -27,7 +27,7 @@ namespace Horde.Build.Issues.Handlers
 		/// </summary>
 		/// <param name="eventId">The event id to compare</param>
 		/// <returns>True if the given event id matches</returns>
-		public static bool IsMatchingEventId(EventId? eventId)
+		public static bool IsMatchingEventId(EventId eventId)
 		{
 			return eventId == KnownLogEvents.AutomationTool_UnacceptableWords;
 		}
@@ -39,19 +39,18 @@ namespace Horde.Build.Issues.Handlers
 		}
 
 		/// <inheritdoc/>
-		public override bool TryGetFingerprint(IJob job, INode node, IReadOnlyNodeAnnotations annotations, ILogEventData eventData, [NotNullWhen(true)] out NewIssueFingerprint? fingerprint)
+		public override void TagEvents(IJob job, INode node, IReadOnlyNodeAnnotations annotations, IReadOnlyList<IssueEvent> stepEvents)
 		{
-			if (!IsMatchingEventId(eventData.EventId))
+			foreach (IssueEvent stepEvent in stepEvents)
 			{
-				fingerprint = null;
-				return false;
-			}
-
-			List<string> newFileNames = new List<string>();
-			GetSourceFiles(eventData, newFileNames);
+				if(stepEvent.EventId != null && IsMatchingEventId(stepEvent.EventId.Value))
+				{
+					List<string> newFileNames = new List<string>();
+					GetSourceFiles(stepEvent.EventData, newFileNames);
 
-			fingerprint = new NewIssueFingerprint(Type, newFileNames, null, null);
-			return true;
+					stepEvent.Fingerprint = new NewIssueFingerprint(Type, newFileNames, null, null);
+				}
+			}
 		}
 	}
 }
diff --git a/Engine/Source/Programs/Horde/Horde.Build/Issues/IIssueCollection.cs b/Engine/Source/Programs/Horde/Horde.Build/Issues/IIssueCollection.cs
--- a/Engine/Source/Programs/Horde/Horde.Build/Issues/IIssueCollection.cs
+++ b/Engine/Source/Programs/Horde/Horde.Build/Issues/IIssueCollection.cs
@@ -95,19 +95,21 @@ namespace Horde.Build.Issues
 		/// <inheritdoc/>
 		public bool Equals(IIssueFingerprint? otherFingerprint)
 		{
-			if(otherFingerprint == null || !Type.Equals(otherFingerprint.Type, StringComparison.Ordinal) || !Keys.SetEquals(otherFingerprint.Keys))
+			if (!ReferenceEquals(this, otherFingerprint))
 			{
-				return false;
-			}
-			if (!ContentsEqual(RejectKeys, otherFingerprint.RejectKeys))
-			{
-				return false;
-			}
-			if (!ContentsEqual(Metadata, otherFingerprint.Metadata))
-			{
-				return false;
+				if (otherFingerprint == null || !Type.Equals(otherFingerprint.Type, StringComparison.Ordinal) || !Keys.SetEquals(otherFingerprint.Keys))
+				{
+					return false;
+				}
+				if (!ContentsEqual(RejectKeys, otherFingerprint.RejectKeys))
+				{
+					return false;
+				}
+				if (!ContentsEqual(Metadata, otherFingerprint.Metadata))
+				{
+					return false;
+				}
 			}
-
 			return true;
 		}
 
diff --git a/Engine/Source/Programs/Horde/Horde.Build/Issues/IssueService.cs b/Engine/Source/Programs/Horde/Horde.Build/Issues/IssueService.cs
--- a/Engine/Source/Programs/Horde/Horde.Build/Issues/IssueService.cs
+++ b/Engine/Source/Programs/Horde/Horde.Build/Issues/IssueService.cs
@@ -126,7 +126,7 @@ namespace Horde.Build.Issues
 		/// <returns>True if the user should be notified for this change</returns>
 		public bool ShowNotifications()
 		{
-			return _showDesktopAlerts && Issue.Fingerprints.Any(x => x.Type != DefaultIssueHandler.Type);
+			return _showDesktopAlerts && Issue.Fingerprints.Any(x => x.Type != DefaultIssueHandler.TypeConst);
 		}
 
 		/// <summary>
@@ -176,14 +176,14 @@ namespace Horde.Build.Issues
 		public event IssueUpdatedEvent? OnIssueUpdated;
 
 		/// <summary>
-		/// List of issue matchers
+		/// List of issue handlers
 		/// </summary>
-		readonly List<IIssueHandler> _matchers = new List<IIssueHandler>();
+		readonly List<IssueHandler> _handlers = new List<IssueHandler>();
 
 		/// <summary>
 		/// Available issue serializers
 		/// </summary>
-		readonly Dictionary<string, IIssueHandler> _typeToHandler = new Dictionary<string, IIssueHandler>(StringComparer.Ordinal);
+		readonly Dictionary<string, IssueHandler> _typeToHandler = new Dictionary<string, IssueHandler>(StringComparer.Ordinal);
 
 		/// <summary>
 		/// Cached list of currently open issues
@@ -228,17 +228,17 @@ namespace Horde.Build.Issues
 			_ticker = clock.AddTicker<IssueService>(TimeSpan.FromMinutes(1.0), TickAsync, logger);
 			_logger = logger;
 
-			// Create all the issue factories
-			Type[] matcherTypes = Assembly.GetExecutingAssembly().GetTypes().Where(x => !x.IsAbstract && typeof(IIssueHandler).IsAssignableFrom(x)).ToArray();
-			foreach (Type matcherType in matcherTypes)
+			// Create all the issue handlers
+			Type[] handlerTypes = Assembly.GetExecutingAssembly().GetTypes().Where(x => !x.IsAbstract && typeof(IssueHandler).IsAssignableFrom(x)).ToArray();
+			foreach (Type handlerType in handlerTypes)
 			{
-				IIssueHandler matcher = (IIssueHandler)Activator.CreateInstance(matcherType)!;
-				_matchers.Add(matcher);
+				IssueHandler matcher = (IssueHandler)Activator.CreateInstance(handlerType)!;
+				_handlers.Add(matcher);
 			}
-			_matchers.SortBy(x => -x.Priority);
+			_handlers.SortBy(x => -x.Priority);
 
 			// Build the type name to factory map
-			_typeToHandler = _matchers.ToDictionary(x => x.Type, x => x, StringComparer.Ordinal);
+			_typeToHandler = _handlers.ToDictionary(x => x.Type, x => x, StringComparer.Ordinal);
 		}
 
 		/// <inheritdoc/>
@@ -440,83 +440,6 @@ namespace Horde.Build.Issues
 			return true;
 		}
 
-		/// <summary>
-		/// Information about a particular log event to be added to a span
-		/// </summary>
-		class NewEvent
-		{
-			/// <summary>
-			/// The event interface
-			/// </summary>
-			public ILogEvent Event { get; }
-
-			/// <summary>
-			/// Data for the event
-			/// </summary>
-			public ILogEventData EventData { get; }
-
-			/// <summary>
-			/// Constructor
-			/// </summary>
-			/// <param name="logEvent">The event interface</param>
-			/// <param name="eventData">Data for the event</param>
-			public NewEvent(ILogEvent logEvent, ILogEventData eventData)
-			{
-				Event = logEvent;
-				EventData = eventData;
-			}
-
-			/// <inheritdoc/>
-			public override string ToString()
-			{
-				return $"[{Event.LineIndex}] {EventData.Message}";
-			}
-		}
-
-		/// <summary>
-		/// Information about a new event to be added to a span
-		/// </summary>
-		class NewEventGroup
-		{
-			/// <summary>
-			/// Digest of the fingerprint, for log tracking
-			/// </summary>
-			public ContentHash Digest { get; }
-
-			/// <summary>
-			/// Fingerprint for the event
-			/// </summary>
-			public NewIssueFingerprint Fingerprint { get; }
-
-			/// <summary>
-			/// Individual log events
-			/// </summary>
-			public List<NewEvent> Events { get; } = new List<NewEvent>();
-
-			/// <summary>
-			/// Constructor
-			/// </summary>
-			/// <param name="fingerprint">Fingerprint for the event</param>
-			public NewEventGroup(NewIssueFingerprint fingerprint)
-			{
-				Digest = ContentHash.MD5(fingerprint.ToString());
-				Fingerprint = fingerprint;
-			}
-
-			/// <summary>
-			/// Merge with another group
-			/// </summary>
-			/// <param name="otherGroup">The group to merge with</param>
-			/// <returns>A new group combining both groups</returns>
-			public NewEventGroup MergeWith(NewEventGroup otherGroup)
-			{
-				NewEventGroup newGroup = new NewEventGroup(NewIssueFingerprint.Merge(Fingerprint, otherGroup.Fingerprint));
-				newGroup.Events.AddRange(Events);
-				newGroup.Events.AddRange(otherGroup.Events);
-				return newGroup;
-			}
-		}
-
 		/// <summary>
 		/// Marks a step as complete
 		/// </summary>
@@ -591,7 +514,7 @@ namespace Horde.Build.Issues
 			}
 
 			// Gets the events for this step grouped by fingerprint
-			HashSet<NewEventGroup> eventGroups = await GetEventGroupsForStepAsync(job, batch, step, node, annotations);
+			HashSet<IssueEventGroup> eventGroups = await GetEventGroupsForStepAsync(job, batch, step, node, annotations);
 
 			// Try to update all the events. We may need to restart this due to optimistic transactions, so keep track of any existing spans we do not need to check against.
 			await using(IAsyncDisposable issueLock = await _issueCollection.EnterCriticalSectionAsync())
@@ -636,7 +559,7 @@ namespace Horde.Build.Issues
 		/// <param name="node">The node corresponding to the step</param>
 		/// <param name="annotations">Annotations for this node</param>
 		/// <returns>Set of new events</returns>
-		async Task<HashSet<NewEventGroup>> GetEventGroupsForStepAsync(IJob job, IJobStepBatch batch, IJobStep step, INode node, IReadOnlyNodeAnnotations annotations)
+		async Task<HashSet<IssueEventGroup>> GetEventGroupsForStepAsync(IJob job, IJobStepBatch batch, IJobStep step, INode node, IReadOnlyNodeAnnotations annotations)
 		{
 			// Make sure the step has a log file
 			if (step.LogId == null)
@@ -651,70 +574,48 @@ namespace Horde.Build.Issues
 				throw new ArgumentException($"Unable to retrieve log {step.LogId}");
 			}
 
-			// Process all the events for this step
-			List<ILogEvent> stepEvents = await _logFileService.FindEventsAsync(logFile);
+			// Create a list of all the events for this step
+			List<IssueEvent> issueEvents = new List<IssueEvent>();
 
-			// Create a new list of event fingerprints for each one
-			bool nonSystemicErrors = false;
-			Dictionary<NewIssueFingerprint, NewEventGroup> fingerprintToEventGroup = new Dictionary<NewIssueFingerprint, NewEventGroup>();
+			List<ILogEvent> stepEvents = await _logFileService.FindEventsAsync(logFile);
 			foreach (ILogEvent stepEvent in stepEvents)
 			{
-				// Read the full message for this log event
 				ILogEventData stepEventData = await _logFileService.GetEventDataAsync(logFile, stepEvent.LineIndex, stepEvent.LineCount);
-
-				// Parse a fingerprint from the event
-				NewIssueFingerprint? fingerprint = null;
-				foreach (IIssueHandler matcher in _matchers)
-				{
-					if (matcher.TryGetFingerprint(job, node, annotations, stepEventData, out fingerprint))
-					{
-						if (fingerprint.Type != "Systemic")
-						{
-							if (stepEventData.Severity == EventSeverity.Error)
-							{
-								nonSystemicErrors = true;
-							}
-						}
-
-						break;
-					}
-				}
-
-				// If we matched a fingerprint, add it to the existing set
-				if (fingerprint != null)
-				{
-					NewEventGroup? eventGroup;
-					if (!fingerprintToEventGroup.TryGetValue(fingerprint, out eventGroup))
-					{
-						eventGroup = new NewEventGroup(fingerprint);
-						fingerprintToEventGroup.Add(fingerprint, eventGroup);
-					}
-					eventGroup.Events.Add(new NewEvent(stepEvent, stepEventData));
-				}
+				issueEvents.Add(new IssueEvent(stepEvent, stepEventData));
 			}
-						
-			// Only generate user issues for systemic issues where there are no other failures, otherwise this leads to many false positives (Experimental)
-			if (nonSystemicErrors)
+
+			// Pass all the events to each matcher in turn, and allow it to tag any events it can handle
+			List<IssueEvent> remainingEvents = new List<IssueEvent>(issueEvents);
+			foreach (IssueHandler handler in _handlers)
 			{
-				List<NewIssueFingerprint> systemicFingerprints = fingerprintToEventGroup.Keys.Where(x => x.Type == "Systemic").ToList();
+				handler.TagEvents(job, node, annotations, remainingEvents);
+				remainingEvents.RemoveAll(x => x.Ignored || x.Fingerprint != null);
+			}
 
-				if ((systemicFingerprints.Count != 0) && (systemicFingerprints.Count != fingerprintToEventGroup.Values.Count))
+			// Create all the event groups from the non-ignored events
+			Dictionary<NewIssueFingerprint, IssueEventGroup> fingerprintToEventGroup = new Dictionary<NewIssueFingerprint, IssueEventGroup>();
+			foreach (IssueEvent issueEvent in issueEvents)
+			{
+				if (!issueEvent.Ignored && issueEvent.Fingerprint != null)
 				{
-					for (int i = 0; i < systemicFingerprints.Count; i++)
+					IssueEventGroup? eventGroup;
+					if (!fingerprintToEventGroup.TryGetValue(issueEvent.Fingerprint, out eventGroup))
 					{
-						fingerprintToEventGroup.Remove(systemicFingerprints[i]);
+						eventGroup = new IssueEventGroup(issueEvent.Fingerprint);
+						fingerprintToEventGroup.Add(issueEvent.Fingerprint, eventGroup);
 					}
+					eventGroup.Events.Add(issueEvent);
 				}
 			}
 
 			// Print the list of new events
-			HashSet<NewEventGroup> eventGroups = new HashSet<NewEventGroup>(fingerprintToEventGroup.Values);
+			HashSet<IssueEventGroup> eventGroups = new HashSet<IssueEventGroup>(fingerprintToEventGroup.Values);
 
 			_logger.LogDebug("UpdateCompleteStep({JobId}, {BatchId}, {StepId}): {NumEvents} events, {NumFingerprints} unique fingerprints", job.Id, batch.Id, step.Id, stepEvents.Count, eventGroups.Count);
-			foreach (NewEventGroup eventGroup in eventGroups)
+			foreach (IssueEventGroup eventGroup in eventGroups)
 			{
 				_logger.LogDebug("Group {Digest}: Type '{FingerprintType}', keys '{FingerprintKeys}'", eventGroup.Digest.ToString(), eventGroup.Fingerprint.Type, String.Join(", ", eventGroup.Fingerprint.Keys));
-				foreach (NewEvent eventItem in eventGroup.Events)
+				foreach (IssueEvent eventItem in eventGroup.Events)
 				{
 					_logger.LogDebug("Group {Digest}: [{Line}] {Message}", eventGroup.Digest.ToString(), eventItem.Event.LineIndex, eventItem.EventData.Message);
 				}
@@ -735,7 +636,7 @@ namespace Horde.Build.Issues
 		/// <param name="annotations">Annotations for this step</param>
 		/// <param name="promoteByDefault"></param>
 		/// <returns>True if the adding completed</returns>
-		async Task<bool> AddEventsToExistingSpans(IJob job, IJobStepBatch batch, IJobStep step, HashSet<NewEventGroup> newEventGroups, List<IIssueSpan> openSpans, HashSet<ObjectId> checkedSpanIds, IReadOnlyNodeAnnotations? annotations, bool promoteByDefault)
+		async Task<bool> AddEventsToExistingSpans(IJob job, IJobStepBatch batch, IJobStep step, HashSet<IssueEventGroup> newEventGroups, List<IIssueSpan> openSpans, HashSet<ObjectId> checkedSpanIds, IReadOnlyNodeAnnotations? annotations, bool promoteByDefault)
 		{
 			for(int spanIdx = 0; spanIdx < openSpans.Count; spanIdx++)
 			{
@@ -743,7 +644,7 @@ namespace Horde.Build.Issues
 				if (!checkedSpanIds.Contains(openSpan.Id))
 				{
 					// Filter out the events which match the span's fingerprint
-					List<NewEventGroup> matchEventGroups = newEventGroups.Where(x => openSpan.Fingerprint.IsMatch(x.Fingerprint)).ToList();
+					List<IssueEventGroup> matchEventGroups = newEventGroups.Where(x => openSpan.Fingerprint.IsMatch(x.Fingerprint)).ToList();
 					if (matchEventGroups.Count > 0)
 					{
 						// Add the new step data
@@ -763,7 +664,7 @@ namespace Horde.Build.Issues
 						}
 
 						// Write out all the merged events
-						foreach (NewEventGroup eventGroup in matchEventGroups)
+						foreach (IssueEventGroup eventGroup in matchEventGroups)
 						{
 							_logger.LogDebug("Matched fingerprint {Digest} ({NumLogEvents} log events) to span {SpanId}", eventGroup.Digest, eventGroup.Events.Count, newSpan.Id);
 						}
@@ -787,14 +688,14 @@ namespace Horde.Build.Issues
 		/// </summary>
 		/// <param name="events">Set of event severity values</param>
 		/// <returns>Severity of this issue</returns>
-		static IssueSeverity GetIssueSeverity(IEnumerable<NewEvent> events)
+		static IssueSeverity GetIssueSeverity(IEnumerable<IssueEvent> events)
 		{
 			IssueSeverity severity;
-			if (events.Any(x => x.Event.Severity == EventSeverity.Error))
+			if (events.Any(x => x.Severity == EventSeverity.Error))
 			{
 				severity = IssueSeverity.Error;
 			}
-			else if (events.Any(x => x.Event.Severity == EventSeverity.Warning))
+			else if (events.Any(x => x.Severity == EventSeverity.Warning))
 			{
 				severity = IssueSeverity.Warning;
 			}
@@ -818,21 +719,21 @@ namespace Horde.Build.Issues
 		/// <param name="annotations"></param>
 		/// <param name="promoteByDefault"></param>
 		/// <returns>True if all events were added</returns>
-		async Task<bool> AddEventsToNewSpans(IStream stream, IJob job, IJobStepBatch batch, IJobStep step, INode node, List<IIssueSpan> openSpans, HashSet<NewEventGroup> newEventGroups, IReadOnlyNodeAnnotations? annotations, bool promoteByDefault)
+		async Task<bool> AddEventsToNewSpans(IStream stream, IJob job, IJobStepBatch batch, IJobStep step, INode node, List<IIssueSpan> openSpans, HashSet<IssueEventGroup> newEventGroups, IReadOnlyNodeAnnotations? annotations, bool promoteByDefault)
 		{
 			while (newEventGroups.Count > 0)
 			{
 				// Keep track of the event groups we merge together
-				List<NewEventGroup> sourceEventGroups = new List<NewEventGroup>();
+				List<IssueEventGroup> sourceEventGroups = new List<IssueEventGroup>();
 				sourceEventGroups.Add(newEventGroups.First());
 
 				// Take the first event, and find all other events that match against it
-				NewEventGroup eventGroup = sourceEventGroups[0];
-				foreach (NewEventGroup otherEventGroup in newEventGroups.Skip(1))
+				IssueEventGroup eventGroup = sourceEventGroups[0];
+				foreach (IssueEventGroup otherEventGroup in newEventGroups.Skip(1))
 				{
 					if (otherEventGroup.Fingerprint.IsMatchForNewSpan(eventGroup.Fingerprint))
 					{
-						NewEventGroup newEventGroup = eventGroup.MergeWith(otherEventGroup);
+						IssueEventGroup newEventGroup = eventGroup.MergeWith(otherEventGroup);
 						_logger.LogDebug("Merging group {Group} with group {OtherGroup} to form {NewGroup}", eventGroup.Digest.ToString(), otherEventGroup.Digest.ToString(), newEventGroup.Digest.ToString());
 						sourceEventGroups.Add(otherEventGroup);
 						eventGroup = newEventGroup;
@@ -1154,7 +1055,7 @@ namespace Horde.Build.Issues
 		async Task<List<NewIssueSpanSuspectData>> FindSuspectsForSpanAsync(IStream stream, IIssueFingerprint fingerprint, int minChange, int maxChange)
 		{
 			List<NewIssueSpanSuspectData> suspects = new List<NewIssueSpanSuspectData>();
-			if (_typeToHandler.TryGetValue(fingerprint.Type, out IIssueHandler? handler))
+			if (_typeToHandler.TryGetValue(fingerprint.Type, out IssueHandler? handler))
 			{
 				_logger.LogDebug("Querying for changes in {StreamName} between {MinChange} and {MaxChange}", stream.Name, minChange, maxChange);
 
@@ -1256,9 +1157,9 @@ namespace Horde.Build.Issues
 		/// <returns>The summary text</returns>
 		string GetSummary(IIssueFingerprint fingerprint, IssueSeverity severity)
 		{
-			if (_typeToHandler.TryGetValue(fingerprint.Type, out IIssueHandler? matcher))
+			if (_typeToHandler.TryGetValue(fingerprint.Type, out IssueHandler? handler))
 			{
-				return matcher.GetSummary(fingerprint, severity);
+				return handler.GetSummary(fingerprint, severity);
 			}
 			else
 			{
