Horde: Missing files.

diff --git a/Engine/Source/Programs/Horde/Horde.Build/Issues/Handlers/SystemicIssueHandler.cs b/Engine/Source/Programs/Horde/Horde.Build/Issues/Handlers/SystemicIssueHandler.cs
--- a/Engine/Source/Programs/Horde/Horde.Build/Issues/Handlers/SystemicIssueHandler.cs
+++ b/Engine/Source/Programs/Horde/Horde.Build/Issues/Handlers/SystemicIssueHandler.cs
@@ -16,13 +16,13 @@ namespace Horde.Build.Issues.Handlers
 	/// <summary>
 	/// Instance of a particular systemic error
 	/// </summary>
-	class SystemicIssueHandler : IIssueHandler
+	class SystemicIssueHandler : IssueHandler
 	{
 		/// <inheritdoc/>
-		public string Type => "Systemic";
+		public override string Type => "Systemic";
 
 		/// <inheritdoc/>
-		public int Priority => 10;
+		public override int Priority => 10;
 
 		/// <summary>
 		///  Known systemic errors
@@ -40,15 +40,23 @@ namespace Horde.Build.Issues.Handlers
 		}
 
 		/// <inheritdoc/>
-		public void RankSuspects(IIssueFingerprint fingerprint, List<SuspectChange> suspects)
+		public override void RankSuspects(IIssueFingerprint fingerprint, List<SuspectChange> suspects)
 		{
 		}
 
 		/// <inheritdoc/>
-		public bool TryGetFingerprint(IJob job, INode node, IReadOnlyNodeAnnotations annotations, ILogEventData eventData, [NotNullWhen(true)] out NewIssueFingerprint? fingerprint)
+		public override void TagEvents(IJob job, INode node, IReadOnlyNodeAnnotations annotations, IReadOnlyList<IssueEvent> stepEvents)
 		{
-			fingerprint = null;
+			foreach (IssueEvent stepEvent in stepEvents)
+			{
+				if (stepEvent.EventId != null && MatchEvent(stepEvent.EventData))
+				{
+				}
+			}
+		}
 
+		static bool MatchEvent(ILogEventData eventData)
+		{
 			if (eventData.EventId == KnownLogEvents.ExitCode)
 			{
 				for (int i = 0; i < eventData.Lines.Count; i++)
@@ -61,24 +69,17 @@ namespace Horde.Build.Issues.Handlers
 					{
 						if (message.Contains(systemicExitMessages[j], StringComparison.InvariantCultureIgnoreCase))
 						{
-							fingerprint = new NewIssueFingerprint(Type, new[] { node.Name }, null, null);
 							return true;
 						}
 					}
 				}
 			}
 
-			if (eventData.EventId == null || !IsMatchingEventId(eventData.EventId.Value))
-			{				
-				return false;
-			}
-
-			fingerprint = new NewIssueFingerprint(Type, new[] { node.Name }, null, null);
-			return true;
+			return eventData.EventId != null && IsMatchingEventId(eventData.EventId.Value);
 		}
 
 		/// <inheritdoc/>
-		public string GetSummary(IIssueFingerprint fingerprint, IssueSeverity severity)
+		public override string GetSummary(IIssueFingerprint fingerprint, IssueSeverity severity)
 		{
 			string type = (severity == IssueSeverity.Warning) ? "Systemic warnings" : "Systemic errors";
 			string nodeName = fingerprint.Keys.FirstOrDefault() ?? "(unknown)";
diff --git a/Engine/Source/Programs/Horde/Horde.Build/Issues/IIssueHandler.cs b/Engine/Source/Programs/Horde/Horde.Build/Issues/IIssueHandler.cs
deleted file mode 100644
diff --git a/Engine/Source/Programs/Horde/Horde.Build/Issues/IssueEvent.cs b/Engine/Source/Programs/Horde/Horde.Build/Issues/IssueEvent.cs
new file mode 100644
--- /dev/null
+++ b/Engine/Source/Programs/Horde/Horde.Build/Issues/IssueEvent.cs
@@ -0,0 +1,119 @@
+// Copyright Epic Games, Inc. All Rights Reserved.
+
+using System.Collections.Generic;
+using System.Text;
+using EpicGames.Core;
+using Horde.Build.Logs;
+using HordeCommon;
+using Microsoft.Extensions.Logging;
+
+namespace Horde.Build.Issues
+{
+	/// <summary>
+	/// Wraps a log event and allows it to be tagged by issue handlers
+	/// </summary>
+	public class IssueEvent
+	{
+		/// <summary>
+		/// The underlying log event
+		/// </summary>
+		public ILogEvent Event { get; }
+
+		/// <summary>
+		/// The log event data
+		/// </summary>
+		public ILogEventData EventData { get; }
+
+		/// <summary>
+		/// Severity of the event
+		/// </summary>
+		public EventSeverity Severity => EventData.Severity;
+
+		/// <summary>
+		/// The type of event
+		/// </summary>
+		public EventId? EventId => EventData.EventId;
+
+		/// <summary>
+		/// The complete rendered message, in plaintext
+		/// </summary>
+		public string Message => EventData.Message;
+
+		/// <summary>
+		/// Gets this event data as a BSON document
+		/// </summary>
+		public IReadOnlyList<ILogEventLine> Lines => EventData.Lines;
+
+		/// <summary>
+		/// Fingerprint assigned to this issue
+		/// </summary>
+		public NewIssueFingerprint? Fingerprint { get; set; }
+
+		/// <summary>
+		/// Whether this event should be ignored
+		/// </summary>
+		public bool Ignored { get; set; }
+
+		/// <summary>
+		/// Constructor
+		/// </summary>
+		public IssueEvent(ILogEvent stepEvent, ILogEventData stepEventData)
+		{
+			Event = stepEvent;
+			EventData = stepEventData;
+		}
+
+		/// <inheritdoc/>
+		public override string ToString()
+		{
+			return $"[{Event.LineIndex}] {EventData.Message}";
+		}
+	}
+
+	/// <summary>
+	/// A group of <see cref="IssueEvent"/> objects with their fingerprint
+	/// </summary>
+	class IssueEventGroup
+	{
+		/// <summary>
+		/// Digest of the fingerprint, for log tracking
+		/// </summary>
+		public Md5Hash Digest { get; }
+
+		/// <summary>
+		/// Fingerprint for the event
+		/// </summary>
+		public NewIssueFingerprint Fingerprint { get; }
+
+		/// <summary>
+		/// Individual log events
+		/// </summary>
+		public List<IssueEvent> Events { get; } = new List<IssueEvent>();
+
+		/// <summary>
+		/// Constructor
+		/// </summary>
+		/// <param name="fingerprint">Fingerprint for the event</param>
+		public IssueEventGroup(NewIssueFingerprint fingerprint)
+		{
+			Digest = Md5Hash.Compute(Encoding.UTF8.GetBytes(fingerprint.ToString()));
+			Fingerprint = fingerprint;
+		}
+
+		/// <summary>
+		/// Merge with another group
+		/// </summary>
+		/// <param name="otherGroup">The group to merge with</param>
+		/// <returns>A new group combining both groups</returns>
+		public IssueEventGroup MergeWith(IssueEventGroup otherGroup)
+		{
+			IssueEventGroup newGroup = new IssueEventGroup(NewIssueFingerprint.Merge(Fingerprint, otherGroup.Fingerprint));
+			newGroup.Events.AddRange(Events);
+			newGroup.Events.AddRange(otherGroup.Events);
+			return newGroup;
+		}
+
+		/// <inheritdoc/>
+		public override string ToString() => Fingerprint.ToString();
+	}
+}
diff --git a/Engine/Source/Programs/Horde/Horde.Build/Issues/IssueHandler.cs b/Engine/Source/Programs/Horde/Horde.Build/Issues/IssueHandler.cs
new file mode 100644
--- /dev/null
+++ b/Engine/Source/Programs/Horde/Horde.Build/Issues/IssueHandler.cs
@@ -0,0 +1,48 @@
+// Copyright Epic Games, Inc. All Rights Reserved.
+
+using System.Collections.Generic;
+using Horde.Build.Jobs;
+using Horde.Build.Jobs.Graphs;
+
+namespace Horde.Build.Issues
+{
+	/// <summary>
+	/// Interface for issue matchers
+	/// </summary>
+	abstract class IssueHandler
+	{
+		/// <summary>
+		/// Identifier for the type of issue
+		/// </summary>
+		public abstract string Type { get; }
+
+		/// <summary>
+		/// Priority of this matcher
+		/// </summary>
+		public abstract int Priority { get; }
+
+		/// <summary>
+		/// Tag all thethe events for a step completing
+		/// </summary>
+		/// <param name="job">The job that spawned the event</param>
+		/// <param name="node">Node that was executed</param>
+		/// <param name="annotations">Combined annotations from this step</param>
+		/// <param name="events">Events from this step</param>
+		public abstract void TagEvents(IJob job, INode node, IReadOnlyNodeAnnotations annotations, IReadOnlyList<IssueEvent> events);
+
+		/// <summary>
+		/// Rank all the suspect changes for a given fingerprint
+		/// </summary>
+		/// <param name="fingerprint">The issue fingerprint</param>
+		/// <param name="suspects">Potential suspects</param>
+		public abstract void RankSuspects(IIssueFingerprint fingerprint, List<SuspectChange> suspects);
+
+		/// <summary>
+		/// Gets the summary text for an issue
+		/// </summary>
+		/// <param name="fingerprint">The fingerprint</param>
+		/// <param name="severity">Severity of the issue</param>
+		/// <returns>The summary text</returns>
+		public abstract string GetSummary(IIssueFingerprint fingerprint, IssueSeverity severity);
+	}
+}
diff --git a/Engine/Source/Programs/Horde/Horde.Build/Issues/SuspectChange.cs b/Engine/Source/Programs/Horde/Horde.Build/Issues/SuspectChange.cs
new file mode 100644
--- /dev/null
+++ b/Engine/Source/Programs/Horde/Horde.Build/Issues/SuspectChange.cs
@@ -0,0 +1,121 @@
+// Copyright Epic Games, Inc. All Rights Reserved.
+
+using System;
+using System.Collections.Generic;
+using System.Diagnostics.CodeAnalysis;
+using EpicGames.Core;
+using Horde.Build.Jobs;
+using Horde.Build.Jobs.Graphs;
+using Horde.Build.Logs;
+using Horde.Build.Perforce;
+
+namespace Horde.Build.Issues
+{
+	/// <summary>
+	/// Information about a changelist and a value ranking the likelihood that it caused an issue
+	/// </summary>
+	class SuspectChange
+	{
+		/// <summary>
+		/// Set of extensions to treat as code
+		/// </summary>
+		static readonly HashSet<StringView> s_codeExtensions = new HashSet<StringView>(StringViewComparer.OrdinalIgnoreCase)
+		{
+			".c",
+			".cc",
+			".cpp",
+			".inl",
+			".m",
+			".mm",
+			".rc",
+			".cs",
+			".csproj",
+			".h",
+			".hpp",
+			".inl",
+			".usf",
+			".ush",
+			".uproject",
+			".uplugin",
+			".sln"
+		};
+
+		/// <summary>
+		/// Set of file extensions to treat as content
+		/// </summary>
+		static readonly HashSet<StringView> s_contentExtensions = new HashSet<StringView>(StringViewComparer.OrdinalIgnoreCase)
+		{
+			".uasset",
+			".umap",
+			".ini"
+		};
+
+		/// <summary>
+		/// The change detials
+		/// </summary>
+		public ChangeDetails Details { get; set; }
+
+		/// <summary>
+		/// Rank for the change. A value less than or equal to zero indicates a lack of culpability, positive values indicate
+		/// the possibility of a change being the culprit.
+		/// </summary>
+		public int Rank { get; set; }
+
+		/// <summary>
+		/// Whether the change contains code
+		/// </summary>
+		public bool ContainsCode { get; }
+
+		/// <summary>
+		/// Whether the change modifies content
+		/// </summary>
+		public bool ContainsContent { get; }
+
+		/// <summary>
+		/// Constructor
+		/// </summary>
+		/// <param name="details">The changelist details</param>
+		public SuspectChange(ChangeDetails details)
+		{
+			Details = details;
+
+			foreach (ChangeFile file in details.Files)
+			{
+				int idx = file.Path.LastIndexOf('.');
+				if (idx != -1)
+				{
+					StringView extension = new StringView(file.Path, idx);
+					if (s_codeExtensions.Contains(extension))
+					{
+						ContainsCode = true;
+					}
+					if (s_contentExtensions.Contains(extension))
+					{
+						ContainsContent = true;
+					}
+					if (ContainsCode && ContainsContent)
+					{
+						break;
+					}
+				}
+			}
+		}
+
+		/// <summary>
+		/// Determines whether this change modifies the given file
+		/// </summary>
+		/// <param name="fileToCheck">The file to look for</param>
+		/// <returns>True if the change modifies the given file</returns>
+		public bool ModifiesFile(string fileToCheck)
+		{
+			foreach (ChangeFile file in Details.Files)
+			{
+				if (file.Path.EndsWith(fileToCheck, StringComparison.OrdinalIgnoreCase) && (file.Length == fileToCheck.Length || file.Path[file.Path.Length - fileToCheck.Length - 1] == '/'))
+				{
+					return true;
+				}
+			}
+			return false;
+		}
+	}
+}
