Horde: Add an issue handler that redacts and hashes error strings.

diff --git a/Engine/Source/Programs/Horde/Horde.Build.Tests/IssueServiceTests.cs b/Engine/Source/Programs/Horde/Horde.Build.Tests/IssueServiceTests.cs
--- a/Engine/Source/Programs/Horde/Horde.Build.Tests/IssueServiceTests.cs
+++ b/Engine/Source/Programs/Horde/Horde.Build.Tests/IssueServiceTests.cs
@@ -293,7 +293,7 @@ namespace Horde.Build.Tests
 			// Expected: Default issues is created
 			{
 				IJob job = CreateJob(_mainStreamId, 105, "Test Build", _graph);
-				await AddEvent(job, 0, 0, new { level = nameof(LogLevel.Warning) }, EventSeverity.Warning);
+				await AddEvent(job, 0, 0, new { level = nameof(LogLevel.Warning), message = "" }, EventSeverity.Warning);
 				await UpdateCompleteStep(job, 0, 0, JobStepOutcome.Warnings);
 
 				List<IIssue> issues = await IssueCollection.FindIssuesAsync();
@@ -307,9 +307,9 @@ namespace Horde.Build.Tests
 			// Expected: Nodes are NOT added to issue
 			{
 				IJob job = CreateJob(_mainStreamId, 105, "Test Build", _graph);
-				await AddEvent(job, 0, 1, new { });
+				await AddEvent(job, 0, 1, new { message = "" });
 				await UpdateCompleteStep(job, 0, 1, JobStepOutcome.Failure);
-				await AddEvent(job, 0, 2, new { });
+				await AddEvent(job, 0, 2, new { message = "" });
 				await UpdateCompleteStep(job, 0, 2, JobStepOutcome.Failure);
 
 				List<IIssue> issues = (await IssueCollection.FindIssuesAsync()).OrderBy(x => x.Summary).ToList();
@@ -325,7 +325,7 @@ namespace Horde.Build.Tests
 			// Expected: Nodes are added to issue, but change outcome to error
 			{
 				IJob job = CreateJob(_mainStreamId, 110, "Test Build", _graph);
-				await AddEvent(job, 0, 0, new { });
+				await AddEvent(job, 0, 0, new { message = "" });
 				await UpdateCompleteStep(job, 0, 0, JobStepOutcome.Failure);
 
 				List<IIssue> issues = (await IssueCollection.FindIssuesAsync()).OrderBy(x => x.Summary).ToList();
@@ -341,7 +341,7 @@ namespace Horde.Build.Tests
 			// Expected: Additional error is created
 			{
 				IJob job = CreateJob(_mainStreamId, 110, "Test Build", _graph);
-				await AddEvent(job, 0, 3, new { });
+				await AddEvent(job, 0, 3, new { message = "" });
 				await UpdateCompleteStep(job, 0, 3, JobStepOutcome.Warnings);
 
 				List<IIssue> issues = await IssueCollection.FindIssuesAsync();
@@ -917,6 +917,61 @@ namespace Horde.Build.Tests
 			}
 		}
 
+		[TestMethod]
+		public async Task HashedIssueTest()
+		{
+			IJob job = CreateJob(_mainStreamId, 120, "Compile Test", _graph);
+
+			await ParseEventsAsync(job, 0, 0, new[] { "LogSomething: Warning: This is a warning from the editor" });
+			await UpdateCompleteStep(job, 0, 0, JobStepOutcome.Failure);
+
+			await ParseEventsAsync(job, 0, 1, new[] { "LogSomething: Warning: This is a warning from the editor" });
+			await UpdateCompleteStep(job, 0, 1, JobStepOutcome.Failure);
+
+			List<IIssue> issues = await IssueCollection.FindIssuesAsync();
+			Assert.AreEqual(1, issues.Count);
+
+			IIssue issue = issues[0];
+			Assert.AreEqual("Warnings in Update Version Files and Compile UnrealHeaderTool Win64", issue.Summary);
+		}
+
+		[TestMethod]
+		public async Task HashedIssueTest2()
+		{
+			IJob job = CreateJob(_mainStreamId, 120, "Compile Test", _graph);
+
+			await ParseEventsAsync(job, 0, 0, new[] { "LogSomething: Warning: This is a warning from the editor" });
+			await UpdateCompleteStep(job, 0, 0, JobStepOutcome.Failure);
+
+			await ParseEventsAsync(job, 0, 1, new[] { "LogSomething: Warning: This is a warning from the editor2" });
+			await UpdateCompleteStep(job, 0, 1, JobStepOutcome.Failure);
+
+			List<IIssue> issues = await IssueCollection.FindIssuesAsync();
+			issues.SortBy(x => x.Id);
+			Assert.AreEqual(2, issues.Count);
+
+			Assert.AreEqual("Warnings in Update Version Files", issues[0].Summary);
+			Assert.AreEqual("Warnings in Compile UnrealHeaderTool Win64", issues[1].Summary);
+		}
+
+		[TestMethod]
+		public async Task HashedIssueTest3()
+		{
+			IJob job = CreateJob(_mainStreamId, 120, "Compile Test", _graph);
+
+			await ParseEventsAsync(job, 0, 0, new[] { "LogSomething: Warning: This is a warning from the editor", "warning: some generic thing that will be ignored" });
+			await UpdateCompleteStep(job, 0, 0, JobStepOutcome.Failure);
+
+			await ParseEventsAsync(job, 0, 1, new[] { "LogSomething: Warning: This is a warning from the editor" });
+			await UpdateCompleteStep(job, 0, 1, JobStepOutcome.Failure);
+
+			List<IIssue> issues = await IssueCollection.FindIssuesAsync();
+			Assert.AreEqual(1, issues.Count);
+
+			IIssue issue = issues[0];
+			Assert.AreEqual("Warnings in Update Version Files and Compile UnrealHeaderTool Win64", issue.Summary);
+		}
+
 		[TestMethod]
 		public async Task SymbolIssueTest()
 		{
diff --git a/Engine/Source/Programs/Horde/Horde.Build/Issues/Handlers/HashedIssueHandler.cs b/Engine/Source/Programs/Horde/Horde.Build/Issues/Handlers/HashedIssueHandler.cs
new file mode 100644
--- /dev/null
+++ b/Engine/Source/Programs/Horde/Horde.Build/Issues/Handlers/HashedIssueHandler.cs
@@ -0,0 +1,79 @@
+// Copyright Epic Games, Inc. All Rights Reserved.
+
+using System.Collections.Generic;
+using System.Linq;
+using System.Text;
+using System.Text.RegularExpressions;
+using EpicGames.Core;
+using Horde.Build.Jobs;
+using Horde.Build.Jobs.Graphs;
+
+namespace Horde.Build.Issues.Handlers
+{
+	/// <summary>
+	/// Instance of a particular compile error
+	/// </summary>
+	class HashedIssueHandler : IssueHandler
+	{
+		const string NodeName = "Node";
+
+		/// <inheritdoc/>
+		public override string Type => "Hashed";
+
+		/// <inheritdoc/>
+		public override int Priority => 1;
+
+		/// <inheritdoc/>
+		public override void TagEvents(IJob job, INode node, IReadOnlyNodeAnnotations annotations, IReadOnlyList<IssueEvent> stepEvents)
+		{
+			string[] metadata = new[] { $"{NodeName}={node.Name}" };
+
+			NewIssueFingerprint? genericFingerprint = null;
+			HashSet<Md5Hash> hashes = new HashSet<Md5Hash>();
+
+			foreach (IssueEvent stepEvent in stepEvents)
+			{
+				if (hashes.Count < 25 && TryGetHash(stepEvent.Message, out Md5Hash hash))
+				{
+					hashes.Add(hash);
+					stepEvent.Fingerprint = new NewIssueFingerprint(Type, new[] { $"hash:{hash}" }, null, metadata);
+				}
+				else
+				{
+					genericFingerprint ??= new NewIssueFingerprint(Type, new[] { $"step:{job.StreamId}:{job.TemplateId}:{node.Name}" }, null, metadata);
+					stepEvent.Fingerprint = genericFingerprint;
+				}
+			}
+		}
+
+		static bool TryGetHash(string message, out Md5Hash hash)
+		{
+			string sanitized = message.ToUpperInvariant();
+			sanitized = Regex.Replace(sanitized, @"\d[\d.,:]*", "n"); // Redact numbers and timestamp like things
+
+			if (sanitized.Length > 30)
+			{
+				hash = Md5Hash.Compute(Encoding.UTF8.GetBytes(sanitized));
+				return true;
+			}
+			else
+			{
+				hash = Md5Hash.Zero;
+				return false;
+			}
+		}
+
+		/// <inheritdoc/>
+		public override void RankSuspects(IIssueFingerprint fingerprint, List<SuspectChange> suspects)
+		{
+		}
+
+		/// <inheritdoc/>
+		public override string GetSummary(IIssueFingerprint fingerprint, IssueSeverity severity)
+		{
+			string severityText = (severity == IssueSeverity.Warning) ? "Warnings" : "Errors";
+			string[] nodeNames = fingerprint.GetMetadataValues(NodeName).ToArray();
+			return $"{severityText} in {StringUtils.FormatList(nodeNames, 2)}";
+		}
+	}
+}
