Horde: Fix test.

diff --git a/Engine/Source/Programs/Horde/Horde.Build.Tests/IssueServiceTests.cs b/Engine/Source/Programs/Horde/Horde.Build.Tests/IssueServiceTests.cs
--- a/Engine/Source/Programs/Horde/Horde.Build.Tests/IssueServiceTests.cs
+++ b/Engine/Source/Programs/Horde/Horde.Build.Tests/IssueServiceTests.cs
@@ -1143,7 +1143,7 @@ namespace Horde.Build.Tests
 				List<IIssueSpan> spans = await IssueCollection.FindSpansAsync(issue.Id);
 				Assert.AreEqual(1, spans.Count);
 				Assert.AreEqual(1, issue.Fingerprints.Count);
-				Assert.AreEqual("Default", issue.Fingerprints[0].Type);
+				Assert.AreEqual("Hashed", issue.Fingerprints[0].Type);
 
 				IIssueSpan span = spans[0];
 				Assert.AreEqual(120, span.LastSuccess?.Change);
