Horde: Allow sending job complete notifications to multiple slack channels (using ";" as a separator)

diff --git a/Engine/Source/Programs/Horde/Horde.Build/Notifications/Sinks/SlackNotificationSink.cs b/Engine/Source/Programs/Horde/Horde.Build/Notifications/Sinks/SlackNotificationSink.cs
--- a/Engine/Source/Programs/Horde/Horde.Build/Notifications/Sinks/SlackNotificationSink.cs
+++ b/Engine/Source/Programs/Horde/Horde.Build/Notifications/Sinks/SlackNotificationSink.cs
@@ -327,7 +327,10 @@ namespace Horde.Build.Notifications.Sinks
 			{
 				string jobIds = String.Join(", ", notifications.Select(x => x.JobId));
 				_logger.LogInformation("Sending Slack notification for scheduled job IDs {JobIds} to channel {SlackChannel}", jobIds, _settings.JobNotificationChannel);
-				await SendJobScheduledOnEmptyAutoScaledPoolMessageAsync($"#{_settings.JobNotificationChannel}", notifications);
+				foreach (string channel in _settings.JobNotificationChannel.Split(';'))
+				{
+					await SendJobScheduledOnEmptyAutoScaledPoolMessageAsync($"#{channel}", notifications);
+				}
 			}
 		}
 		
@@ -369,7 +372,7 @@ namespace Horde.Build.Notifications.Sinks
 			}
 		}
 
-		Task SendJobCompleteNotificationToChannelAsync(string notificationChannel, string? notificationFilter, IStream jobStream, IJob job, IGraph graph, LabelOutcome outcome)
+		async Task SendJobCompleteNotificationToChannelAsync(string notificationChannel, string? notificationFilter, IStream jobStream, IJob job, IGraph graph, LabelOutcome outcome)
 		{
 			if (notificationFilter != null)
 			{
@@ -388,10 +391,13 @@ namespace Horde.Build.Notifications.Sinks
 				}
 				if (!outcomes.Contains(outcome))
 				{
-					return Task.CompletedTask;
+					return;
 				}
 			}
-			return SendJobCompleteMessageAsync(notificationChannel, jobStream, job, graph);
+			foreach (string channel in notificationChannel.Split(';'))
+			{
+				await SendJobCompleteMessageAsync(channel, jobStream, job, graph);
+			}
 		}
 
 		/// <inheritdoc/>
