Revert change to treat localization paths as relative to the workspace root directory.

This change was prompted by copying output from the Horde dashboard, which does not reflect the actual output from the engine.

diff --git a/Engine/Source/Programs/AutomationTool/AutomationUtils/Matchers/LocalizationEventMatcher.cs b/Engine/Source/Programs/AutomationTool/AutomationUtils/Matchers/LocalizationEventMatcher.cs
--- a/Engine/Source/Programs/AutomationTool/AutomationUtils/Matchers/LocalizationEventMatcher.cs
+++ b/Engine/Source/Programs/AutomationTool/AutomationUtils/Matchers/LocalizationEventMatcher.cs
@@ -55,7 +55,7 @@ namespace AutomationUtils.Matchers
 		{
 			builder.Annotate(match.Groups["channel"], LogEventMarkup.Channel);
 			builder.Annotate(match.Groups["severity"], LogEventMarkup.Severity);
-			builder.AnnotateSourceFile(match.Groups["file"], null);
+			builder.AnnotateSourceFile(match.Groups["file"], "Engine");
 			builder.Annotate(match.Groups["line"], LogEventMarkup.LineNumber);
 		}
 	}
diff --git a/Engine/Source/Programs/Horde/Horde.Agent.Tests/LogParserTests.cs b/Engine/Source/Programs/Horde/Horde.Agent.Tests/LogParserTests.cs
--- a/Engine/Source/Programs/Horde/Horde.Agent.Tests/LogParserTests.cs
+++ b/Engine/Source/Programs/Horde/Horde.Agent.Tests/LogParserTests.cs
@@ -720,9 +720,9 @@ namespace Horde.Agent.Tests
 		{
 			string[] lines =
 			{
-				@"[2022.05.31-04.25.40:235][  0]LogGatherTextFromSourceCommandlet: Warning: Engine/Plugins/Enterprise/VariantManager/Source/VariantManager/Private/SVariantManager.cpp(3717): LOCTEXT macro has an empty identifier and cannot be gathered.",
-				@"[2022.06.01-04.29.09:630][  0]LogLocTextHelper: Warning: Engine/Plugins/Experimental/UVEditor/Source/UVEditor/Private/UVEditorCommands.cpp(39): Text conflict from UI_COMMAND macro for namespace ""UICommands.FUVEditorCommands"" and key ""SplitAction_ToolTip"". The conflicting sources are ""Given an edge selection, split those edges. Given a vertex selection, split any selected bowtie vertices. Given a triangle selection, split along selection boundaries."" and ""Given an edge selection, split those edges. Given a vertex selection, split any selected bowtie vertices."".",
-				@"[2022.06.01-04.29.09:630][  0]LogLocTextHelper: Warning: Engine/Content/Localization/Engine/Engine.manifest: See conflicting location.",
+				@"[2022.05.31-04.25.40:235][  0]LogGatherTextFromSourceCommandlet: Warning: Plugins/Enterprise/VariantManager/Source/VariantManager/Private/SVariantManager.cpp(3717): LOCTEXT macro has an empty identifier and cannot be gathered.",
+				@"[2022.06.01-04.29.09:630][  0]LogLocTextHelper: Warning: Plugins/Experimental/UVEditor/Source/UVEditor/Private/UVEditorCommands.cpp(39): Text conflict from UI_COMMAND macro for namespace ""UICommands.FUVEditorCommands"" and key ""SplitAction_ToolTip"". The conflicting sources are ""Given an edge selection, split those edges. Given a vertex selection, split any selected bowtie vertices. Given a triangle selection, split along selection boundaries."" and ""Given an edge selection, split those edges. Given a vertex selection, split any selected bowtie vertices."".",
+				@"[2022.06.01-04.29.09:630][  0]LogLocTextHelper: Warning: Content/Localization/Engine/Engine.manifest: See conflicting location.",
 			};
 
 			{
