Horde: Add agent operations for viewing pool, filtering to pool, and copying pool name to clipboard
#rnx
#skipci

diff --git a/Engine/Source/Programs/Horde/HordeDashboard/src/base/utilities/clipboard.ts b/Engine/Source/Programs/Horde/HordeDashboard/src/base/utilities/clipboard.ts
new file mode 100644
--- /dev/null
+++ b/Engine/Source/Programs/Horde/HordeDashboard/src/base/utilities/clipboard.ts
@@ -0,0 +1,15 @@
+
+
+export const copyToClipboard = (value: string | undefined) => {
+
+    if (!value) {
+        return;
+    }
+
+    const el = document.createElement('textarea');
+    el.value = value;
+    document.body.appendChild(el);
+    el.select();
+    document.execCommand('copy');
+    document.body.removeChild(el);
+}
