Horde: Allow overriding the grouping of issues specifically for compile issues. This will allow generic errors to continue to merge together, while putting compilation errors for (eg.) static analysis in a separate group.

diff --git a/Engine/Source/Programs/Horde/Horde.Build/Issues/Handlers/CompileIssueHandler.cs b/Engine/Source/Programs/Horde/Horde.Build/Issues/Handlers/CompileIssueHandler.cs
--- a/Engine/Source/Programs/Horde/Horde.Build/Issues/Handlers/CompileIssueHandler.cs
+++ b/Engine/Source/Programs/Horde/Horde.Build/Issues/Handlers/CompileIssueHandler.cs
@@ -23,6 +23,11 @@ namespace Horde.Build.Issues.Handlers
 		/// </summary>
 		const string CompileTypeAnnotation = "CompileType";
 
+		/// <summary>
+		/// Annotation specifying a group for compile issues from this node
+		/// </summary>
+		const string CompileGroupAnnotation = "CompileGroup";
+
 		/// <inheritdoc/>
 		public override string Type => "Compile";
 
@@ -49,16 +54,22 @@ namespace Horde.Build.Issues.Handlers
 					List<string> newFileNames = new List<string>();
 					GetSourceFiles(stepEvent.EventData, newFileNames);
 
-					string? compileType;
-					if (!annotations.TryGetValue(CompileTypeAnnotation, out compileType))
+					string compileType = "Compile";
+					if (annotations.TryGetValue(CompileTypeAnnotation, out string? type))
+					{
+						compileType = type;
+					}
+
+					string fingerprintType = Type;
+					if (annotations.TryGetValue(CompileGroupAnnotation, out string? group))
 					{
-						compileType = "Compile";
+						fingerprintType = $"{fingerprintType}:{group}";
 					}
 
 					List<string> newMetadata = new List<string>();
 					newMetadata.Add($"{CompileTypeAnnotation}={compileType}");
 
-					stepEvent.Fingerprint = new NewIssueFingerprint(Type, newFileNames, null, newMetadata);
+					stepEvent.Fingerprint = new NewIssueFingerprint(fingerprintType, newFileNames, null, newMetadata);
 				}
 			}
 		}
