Horde: Add a red cross reaction to threads when a fix has failed.

diff --git a/Engine/Source/Programs/Horde/Horde.Build/Notifications/Sinks/SlackNotificationSink.cs b/Engine/Source/Programs/Horde/Horde.Build/Notifications/Sinks/SlackNotificationSink.cs
--- a/Engine/Source/Programs/Horde/Horde.Build/Notifications/Sinks/SlackNotificationSink.cs
+++ b/Engine/Source/Programs/Horde/Horde.Build/Notifications/Sinks/SlackNotificationSink.cs
@@ -841,20 +841,10 @@ namespace Horde.Build.Notifications.Sinks
 					await RemoveReactionAsync(state.Channel, state.Ts, "eyes");
 				}
 
-				if (issue.ResolvedAt != null)
-				{
-					await AddReactionAsync(state.Channel, state.Ts, "tick");
-				}
-				else
-				{
-					await RemoveReactionAsync(state.Channel, state.Ts, "tick");
-				}
-
+				IIssueSpan? fixFailedSpan = null;
 				if (issue.FixChange != null)
 				{
 					HashSet<StreamId> originStreams = new HashSet<StreamId>(issue.Streams.Where(x => x.MergeOrigin ?? false).Select(x => x.StreamId));
-
-					IIssueSpan? fixFailedSpan = null;
 					foreach (IIssueSpan originSpan in spans.Where(x => originStreams.Contains(x.StreamId)).OrderBy(x => x.StreamId.ToString()))
 					{
 						if (originSpan.LastFailure.Change > issue.FixChange.Value)
@@ -897,7 +887,25 @@ namespace Horde.Build.Notifications.Sinks
 					}
 				}
 
-				if(issue.ExternalIssueKey != null)
+				if (fixFailedSpan != null)
+				{
+					await AddReactionAsync(state.Channel, state.Ts, "x");
+				}
+				else
+				{
+					await RemoveReactionAsync(state.Channel, state.Ts, "x");
+				}
+
+				if (issue.ResolvedAt != null && fixFailedSpan == null)
+				{
+					await AddReactionAsync(state.Channel, state.Ts, "tick");
+				}
+				else
+				{
+					await RemoveReactionAsync(state.Channel, state.Ts, "tick");
+				}
+
+				if (issue.ExternalIssueKey != null)
 				{
 					string extIssueEventId = $"issue_{issue.Id}_ext_{issue.ExternalIssueKey}";
 					string extIssueMessage = $"Linked to issue {FormatExternalIssue(issue.ExternalIssueKey)}";
