BuildGraph: Reformat output from Docker to allow regular pattern matchers to work correctly, and add an MSTest event matcher.

diff --git a/Engine/Source/Programs/AutomationTool/AutomationUtils/Matchers/MsTestEventMatcher.cs b/Engine/Source/Programs/AutomationTool/AutomationUtils/Matchers/MsTestEventMatcher.cs
new file mode 100644
--- /dev/null
+++ b/Engine/Source/Programs/AutomationTool/AutomationUtils/Matchers/MsTestEventMatcher.cs
@@ -0,0 +1,39 @@
+// Copyright Epic Games, Inc. All Rights Reserved.
+
+using EpicGames.Core;
+using Microsoft.Extensions.Logging;
+using System.Text.RegularExpressions;
+
+#nullable enable
+
+namespace AutomationUtils.Matchers
+{
+	class MsTestEventMatcher : ILogEventMatcher
+	{
+		static readonly Regex s_failedPattern = new Regex(@"^  \s*Failed [A-Za-z0-9_]+ \[\d+ ms\]$");
+		static readonly Regex s_detailPattern = new Regex(@"^  \s*(?:Error Message|Stack Trace):$");
+
+		public LogEventMatch? Match(ILogCursor cursor)
+		{
+			if (cursor.IsMatch(s_failedPattern))
+			{
+				int lineCount = 1;
+				while (cursor.IsMatch(lineCount, s_detailPattern) || cursor.IsHanging(lineCount, cursor.CurrentLine))
+				{
+					lineCount++;
+				}
+				while (lineCount > 0 && cursor.IsBlank(lineCount - 1))
+				{
+					lineCount--;
+				}
+
+				if (lineCount > 1)
+				{
+					LogEventBuilder builder = new LogEventBuilder(cursor, lineCount);
+					return builder.ToMatch(LogEventPriority.High, LogLevel.Error, KnownLogEvents.MSTest);
+				}
+			}
+			return null;
+		}
+	}
+}
diff --git a/Engine/Source/Programs/AutomationTool/BuildGraph/Tasks/DockerBuildTask.cs b/Engine/Source/Programs/AutomationTool/BuildGraph/Tasks/DockerBuildTask.cs
--- a/Engine/Source/Programs/AutomationTool/BuildGraph/Tasks/DockerBuildTask.cs
+++ b/Engine/Source/Programs/AutomationTool/BuildGraph/Tasks/DockerBuildTask.cs
@@ -10,6 +10,7 @@ using System.Xml;
 using AutomationTool;
 using UnrealBuildBase;
 using System.Threading.Tasks;
+using System.Text.RegularExpressions;
 
 namespace AutomationTool.Tasks
 {
@@ -172,10 +173,14 @@ namespace AutomationTool.Tasks
 				}
 				
 				string WorkingDir = isStagingEnabled ? StagingDir.FullName : BaseDir.FullName;
-				await SpawnTaskBase.ExecuteAsync("docker", Arguments.ToString(), EnvVars: EnvVars, WorkingDir: WorkingDir);
+				await SpawnTaskBase.ExecuteAsync("docker", Arguments.ToString(), EnvVars: EnvVars, WorkingDir: WorkingDir, SpewFilterCallback: FilterOutput);
 			}
 		}
 
+		static Regex FilterOutputPattern = new Regex(@"^#\d+ (?:\d+\.\d+ )?");
+
+		static string FilterOutput(string Line) => FilterOutputPattern.Replace(Line, "");
+
 		/// <summary>
 		/// Output this task out to an XML writer.
 		/// </summary>
diff --git a/Engine/Source/Programs/AutomationTool/BuildGraph/Tasks/SpawnTask.cs b/Engine/Source/Programs/AutomationTool/BuildGraph/Tasks/SpawnTask.cs
--- a/Engine/Source/Programs/AutomationTool/BuildGraph/Tasks/SpawnTask.cs
+++ b/Engine/Source/Programs/AutomationTool/BuildGraph/Tasks/SpawnTask.cs
@@ -71,7 +71,7 @@ namespace AutomationTool.Tasks
 		/// <summary>
 		/// Execute a command
 		/// </summary>
-		protected static Task<IProcessResult> ExecuteAsync(string Exe, string Arguments, string WorkingDir = null, Dictionary<string, string> EnvVars = null, bool LogOutput = true, int ErrorLevel = 1, string Input = null)
+		protected static Task<IProcessResult> ExecuteAsync(string Exe, string Arguments, string WorkingDir = null, Dictionary<string, string> EnvVars = null, bool LogOutput = true, int ErrorLevel = 1, string Input = null, ProcessResult.SpewFilterCallbackType SpewFilterCallback = null)
 		{
 			if (WorkingDir != null)
 			{
@@ -84,7 +84,7 @@ namespace AutomationTool.Tasks
 				Options &= ~CommandUtils.ERunOptions.AllowSpew;
 			}
 
-			IProcessResult Result = CommandUtils.Run(Exe, Arguments, Env: EnvVars, WorkingDir: WorkingDir, Options: Options, Input: Input);
+			IProcessResult Result = CommandUtils.Run(Exe, Arguments, Env: EnvVars, WorkingDir: WorkingDir, Options: Options, Input: Input, SpewFilterCallback: SpewFilterCallback);
 			if (Result.ExitCode < 0 || Result.ExitCode >= ErrorLevel)
 			{
 				throw new AutomationException("{0} terminated with an exit code indicating an error ({1})", Path.GetFileName(Exe), Result.ExitCode);
diff --git a/Engine/Source/Programs/Horde/Horde.Agent.Tests/LogParserTests.cs b/Engine/Source/Programs/Horde/Horde.Agent.Tests/LogParserTests.cs
--- a/Engine/Source/Programs/Horde/Horde.Agent.Tests/LogParserTests.cs
+++ b/Engine/Source/Programs/Horde/Horde.Agent.Tests/LogParserTests.cs
@@ -793,6 +793,39 @@ namespace Horde.Agent.Tests
 			Assert.IsFalse(buffer.IsHanging(4, firstIndentedLine));
 		}
 
+		[TestMethod]
+		public void MsTestEventMatcher()
+		{
+			string[] lines =
+			{
+				@"Microsoft (R) Test Execution Command Line Tool Version 17.2.0 (x64)",
+				@"Copyright (c) Microsoft Corporation.  All rights reserved.",
+				@"Starting test execution, please wait...",
+				@"A total of 1 test files matched the specified pattern.",
+				@"",
+				@"  Failed ExplicitGroupingTest [219 ms]",
+				@"  Error Message:",
+				@"   Assert.AreEqual failed. Expected:<2>. Actual:<1>.",
+				@"  Stack Trace:",
+				@"     at Horde.Build.Tests.IssueServiceTests.ExplicitGroupingTest() in /app/Source/Programs/Horde/Horde.Build.Tests/IssueServiceTests.cs:line 1382",
+				@"   at Microsoft.VisualStudio.TestPlatform.MSTestAdapter.PlatformServices.ThreadOperations.ExecuteWithAbortSafety(Action action)",
+				@"",
+				@"  Standard Output Messages:",
+				@" info: Horde.Build.Issues.IssueService[0]",
+				@"       Updating issues for 62a21de84f3b4344e94f3f10:0000:0004",
+				@" info: Horde.Build.Logs.LogEventCollection[0]",
+				@"       Querying for log events for log 62a21de84f3b4344e94f3f15 creation time 06/09/2022 16:20:56",
+				@" info: Horde.Build.Auditing.AuditLog[0]",
+				@"       Created issue 1",
+				@" info: Horde.Build.Auditing.AuditLog[0]",
+				@"       Changed severity to Error",
+			};
+
+			List<LogEvent> logEvents = Parse(lines);
+			Assert.AreEqual(6, logEvents.Count);
+			CheckEventGroup(logEvents, 5, 6, LogLevel.Error, KnownLogEvents.MSTest);
+		}
+
 		static List<LogEvent> Parse(IEnumerable<string> lines)
 		{
 			return Parse(String.Join("\n", lines));
diff --git a/Engine/Source/Programs/Shared/EpicGames.Core/ILogCursor.cs b/Engine/Source/Programs/Shared/EpicGames.Core/ILogCursor.cs
--- a/Engine/Source/Programs/Shared/EpicGames.Core/ILogCursor.cs
+++ b/Engine/Source/Programs/Shared/EpicGames.Core/ILogCursor.cs
@@ -228,6 +228,30 @@ namespace EpicGames.Core
 			return offset;
 		}
 
+		/// <summary>
+		/// Tests if a line consists only of whitespace
+		/// </summary>
+		/// <param name="cursor">The log cursor</param>
+		/// <param name="offset">Offset of the line to check</param>
+		/// <returns></returns>
+		public static bool IsBlank(this ILogCursor cursor, int offset)
+		{
+			if (!cursor.TryGetLine(offset, out string? line))
+			{
+				return false;
+			}
+
+			for (int idx = 0; idx < line.Length; idx++)
+			{
+				if (!Char.IsWhiteSpace(line[idx]))
+				{
+					return false;
+				}
+			}
+
+			return true;
+		}
+
 		/// <summary>
 		/// Check if a line at a given offset is left-aligned with at least this indent of another line
 		/// </summary>
diff --git a/Engine/Source/Programs/Shared/EpicGames.Core/KnownLogEvents.cs b/Engine/Source/Programs/Shared/EpicGames.Core/KnownLogEvents.cs
--- a/Engine/Source/Programs/Shared/EpicGames.Core/KnownLogEvents.cs
+++ b/Engine/Source/Programs/Shared/EpicGames.Core/KnownLogEvents.cs
@@ -115,6 +115,11 @@ namespace EpicGames.Core
 		/// </summary>
 		public static EventId MSBuild { get; } = new EventId(500);
 
+		/// <summary>
+		/// Microsoft: Generic MSTest error
+		/// </summary>
+		public static EventId MSTest { get; } = new EventId(540);
+
 		/// <summary>
 		/// Microsoft: Generic Visual Studio error
 		/// </summary>
diff --git a/Engine/Source/Programs/Shared/EpicGames.Core/LogEventBuilder.cs b/Engine/Source/Programs/Shared/EpicGames.Core/LogEventBuilder.cs
--- a/Engine/Source/Programs/Shared/EpicGames.Core/LogEventBuilder.cs
+++ b/Engine/Source/Programs/Shared/EpicGames.Core/LogEventBuilder.cs
@@ -360,7 +360,7 @@ namespace EpicGames.Core
 			}
 			else
 			{
-				return _lines[^1]._message;
+				return _lines[0]._message;
 			}
 		}
 
