Horde: Missing file.

diff --git a/Engine/Source/Programs/Horde/Horde.Build/Issues/IssueService.cs b/Engine/Source/Programs/Horde/Horde.Build/Issues/IssueService.cs
--- a/Engine/Source/Programs/Horde/Horde.Build/Issues/IssueService.cs
+++ b/Engine/Source/Programs/Horde/Horde.Build/Issues/IssueService.cs
@@ -2,6 +2,7 @@
 
 using System;
 using System.Collections.Generic;
+using System.Diagnostics.CodeAnalysis;
 using System.Globalization;
 using System.Linq;
 using System.Reflection;
@@ -589,6 +590,20 @@ namespace Horde.Build.Issues
 			foreach (IssueHandler handler in _handlers)
 			{
 				handler.TagEvents(job, node, annotations, remainingEvents);
+
+				bool ignoreSystemicEvents = false;
+				foreach (IssueEvent remainingEvent in remainingEvents)
+				{
+					if (remainingEvent.Fingerprint != null && remainingEvent.Severity == EventSeverity.Error)
+					{
+						ignoreSystemicEvents = true;
+					}
+					else if (ignoreSystemicEvents && remainingEvent.IsSystemic())
+					{
+						remainingEvent.Ignored = true;
+					}
+				}
+
 				remainingEvents.RemoveAll(x => x.Ignored || x.Fingerprint != null);
 			}
 
@@ -608,6 +623,20 @@ namespace Horde.Build.Issues
 				}
 			}
 
+			// If the node has an annotation to prevent grouping issues together, update all the fingerprints to match
+			string? group = annotations.IssueGroup;
+			if (group != null)
+			{
+				Dictionary<NewIssueFingerprint, IssueEventGroup> newFingerprintToEventGroup = new Dictionary<NewIssueFingerprint, IssueEventGroup>(fingerprintToEventGroup.Count);
+				foreach ((NewIssueFingerprint fingerprint, IssueEventGroup eventGroup) in fingerprintToEventGroup)
+				{
+					IssueEventGroup newEventGroup = new IssueEventGroup(new NewIssueFingerprint($"{fingerprint.Type}:{group}", fingerprint.Keys, fingerprint.RejectKeys, fingerprint.Metadata));
+					newEventGroup.Events.AddRange(eventGroup.Events);
+					newFingerprintToEventGroup[newEventGroup.Fingerprint] = newEventGroup;
+				}
+				fingerprintToEventGroup = newFingerprintToEventGroup;
+			}
+
 			// Print the list of new events
 			HashSet<IssueEventGroup> eventGroups = new HashSet<IssueEventGroup>(fingerprintToEventGroup.Values);
 
@@ -1055,7 +1084,7 @@ namespace Horde.Build.Issues
 		async Task<List<NewIssueSpanSuspectData>> FindSuspectsForSpanAsync(IStream stream, IIssueFingerprint fingerprint, int minChange, int maxChange)
 		{
 			List<NewIssueSpanSuspectData> suspects = new List<NewIssueSpanSuspectData>();
-			if (_typeToHandler.TryGetValue(fingerprint.Type, out IssueHandler? handler))
+			if (TryGetHandler(fingerprint, out IssueHandler? handler))
 			{
 				_logger.LogDebug("Querying for changes in {StreamName} between {MinChange} and {MaxChange}", stream.Name, minChange, maxChange);
 
@@ -1157,7 +1186,7 @@ namespace Horde.Build.Issues
 		/// <returns>The summary text</returns>
 		string GetSummary(IIssueFingerprint fingerprint, IssueSeverity severity)
 		{
-			if (_typeToHandler.TryGetValue(fingerprint.Type, out IssueHandler? handler))
+			if (TryGetHandler(fingerprint, out IssueHandler? handler))
 			{
 				return handler.GetSummary(fingerprint, severity);
 			}
@@ -1246,5 +1275,24 @@ namespace Horde.Build.Issues
 				return null;
 			}
 		}
+
+		/// <summary>
+		/// Finds the handler for a given issue
+		/// </summary>
+		/// <param name="fingerprint">The fingerprint to get the handler for</param>
+		/// <param name="handler">Receives the handler on success</param>
+		/// <returns>True if a matching handler was found</returns>
+		bool TryGetHandler(IIssueFingerprint fingerprint, [NotNullWhen(true)] out IssueHandler? handler)
+		{
+			string type = fingerprint.Type;
+
+			int endIdx = type.IndexOf(':', StringComparison.Ordinal);
+			if (endIdx != -1)
+			{
+				type = type.Substring(0, endIdx);
+			}
+
+			return _typeToHandler.TryGetValue(type, out handler);
+		}
 	}
 }
