Horde: Expose step name as an env var during job execution

diff --git a/Engine/Source/Programs/Horde/Horde.Build/Server/RpcService.cs b/Engine/Source/Programs/Horde/Horde.Build/Server/RpcService.cs
--- a/Engine/Source/Programs/Horde/Horde.Build/Server/RpcService.cs
+++ b/Engine/Source/Programs/Horde/Horde.Build/Server/RpcService.cs
@@ -622,6 +622,7 @@ namespace Horde.Build.Server
 				
 				response.EnvVars.Add("UE_HORDE_TEMPLATEID", job.TemplateId.ToString());
 				response.EnvVars.Add("UE_HORDE_TEMPLATENAME", templateName);
+				response.EnvVars.Add("UE_HORDE_STEPNAME", node.Name);
 
 				IJobStepRef? lastStep = await _jobStepRefCollection.GetPrevStepForNodeAsync(job.StreamId, job.TemplateId, node.Name, job.Change);
 				if (lastStep != null)
