Horde: Remove redundant job data property on GetJobTimingResponse
#rnx

diff --git a/Engine/Source/Programs/Horde/Horde.Build/Jobs/JobMessages.cs b/Engine/Source/Programs/Horde/Horde.Build/Jobs/JobMessages.cs
--- a/Engine/Source/Programs/Horde/Horde.Build/Jobs/JobMessages.cs
+++ b/Engine/Source/Programs/Horde/Horde.Build/Jobs/JobMessages.cs
@@ -386,17 +386,12 @@ namespace Horde.Build.Jobs
 	}
 
 	/// <summary>
-	/// The timing info for 
+	/// The timing info for a job
 	/// </summary>
 	public class GetJobTimingResponse
-	{
+	{		
 		/// <summary>
-		/// The job object
-		/// </summary>
-		public IJob? Job { get; set; }
-		
-		/// <summary>
-		/// A full job response
+		/// The job response
 		/// </summary>
 		public GetJobResponse? JobResponse { get; set; }
 		
@@ -412,14 +407,12 @@ namespace Horde.Build.Jobs
 
 		/// <summary>
 		/// Constructor
-		/// </summary>
-		/// <param name="job">The job object</param>
+		/// </summary>		
 		/// <param name="jobResponse">The job response</param>
 		/// <param name="steps">Timing info for each steps</param>
 		/// <param name="labels">Timing info for each label</param>
-		public GetJobTimingResponse(IJob? job, GetJobResponse? jobResponse, Dictionary<string, GetStepTimingInfoResponse> steps, List<GetLabelTimingInfoResponse> labels)
-		{
-			Job = job;
+		public GetJobTimingResponse(GetJobResponse? jobResponse, Dictionary<string, GetStepTimingInfoResponse> steps, List<GetLabelTimingInfoResponse> labels)
+		{			
 			JobResponse = jobResponse;
 			Steps = steps;
 			Labels = labels;
diff --git a/Engine/Source/Programs/Horde/Horde.Build/Jobs/JobsController.cs b/Engine/Source/Programs/Horde/Horde.Build/Jobs/JobsController.cs
--- a/Engine/Source/Programs/Horde/Horde.Build/Jobs/JobsController.cs
+++ b/Engine/Source/Programs/Horde/Horde.Build/Jobs/JobsController.cs
@@ -582,7 +582,7 @@ namespace Horde.Build.Jobs
 				jobResponse = await CreateJobResponseAsync(job, graph, true, true, false, true);
 			}
 
-			return new GetJobTimingResponse(job, jobResponse, steps, labels);
+			return new GetJobTimingResponse(jobResponse, steps, labels);
 		}
 		
 		/// <summary>
diff --git a/Engine/Source/Programs/Horde/HordeDashboard/src/backend/Api.ts b/Engine/Source/Programs/Horde/HordeDashboard/src/backend/Api.ts
--- a/Engine/Source/Programs/Horde/HordeDashboard/src/backend/Api.ts
+++ b/Engine/Source/Programs/Horde/HordeDashboard/src/backend/Api.ts
@@ -2512,11 +2512,12 @@ export type GetGraphResponse = {
 
 }
 
-/**The timing info for  */
+/**The timing info for a job*/
 export type GetJobTimingResponse = {
-
-	job: JobData;
+	
+	/** The job response */
 	jobResponse: JobData;
+
 	/**Timing info for each step */
 	steps: { [key: string]: GetStepTimingInfoResponse };
 
diff --git a/Engine/Source/Programs/Horde/HordeDashboard/src/backend/Backend.ts b/Engine/Source/Programs/Horde/HordeDashboard/src/backend/Backend.ts
--- a/Engine/Source/Programs/Horde/HordeDashboard/src/backend/Backend.ts
+++ b/Engine/Source/Programs/Horde/HordeDashboard/src/backend/Backend.ts
@@ -474,7 +474,7 @@ export class Backend {
                 let graphCalls: { jobResponse: JobData, call: any }[] = [];
                 Object.values(timingsWrapper.timings).forEach(timing => {
                     if (!timing.jobResponse.graphHash) {
-                        return reject(`Job ${timing.job.id} has undefined graph hash`);
+                        return reject(`Job ${timing.jobResponse.id} has undefined graph hash`);
                     }
                     graphCalls.push({ jobResponse: timing.jobResponse, call: graphCache.get({ graphHash: timing.jobResponse.graphHash, jobId: timing.jobResponse.id }) });
                 });
