Horde: Fix stream summary time sort
#skipci
#rnx

diff --git a/Engine/Source/Programs/Horde/HordeDashboard/src/base/utilities/timeUtils.ts b/Engine/Source/Programs/Horde/HordeDashboard/src/base/utilities/timeUtils.ts
--- a/Engine/Source/Programs/Horde/HordeDashboard/src/base/utilities/timeUtils.ts
+++ b/Engine/Source/Programs/Horde/HordeDashboard/src/base/utilities/timeUtils.ts
@@ -240,33 +240,42 @@ export const getElapsedString = (start: Moment, end: Moment, includeSeconds:bool
     let duration = "";
     const d = moment.duration(end.diff(start));
 
+    if (d.years()) {
+        duration = `${d.years()} year`;
+        if (d.years() > 1) {
+            duration += "s";
+        }
+    }
     if (d.months()) {
-        duration = `${d.months()} month`;
+        duration += ` ${d.months()} month`;
         if (d.months() > 1) {
             duration += "s";
         }
         return duration;
     }
 
-    if (d.days()) {
-        duration += `${d.days()}d `;
-    }
+    if (!d.years() && !d.months()) {
 
-    if (d.hours()) {
-        duration += `${d.hours()}h `;
-    }
+        if (d.days()) {
+            duration += `${d.days()}d `;
+        }
+        
+        if (d.hours()) {
+            duration += `${d.hours()}h `;
+        }
 
-    if (d.minutes()) {
-        duration += `${d.minutes()}m `;
-    }
+        if (d.minutes()) {
+            duration += `${d.minutes()}m `;
+        }
 
-    if (includeSeconds || !duration) {
-        if (d.seconds()) {
-            duration += `${d.seconds()}s `;
-        }    
+        if (includeSeconds || !duration) {
+            if (d.seconds()) {
+                duration += `${d.seconds()}s `;
+            }
+        }
     }
 
-    return duration;
+    return duration?.trim();
 
 };
 
