Horde: Prevent issue reports being output interleaved with each other.

diff --git a/Engine/Source/Programs/Horde/Horde.Build/Issues/IssueReportService.cs b/Engine/Source/Programs/Horde/Horde.Build/Issues/IssueReportService.cs
--- a/Engine/Source/Programs/Horde/Horde.Build/Issues/IssueReportService.cs
+++ b/Engine/Source/Programs/Horde/Horde.Build/Issues/IssueReportService.cs
@@ -31,19 +31,21 @@ namespace Horde.Build.Issues
 
 	public class IssueReport
 	{
+		public string Channel { get; }
 		public DateTimeOffset Time { get; }
 		public IStream Stream { get; }
-		public WorkflowConfig Workflow { get; }
 		public WorkflowStats WorkflowStats { get; }
+		public string? TriageChannel { get; }
 		public List<IIssue> Issues { get; } = new List<IIssue>();
 		public List<IIssueSpan> IssueSpans { get; } = new List<IIssueSpan>();
 
-		public IssueReport(DateTimeOffset time, IStream stream, WorkflowConfig workflow, WorkflowStats workflowStats)
+		public IssueReport(string channel, DateTimeOffset time, IStream stream, WorkflowStats workflowStats, string? triageChannel)
 		{
+			Channel = channel;
 			Time = time;
 			Stream = stream;
-			Workflow = workflow;
 			WorkflowStats = workflowStats;
+			TriageChannel = triageChannel;
 		}
 	}
 #pragma warning restore CS1591 // Missing XML comment for publicly visible type or member
@@ -117,6 +119,11 @@ namespace Horde.Build.Issues
 
 					foreach (WorkflowConfig workflow in config.Workflows)
 					{
+						if (workflow.ReportChannel == null)
+						{
+							continue;
+						}
+
 						string key = $"{stream.Id}:{workflow.Id}";
 						invalidKeys.Remove(key);
 
@@ -146,7 +153,7 @@ namespace Horde.Build.Issues
 							workflowStats = new WorkflowStats();
 						}
 
-						IssueReport report = new IssueReport(currentTime, stream, workflow, workflowStats);
+						IssueReport report = new IssueReport(workflow.ReportChannel, currentTime, stream, workflowStats, workflow.TriageChannel);
 						foreach (IIssueSpan span in spans)
 						{
 							if (span.LastSuccess != null && span.LastFailure.Annotations.WorkflowId == workflow.Id)
@@ -158,7 +165,7 @@ namespace Horde.Build.Issues
 						HashSet<int> issueIds = new HashSet<int>(report.IssueSpans.Select(x => x.IssueId));
 						report.Issues.AddRange(issues.Where(x => issueIds.Contains(x.Id)));
 
-						_notificationService.SendIssueReport(report);
+						await _notificationService.SendIssueReportAsync(report);
 
 						state = await _state.UpdateAsync(s => s.KeyToLastReportTime[key] = currentTime);
 					}
diff --git a/Engine/Source/Programs/Horde/Horde.Build/Notifications/INotificationService.cs b/Engine/Source/Programs/Horde/Horde.Build/Notifications/INotificationService.cs
--- a/Engine/Source/Programs/Horde/Horde.Build/Notifications/INotificationService.cs
+++ b/Engine/Source/Programs/Horde/Horde.Build/Notifications/INotificationService.cs
@@ -123,6 +123,6 @@ namespace Horde.Build.Notifications
 		/// Post a notification for the open issues in a stream
 		/// </summary>
 		/// <param name="report">The report data to send</param>
-		void SendIssueReport(IssueReport report);
+		Task SendIssueReportAsync(IssueReport report);
 	}
 }
diff --git a/Engine/Source/Programs/Horde/Horde.Build/Notifications/NotificationService.cs b/Engine/Source/Programs/Horde/Horde.Build/Notifications/NotificationService.cs
--- a/Engine/Source/Programs/Horde/Horde.Build/Notifications/NotificationService.cs
+++ b/Engine/Source/Programs/Horde/Horde.Build/Notifications/NotificationService.cs
@@ -733,9 +733,12 @@ namespace Horde.Build.Notifications
 		}
 
 		/// <inheritdoc/>
-		public void SendIssueReport(IssueReport report)
+		public async Task SendIssueReportAsync(IssueReport report)
 		{
-			EnqueueTasks(sink => sink.SendIssueReportAsync(report));
+			foreach (INotificationSink sink in _sinks)
+			{
+				await sink.SendIssueReportAsync(report);
+			}
 		}
 	}
 }
diff --git a/Engine/Source/Programs/Horde/Horde.Build/Notifications/Sinks/SlackNotificationSink.cs b/Engine/Source/Programs/Horde/Horde.Build/Notifications/Sinks/SlackNotificationSink.cs
--- a/Engine/Source/Programs/Horde/Horde.Build/Notifications/Sinks/SlackNotificationSink.cs
+++ b/Engine/Source/Programs/Horde/Horde.Build/Notifications/Sinks/SlackNotificationSink.cs
@@ -1240,11 +1240,11 @@ namespace Horde.Build.Notifications.Sinks
 		/// <inheritdoc/>
 		public async Task SendIssueReportAsync(IssueReport report)
 		{
-			if (report.Workflow.ReportChannel != null)
+			if (report.Channel != null)
 			{
 				List<BlockBase> blocks = new List<BlockBase>();
 				blocks.Add(new HeaderBlock(AddEnvironmentAnnotation($"{report.Stream.Name}: {report.Time:d}")));
-				await SendMessageAsync(report.Workflow.ReportChannel, blocks: blocks.ToArray(), withEnvironment: false);
+				await SendMessageAsync(report.Channel, blocks: blocks.ToArray(), withEnvironment: false);
 
 				StringBuilder body = new StringBuilder();
 				if (report.Issues.Count == 0)
@@ -1293,8 +1293,8 @@ namespace Horde.Build.Notifications.Sinks
 					}
 
 					body.Append($"\n\u2022 *Issue <{issueUrl}|{issue.Id}>");
-					
-					string? triageChannel = report.Workflow.TriageChannel;
+
+					string? triageChannel = report.TriageChannel;
 					if (triageChannel != null)
 					{
 						MessageStateDocument? state = await GetMessageStateAsync(triageChannel, GetTriageThreadEventId(issue.Id));
@@ -1320,7 +1320,7 @@ namespace Horde.Build.Notifications.Sinks
 					body.Append($"\n\n{report.WorkflowStats.NumPassingSteps:n0} of {report.WorkflowStats.NumSteps:n0} (*{totalPct:0.0}%*) build steps succeeded since last status update.");
 				}
 
-				await SendMessageAsync(report.Workflow.ReportChannel, text: body.ToString(), withEnvironment: false);
+				await SendMessageAsync(report.Channel, text: body.ToString(), withEnvironment: false);
 			}
 		}
 
