Horde: Add an option to enable grouping of issues by template (enabled by default).

diff --git a/Engine/Source/Programs/Horde/Horde.Build/Issues/IssueReportService.cs b/Engine/Source/Programs/Horde/Horde.Build/Issues/IssueReportService.cs
--- a/Engine/Source/Programs/Horde/Horde.Build/Issues/IssueReportService.cs
+++ b/Engine/Source/Programs/Horde/Horde.Build/Issues/IssueReportService.cs
@@ -38,14 +38,16 @@ namespace Horde.Build.Issues
 		public string? TriageChannel { get; }
 		public List<IIssue> Issues { get; } = new List<IIssue>();
 		public List<IIssueSpan> IssueSpans { get; } = new List<IIssueSpan>();
+		public bool GroupByTemplate { get; }
 
-		public IssueReport(string channel, DateTimeOffset time, IStream stream, WorkflowStats workflowStats, string? triageChannel)
+		public IssueReport(string channel, DateTimeOffset time, IStream stream, WorkflowStats workflowStats, string? triageChannel, bool groupByTemplate)
 		{
 			Channel = channel;
 			Time = time;
 			Stream = stream;
 			WorkflowStats = workflowStats;
 			TriageChannel = triageChannel;
+			GroupByTemplate = groupByTemplate;
 		}
 	}
 #pragma warning restore CS1591 // Missing XML comment for publicly visible type or member
@@ -126,7 +128,7 @@ namespace Horde.Build.Issues
 
 						string key = $"{stream.Id}:{workflow.Id}";
 						invalidKeys.Remove(key);
-
+						
 						DateTime lastReportTime;
 						if (!state.KeyToLastReportTime.TryGetValue(key, out lastReportTime))
 						{
@@ -153,7 +155,7 @@ namespace Horde.Build.Issues
 							workflowStats = new WorkflowStats();
 						}
 
-						IssueReport report = new IssueReport(workflow.ReportChannel, currentTime, stream, workflowStats, workflow.TriageChannel);
+						IssueReport report = new IssueReport(workflow.ReportChannel, currentTime, stream, workflowStats, workflow.TriageChannel, workflow.GroupIssuesByTemplate);
 						foreach (IIssueSpan span in spans)
 						{
 							if (span.LastSuccess != null && span.LastFailure.Annotations.WorkflowId == workflow.Id)
diff --git a/Engine/Source/Programs/Horde/Horde.Build/Notifications/Sinks/SlackNotificationSink.cs b/Engine/Source/Programs/Horde/Horde.Build/Notifications/Sinks/SlackNotificationSink.cs
--- a/Engine/Source/Programs/Horde/Horde.Build/Notifications/Sinks/SlackNotificationSink.cs
+++ b/Engine/Source/Programs/Horde/Horde.Build/Notifications/Sinks/SlackNotificationSink.cs
@@ -43,6 +43,7 @@ namespace Horde.Build.Notifications.Sinks
 	using JobId = ObjectId<IJob>;
 	using LogId = ObjectId<ILogFile>;
 	using StreamId = StringId<IStream>;
+	using TemplateRefId = StringId<TemplateRef>;
 	using UserId = ObjectId<IUser>;
 	using WorkflowId = StringId<WorkflowConfig>;
 
@@ -1246,82 +1247,132 @@ namespace Horde.Build.Notifications.Sinks
 				blocks.Add(new HeaderBlock(AddEnvironmentAnnotation($"{report.Stream.Name}: {report.Time:d}")));
 				await SendMessageAsync(report.Channel, blocks: blocks.ToArray(), withEnvironment: false);
 
-				StringBuilder body = new StringBuilder();
+				string summary;
 				if (report.Issues.Count == 0)
 				{
-					body.Append(":tick: No issues open.");
+					summary = ":tick: No issues open.";
 				}
 				else
 				{
 					TimeSpan averageAge = TimeSpan.FromHours(report.Issues.Select(x => (report.Time - x.CreatedAt).TotalHours).Average());
-					body.Append($"*{report.Issues.Count} issues* currently open (average age {FormatReadableTimeSpan(averageAge)}):");
+					summary = $"*{report.Issues.Count} unique issues* currently open (average age {FormatReadableTimeSpan(averageAge)}).";
 				}
-				body.Append('\n');
+				await SendMessageAsync(report.Channel, text: summary, withEnvironment: false);
 
-				foreach(IIssue issue in report.Issues.OrderByDescending(x => x.Id))
+				if (report.GroupByTemplate)
 				{
-					IIssueSpan? span = report.IssueSpans.FirstOrDefault(x => x.IssueId == issue.Id);
-
-					Uri issueUrl = _settings.DashboardUrl;
-					if(span != null)
+					Dictionary<int, IIssue> issueIdToInfo = new Dictionary<int, IIssue>();
+					foreach (IIssue issue in report.Issues)
 					{
-						IIssueStep lastFailure = span.LastFailure;
-						issueUrl = new Uri(issueUrl, $"job/{lastFailure.JobId}?step={lastFailure.StepId}&issue={issue.Id}");
+						issueIdToInfo[issue.Id] = issue;
 					}
 
-					string status = "*Unassigned*";
-					if(issue.OwnerId != null)
+					foreach (IGrouping<TemplateRefId, IIssueSpan> group in report.IssueSpans.GroupBy(x => x.TemplateRefId).OrderBy(x => x.Key.ToString()))
 					{
-						IUser? user = await _userCollection.GetCachedUserAsync(issue.OwnerId.Value);
-						if (user == null)
+						TemplateRefConfig? templateConfig;
+						if (!report.Stream.Config.TryGetTemplate(group.Key, out templateConfig))
 						{
-							status = $"Assigned to user {issue.OwnerId.Value}";
+							continue;
 						}
-						else
+
+						List<(IIssue, IIssueSpan)> pairs = new List<(IIssue, IIssueSpan)>();
+						foreach (IIssueSpan span in group)
 						{
-							status = $"Assigned to {user.Name}";
+							if (issueIdToInfo.TryGetValue(span.IssueId, out IIssue? issue))
+							{
+								pairs.Add((issue, span));
+							}
 						}
-						if(issue.AcknowledgedAt == null)
+
+						if (pairs.Count > 0)
 						{
-							status = $"{status} (unacknowledged)";
+							StringBuilder body = new StringBuilder($"Issues affecting *{templateConfig.Name}*:");
+							foreach ((IIssue issue, IIssueSpan span) in pairs)
+							{
+								body.Append('\n');
+								body.Append(await FormatIssueAsync(issue, span, report));
+							}
+							await SendMessageAsync(report.Channel, text: body.ToString(), withEnvironment: false);
 						}
 					}
-
-					if (issue.QuarantinedByUserId != null)
-					{
-						status = $"{status} - *Quarantined*";
-					}
-
-					body.Append($"\n\u2022 *Issue <{issueUrl}|{issue.Id}>");
-
-					string? triageChannel = report.TriageChannel;
-					if (triageChannel != null)
+				}
+				else
+				{
+					StringBuilder body = new StringBuilder();
+					foreach (IIssue issue in report.Issues.OrderByDescending(x => x.Id))
 					{
-						MessageStateDocument? state = await GetMessageStateAsync(triageChannel, GetTriageThreadEventId(issue.Id));
-						if (state != null && state.Permalink != null)
+						IIssueSpan? span = report.IssueSpans.FirstOrDefault(x => x.IssueId == issue.Id);
+						if (body.Length > 0)
 						{
-							body.Append($" (<{state.Permalink}|Thread>)");
+							body.Append('\n');
 						}
+						body.Append(await FormatIssueAsync(issue, span, report));
 					}
+					await SendMessageAsync(report.Channel, text: body.ToString(), withEnvironment: false);
+				}
 
-					body.Append($"*: {issue.Summary} [{FormatReadableTimeSpan(report.Time - issue.CreatedAt)}]");
+				if (report.WorkflowStats.NumSteps > 0)
+				{
+					double totalPct = (report.WorkflowStats.NumPassingSteps * 100.0) / report.WorkflowStats.NumSteps;
+					string footer = $"{report.WorkflowStats.NumPassingSteps:n0} of {report.WorkflowStats.NumSteps:n0} (*{totalPct:0.0}%*) build steps succeeded since last status update.";
+					await SendMessageAsync(report.Channel, text: footer, withEnvironment: false);
+				}
+			}
+		}
 
-					if (!String.IsNullOrEmpty(issue.ExternalIssueKey))
-					{
-						body.Append($" ({FormatExternalIssue(issue.ExternalIssueKey)})");
-					}
+		async ValueTask<string> FormatIssueAsync(IIssue issue, IIssueSpan? span, IssueReport report)
+		{
+			Uri issueUrl = _settings.DashboardUrl;
+			if (span != null)
+			{
+				IIssueStep lastFailure = span.LastFailure;
+				issueUrl = new Uri(issueUrl, $"job/{lastFailure.JobId}?step={lastFailure.StepId}&issue={issue.Id}");
+			}
 
-					body.Append($" - {status}");
+			string status = "*Unassigned*";
+			if (issue.OwnerId != null)
+			{
+				IUser? user = await _userCollection.GetCachedUserAsync(issue.OwnerId.Value);
+				if (user == null)
+				{
+					status = $"Assigned to user {issue.OwnerId.Value}";
 				}
+				else
+				{
+					status = $"Assigned to {user.Name}";
+				}
+				if (issue.AcknowledgedAt == null)
+				{
+					status = $"{status} (unacknowledged)";
+				}
+			}
 
-				if (report.WorkflowStats.NumSteps > 0)
+			if (issue.QuarantinedByUserId != null)
+			{
+				status = $"{status} - *Quarantined*";
+			}
+
+			StringBuilder body = new StringBuilder($"\u2022 *Issue <{issueUrl}|{issue.Id}>");
+
+			string? triageChannel = report.TriageChannel;
+			if (triageChannel != null)
+			{
+				MessageStateDocument? state = await GetMessageStateAsync(triageChannel, GetTriageThreadEventId(issue.Id));
+				if (state != null && state.Permalink != null)
 				{
-					double totalPct = (report.WorkflowStats.NumPassingSteps * 100.0) / report.WorkflowStats.NumSteps;
-					body.Append($"\n\n{report.WorkflowStats.NumPassingSteps:n0} of {report.WorkflowStats.NumSteps:n0} (*{totalPct:0.0}%*) build steps succeeded since last status update.");
+					body.Append($" (<{state.Permalink}|Thread>)");
 				}
+			}
 
-				await SendMessageAsync(report.Channel, text: body.ToString(), withEnvironment: false);
+			body.Append($"*: {issue.Summary} [{FormatReadableTimeSpan(report.Time - issue.CreatedAt)}]");
+
+			if (!String.IsNullOrEmpty(issue.ExternalIssueKey))
+			{
+				body.Append($" ({FormatExternalIssue(issue.ExternalIssueKey)})");
 			}
+
+			body.Append($" - {status}");
+			return body.ToString();
 		}
 
 		static string FormatReadableTimeSpan(TimeSpan timeSpan)
diff --git a/Engine/Source/Programs/Horde/Horde.Build/Server/Config.cs b/Engine/Source/Programs/Horde/Horde.Build/Server/Config.cs
--- a/Engine/Source/Programs/Horde/Horde.Build/Server/Config.cs
+++ b/Engine/Source/Programs/Horde/Horde.Build/Server/Config.cs
@@ -349,6 +349,11 @@ namespace Horde.Build.Server
 		/// </summary>
 		public string? ReportChannel { get; set; }
 
+		/// <summary>
+		/// Whether to group issues by template in the report
+		/// </summary>
+		public bool GroupIssuesByTemplate { get; set; } = true;
+
 		/// <summary>
 		/// Channel to post threads for triaging new issues
 		/// </summary>
