Horde: Add more tracing of completed chunk writes

This method does parallel writes which potentially causes a spiky load on the underlying log chunk storage (Redis).

diff --git a/Engine/Source/Programs/Horde/Horde.Build/Logs/LogFileService.cs b/Engine/Source/Programs/Horde/Horde.Build/Logs/LogFileService.cs
--- a/Engine/Source/Programs/Horde/Horde.Build/Logs/LogFileService.cs
+++ b/Engine/Source/Programs/Horde/Horde.Build/Logs/LogFileService.cs
@@ -1125,6 +1125,11 @@ namespace Horde.Build.Logs
 		/// <returns>Async task</returns>
 		private async Task<ILogFile?> WriteCompleteChunksForLogAsync(ILogFile logFileInterface, List<long> offsets, bool bCreateIndex)
 		{
+			using IScope scope = GlobalTracer.Instance.BuildSpan("WriteCompleteChunksForLogAsync").StartActive();
+			scope.Span.SetTag("LogId", logFileInterface.Id.ToString());
+			scope.Span.SetTag("NumOffsets", offsets.Count);
+			scope.Span.SetTag("CreateIndex", bCreateIndex);
+			
 			// Write the data to the storage provider
 			List<Task<LogChunkData?>> chunkWriteTasks = new List<Task<LogChunkData?>>();
 			foreach (long offset in offsets)
