Horde: Add partial matching for usernames when more than one character
#rnx
#skipci

diff --git a/Engine/Source/Programs/Horde/HordeDashboard/src/backend/UserCache.ts b/Engine/Source/Programs/Horde/HordeDashboard/src/backend/UserCache.ts
--- a/Engine/Source/Programs/Horde/HordeDashboard/src/backend/UserCache.ts
+++ b/Engine/Source/Programs/Horde/HordeDashboard/src/backend/UserCache.ts
@@ -14,11 +14,11 @@ class UserCache {
     @observable
     updated: number = 0;
 
-    async getUsers(startsWith:string) {
+    async getUsers(partialMatch: string) {
 
-        startsWith = startsWith.toLowerCase();
-        
-        let users = this.userMap.get(startsWith);
+        partialMatch = partialMatch.toLowerCase();
+
+        let users = this.userMap.get(partialMatch);
 
         if (users !== undefined) {
             return users;
@@ -27,8 +27,9 @@ class UserCache {
         users = [];
 
         try {
-            users = await backend.getUsers({count: 128, includeAvatar: true, nameRegex: `^${startsWith}`});
-            console.log(`Found ${users.length} users starting with ${startsWith}`);
+            const regex = partialMatch.length === 1 ? `^${partialMatch}` : `${partialMatch}`;
+            users = await backend.getUsers({ count: 128, includeAvatar: true, nameRegex: regex });
+            console.log(`Found ${users.length} users with partial match: ${partialMatch}`);
 
             users = users.sort((a, b) => {
                 const aname = a.name.toLowerCase();
@@ -38,16 +39,16 @@ class UserCache {
                 if (aname > bname) return 1;
                 return 0;
             })
-            this.userMap.set(startsWith, users);
+            this.userMap.set(partialMatch, users);
             this.setUpdated();
         } catch (reason) {
             console.error("Unable to get users", reason);
         }
-        
+
         return users;
     }
 
-    private userMap:Map<string, GetUserResponse[]> = new Map();
+    private userMap: Map<string, GetUserResponse[]> = new Map();
 
 }
 
