Horde: Update issue user query to use OwnerId instead of finding where supplied UserId is a suspect
#rnx

diff --git a/Engine/Source/Programs/Horde/Horde.Build/Issues/IIssueCollection.cs b/Engine/Source/Programs/Horde/Horde.Build/Issues/IIssueCollection.cs
--- a/Engine/Source/Programs/Horde/Horde.Build/Issues/IIssueCollection.cs
+++ b/Engine/Source/Programs/Horde/Horde.Build/Issues/IIssueCollection.cs
@@ -425,7 +425,7 @@ namespace Horde.Build.Issues
 		/// Searches for open issues
 		/// </summary>
 		/// <param name="ids">Set of issue ids to find</param>
-		/// <param name="userId">The user to find issues for</param>
+		/// <param name="ownerId">The user to find issues for</param>
 		/// <param name="streamId">The stream affected by the issue</param>
 		/// <param name="minChange">Minimum changelist affected by the issue</param>
 		/// <param name="maxChange">Maximum changelist affected by the issue</param>
@@ -434,7 +434,7 @@ namespace Horde.Build.Issues
 		/// <param name="index">Index within the results to return</param>
 		/// <param name="count">Number of results</param>
 		/// <returns>List of streams open in the given stream at the given changelist</returns>
-		Task<List<IIssue>> FindIssuesAsync(IEnumerable<int>? ids = null, UserId? userId = null, StreamId? streamId = null, int? minChange = null, int? maxChange = null, bool? resolved = null, bool? promoted = null, int? index = null, int? count = null);
+		Task<List<IIssue>> FindIssuesAsync(IEnumerable<int>? ids = null, UserId? ownerId = null, StreamId? streamId = null, int? minChange = null, int? maxChange = null, bool? resolved = null, bool? promoted = null, int? index = null, int? count = null);
 
 		/// <summary>
 		/// Searches for open issues
@@ -612,13 +612,13 @@ namespace Horde.Build.Issues
 		/// <param name="stepId"></param>
 		/// <param name="batchId"></param>
 		/// <param name="labelIdx"></param>
-		/// <param name="userId"></param>
+		/// <param name="ownerId"></param>
 		/// <param name="resolved">Whether to include results that are resolved</param>
 		/// <param name="promoted">Whether to filter by promoted issues</param>
 		/// <param name="index">Index within the results to return</param>
 		/// <param name="count">Number of results</param>
 		/// <returns></returns>
-		public static async Task<List<IIssue>> FindIssuesForJobAsync(this IIssueCollection issueCollection, IJob job, IGraph graph, SubResourceId? stepId = null, SubResourceId? batchId = null, int? labelIdx = null, UserId? userId = null, bool? resolved = null, bool? promoted = null, int? index = null, int? count = null)
+		public static async Task<List<IIssue>> FindIssuesForJobAsync(this IIssueCollection issueCollection, IJob job, IGraph graph, SubResourceId? stepId = null, SubResourceId? batchId = null, int? labelIdx = null, UserId? ownerId = null, bool? resolved = null, bool? promoted = null, int? index = null, int? count = null)
 		{
 			List<IIssueStep> steps = await issueCollection.FindStepsAsync(job.Id, batchId, stepId);
 			List<IIssueSpan> spans = await issueCollection.FindSpansAsync(steps.Select(x => x.SpanId));
@@ -635,7 +635,7 @@ namespace Horde.Build.Issues
 				return new List<IIssue>();
 			}
 
-			return await issueCollection.FindIssuesAsync(ids: issueIds, userId: userId, resolved: resolved, promoted: promoted, index: index, count: count);
+			return await issueCollection.FindIssuesAsync(ids: issueIds, ownerId: ownerId, resolved: resolved, promoted: promoted, index: index, count: count);
 		}
 
 
diff --git a/Engine/Source/Programs/Horde/Horde.Build/Issues/IssueCollection.cs b/Engine/Source/Programs/Horde/Horde.Build/Issues/IssueCollection.cs
--- a/Engine/Source/Programs/Horde/Horde.Build/Issues/IssueCollection.cs
+++ b/Engine/Source/Programs/Horde/Horde.Build/Issues/IssueCollection.cs
@@ -694,42 +694,20 @@ namespace Horde.Build.Issues
 		}
 
 		/// <inheritdoc/>
-		public async Task<List<IIssue>> FindIssuesAsync(IEnumerable<int>? ids = null, UserId? userId = null, StreamId? streamId = null, int? minChange = null, int? maxChange = null, bool? resolved = null, bool? promoted = null, int? index = null, int? count = null)
+		public async Task<List<IIssue>> FindIssuesAsync(IEnumerable<int>? ids = null, UserId? ownerId = null, StreamId? streamId = null, int? minChange = null, int? maxChange = null, bool? resolved = null, bool? promoted = null, int? index = null, int? count = null)
 		{
-			List<Issue> results = await FilterIssuesByUserIdAsync(ids, userId, streamId, minChange, maxChange, resolved ?? false, promoted, index ?? 0, count);
-			return results.ConvertAll<IIssue>(x => x);
-		}
+			List<Issue> results;
 
-		async Task<List<Issue>> FilterIssuesByUserIdAsync(IEnumerable<int>? ids, UserId? userId, StreamId? streamId, int? minChange, int? maxChange, bool? resolved, bool? promoted, int index, int? count)
-		{
-			if (userId == null)
+			if (ownerId == null)
 			{
-				return await FilterIssuesByStreamIdAsync(ids, streamId, minChange, maxChange, resolved, promoted, index, count);
+				results = await FilterIssuesByStreamIdAsync(ids, streamId, minChange, maxChange, resolved, promoted, index ?? 0, count);
 			}
 			else
 			{
-				FilterDefinition<IssueSuspect> filter = Builders<IssueSuspect>.Filter.Eq(x => x.AuthorId, userId);
-				if (ids != null)
-				{
-					filter &= Builders<IssueSuspect>.Filter.In(x => x.IssueId, ids);
-				}
-				if (resolved != null)
-				{
-					if (resolved.Value)
-					{
-						filter &= Builders<IssueSuspect>.Filter.Ne(x => x.ResolvedAt, null);
-					}
-					else
-					{
-						filter &= Builders<IssueSuspect>.Filter.Eq(x => x.ResolvedAt, null);
-					}
-				}
-
-				using (IAsyncCursor<ProjectedIssueId> cursor = await _issueSuspects.Aggregate().Match(filter).Group(x => x.IssueId, x => new ProjectedIssueId { _id = x.Key }).SortByDescending(x => x._id).ToCursorAsync())
-				{
-					return await PaginatedJoinAsync(cursor, (nextIds, nextIndex, nextCount) => FilterIssuesByStreamIdAsync(nextIds, streamId, minChange, maxChange, null, promoted, nextIndex, nextCount), index, count);
-				}
+				results = await _issues.Find(x => x.OwnerId == ownerId).ToListAsync();
 			}
+
+			return results.ConvertAll<IIssue>(x => x);
 		}
 
 		async Task<List<Issue>> FilterIssuesByStreamIdAsync(IEnumerable<int>? ids, StreamId? streamId, int? minChange, int? maxChange, bool? resolved, bool? promoted, int index, int? count)
diff --git a/Engine/Source/Programs/Horde/Horde.Build/Issues/IssuesController.cs b/Engine/Source/Programs/Horde/Horde.Build/Issues/IssuesController.cs
--- a/Engine/Source/Programs/Horde/Horde.Build/Issues/IssuesController.cs
+++ b/Engine/Source/Programs/Horde/Horde.Build/Issues/IssuesController.cs
@@ -192,7 +192,7 @@ namespace Horde.Build.Issues
 		/// <param name="batchId">The batch to filter by</param>
 		/// <param name="stepId">The step to filter by</param>
 		/// <param name="labelIdx">The label within the job to filter by</param>
-		/// <param name="userId">User to filter issues for</param>
+		/// <param name="ownerId">User to filter issues for</param>
 		/// <param name="resolved">Whether to include resolved issues</param>
 		/// <param name="promoted">Whether to include promoted issues</param>
 		/// <param name="index">Starting offset of the window of results to return</param>
@@ -202,17 +202,17 @@ namespace Horde.Build.Issues
 		[HttpGet]
 		[Route("/api/v1/issues")]
 		[ProducesResponseType(typeof(List<GetIssueResponse>), 200)]
-		public async Task<ActionResult<object>> FindIssuesAsync([FromQuery(Name = "Id")] int[]? ids = null, [FromQuery] string? streamId = null, [FromQuery] int? change = null, [FromQuery] int? minChange = null, [FromQuery] int? maxChange = null, [FromQuery] JobId? jobId = null, [FromQuery] string? batchId = null, [FromQuery] string? stepId = null, [FromQuery(Name = "label")] int? labelIdx = null, [FromQuery] string? userId = null, [FromQuery] bool? resolved = null, [FromQuery] bool? promoted = null, [FromQuery] int index = 0, [FromQuery] int count = 10, [FromQuery] PropertyFilter? filter = null)
+		public async Task<ActionResult<object>> FindIssuesAsync([FromQuery(Name = "Id")] int[]? ids = null, [FromQuery] string? streamId = null, [FromQuery] int? change = null, [FromQuery] int? minChange = null, [FromQuery] int? maxChange = null, [FromQuery] JobId? jobId = null, [FromQuery] string? batchId = null, [FromQuery] string? stepId = null, [FromQuery(Name = "label")] int? labelIdx = null, [FromQuery] string? ownerId = null, [FromQuery] bool? resolved = null, [FromQuery] bool? promoted = null, [FromQuery] int index = 0, [FromQuery] int count = 10, [FromQuery] PropertyFilter? filter = null)
 		{
 			if(ids != null && ids.Length == 0)
 			{
 				ids = null;
 			}
 
-			UserId? userIdValue = null;
-			if (userId != null)
+			UserId? ownerIdValue = null;
+			if (ownerId != null)
 			{
-				userIdValue = new UserId(userId);
+				ownerIdValue = new UserId(ownerId);
 			}
 
 			List<IIssue> issues;
@@ -224,7 +224,7 @@ namespace Horde.Build.Issues
 					streamIdValue = new StreamId(streamId);
 				}
 
-				issues = await _issueService.Collection.FindIssuesAsync(ids, userIdValue, streamIdValue, minChange ?? change, maxChange ?? change, resolved, promoted, index, count);
+				issues = await _issueService.Collection.FindIssuesAsync(ids, ownerIdValue, streamIdValue, minChange ?? change, maxChange ?? change, resolved, promoted, index, count);
 			}
 			else
 			{
@@ -239,7 +239,7 @@ namespace Horde.Build.Issues
 				}
 
 				IGraph graph = await _jobService.GetGraphAsync(job);
-				issues = await _issueService.Collection.FindIssuesForJobAsync(job, graph, stepId?.ToSubResourceId(), batchId?.ToSubResourceId(), labelIdx, userIdValue, resolved, promoted, index, count);
+				issues = await _issueService.Collection.FindIssuesForJobAsync(job, graph, stepId?.ToSubResourceId(), batchId?.ToSubResourceId(), labelIdx, ownerIdValue, resolved, promoted, index, count);
 			}
 
 			StreamPermissionsCache permissionsCache = new StreamPermissionsCache();
