Horde: Fix issue with job step refs not having an outcome on batch errors
#rnx

diff --git a/Engine/Source/Programs/Horde/Horde.Build/Jobs/IJobStepRefCollection.cs b/Engine/Source/Programs/Horde/Horde.Build/Jobs/IJobStepRefCollection.cs
--- a/Engine/Source/Programs/Horde/Horde.Build/Jobs/IJobStepRefCollection.cs
+++ b/Engine/Source/Programs/Horde/Horde.Build/Jobs/IJobStepRefCollection.cs
@@ -109,6 +109,11 @@ namespace Horde.Build.Jobs
 					outcome = JobStepOutcome.Unspecified;
 				}
 
+				if (batch.Error != JobStepBatchError.None && outcome == null)
+				{
+					outcome = JobStepOutcome.Unspecified;
+				}
+
 				await jobStepRefs.InsertOrReplaceAsync(new JobStepRefId(job.Id, batch.Id, step.Id), job.Name, nodeName, job.StreamId, job.TemplateId, job.Change, step.LogId, batch.PoolId, batch.AgentId, outcome, lastSuccess, lastWarning, waitTime, initTime, step.StartTimeUtc ?? DateTime.UtcNow, step.FinishTimeUtc);
 			}
 		}
