Horde: Restore resolved parameter on issue query
#rnx

diff --git a/Engine/Source/Programs/Horde/Horde.Build/Issues/IssueCollection.cs b/Engine/Source/Programs/Horde/Horde.Build/Issues/IssueCollection.cs
--- a/Engine/Source/Programs/Horde/Horde.Build/Issues/IssueCollection.cs
+++ b/Engine/Source/Programs/Horde/Horde.Build/Issues/IssueCollection.cs
@@ -700,7 +700,7 @@ namespace Horde.Build.Issues
 
 			if (ownerId == null)
 			{
-				results = await FilterIssuesByStreamIdAsync(ids, streamId, minChange, maxChange, resolved, promoted, index ?? 0, count);
+				results = await FilterIssuesByStreamIdAsync(ids, streamId, minChange, maxChange, resolved ?? false, promoted, index ?? 0, count);
 			}
 			else
 			{
