Horde: Change issue userId query to use ownerId (dashboard)
#rnx
#skipci

diff --git a/Engine/Source/Programs/Horde/HordeDashboard/src/backend/Api.ts b/Engine/Source/Programs/Horde/HordeDashboard/src/backend/Api.ts
--- a/Engine/Source/Programs/Horde/HordeDashboard/src/backend/Api.ts
+++ b/Engine/Source/Programs/Horde/HordeDashboard/src/backend/Api.ts
@@ -333,7 +333,7 @@ export type IssueQuery = {
 	maxChange?: number;
 	index?: number;
 	count?: number;
-	userId?: string;
+	ownerId?: string;
 	resolved?: boolean;
 	promoted?: boolean;
 }
