Horde: Add timed out informaiton to step summary when job is aborted due to reaching it
#rnx
#skipci

diff --git a/Engine/Source/Programs/Horde/HordeDashboard/src/backend/Api.ts b/Engine/Source/Programs/Horde/HordeDashboard/src/backend/Api.ts
--- a/Engine/Source/Programs/Horde/HordeDashboard/src/backend/Api.ts
+++ b/Engine/Source/Programs/Horde/HordeDashboard/src/backend/Api.ts
@@ -75,11 +75,11 @@ export enum JobStepOutcome {
 /// Systemic error codes for a job failing
 export enum JobStepError {
 	/// No systemic error
-	None,
+	None = "None",
 	/// Step did not complete in the required amount of time
-	TimedOut,
+	TimedOut = "TimedOut",
 	/// Step is in paused state so was skipped
-	Paused
+	Paused = "Paused"
 }
 
 // The state of a particular run
diff --git a/Engine/Source/Programs/Horde/HordeDashboard/src/backend/JobDetails.ts b/Engine/Source/Programs/Horde/HordeDashboard/src/backend/JobDetails.ts
--- a/Engine/Source/Programs/Horde/HordeDashboard/src/backend/JobDetails.ts
+++ b/Engine/Source/Programs/Horde/HordeDashboard/src/backend/JobDetails.ts
@@ -5,7 +5,7 @@ import { getTheme, mergeStyles, mergeStyleSets } from '@fluentui/react/lib/Styli
 import backend from '.';
 import { getBatchInitElapsed, getNiceTime, getStepElapsed, getStepETA, getStepFinishTime, getStepTimingDelta } from '../base/utilities/timeUtils';
 import { getBatchText } from '../components/JobDetailCommon';
-import { AgentData, ArtifactData, BatchData, EventData, GetGroupResponse, GetJobStepRefResponse, GetJobTimingResponse, GetLabelResponse, GetLabelStateResponse, GetLabelTimingInfoResponse, GetTemplateResponse, GroupData, IssueData, JobData, JobState, JobStepBatchState, JobStepOutcome, JobStepState, LabelState, NodeData, ReportPlacement, StepData, StreamData, TestData } from './Api';
+import { AgentData, ArtifactData, BatchData, EventData, GetGroupResponse, GetJobStepRefResponse, GetJobTimingResponse, GetLabelResponse, GetLabelStateResponse, GetLabelTimingInfoResponse, GetTemplateResponse, GroupData, IssueData, JobData, JobState, JobStepBatchState, JobStepError, JobStepOutcome, JobStepState, LabelState, NodeData, ReportPlacement, StepData, StreamData, TestData } from './Api';
 import { projectStore } from './ProjectStore';
 import moment from 'moment';
 
@@ -885,7 +885,7 @@ export class JobDetails {
                 requests.push(backend.getJobTestData(this.id!));
                 if (!this.suppressIssues) {
                     requests.push(this.getIssues());
-                }                
+                }
                 requests.push(this.queryReports());
             }
 
@@ -1035,7 +1035,7 @@ export class JobDetails {
 
     fatalError?: string;
 
-    suppressIssues:boolean = false;
+    suppressIssues: boolean = false;
 
     private static templates = new Map<string, GetTemplateResponse>();
 }
@@ -1194,8 +1194,8 @@ export const getStepSummaryMarkdown = (jobDetails: JobDetails, stepId: string):
         } else {
             text.push(`Step was retried by ${step.retriedByUserInfo.name}`);
         }
-
     }
+
     if (step.abortRequested || step.state === JobStepState.Aborted) {
         eta.display = eta.server = "";
         let aborted = "";
@@ -1207,6 +1207,10 @@ export const getStepSummaryMarkdown = (jobDetails: JobDetails, stepId: string):
             aborted += ` by ${jobDetails.jobdata?.abortedByUserInfo.name}.`;
         } else {
             aborted = "The step was aborted";
+            
+            if (step.error === JobStepError.TimedOut) {
+                aborted = "The step was aborted due to reaching the maximum run time limit";
+            }
         }
         text.push(aborted);
     } else if (step.state === JobStepState.Skipped) {
@@ -1244,7 +1248,7 @@ export const getStepSummaryMarkdown = (jobDetails: JobDetails, stepId: string):
     } else {
         if (!step.abortRequested) {
             text.push(batchText() ?? "Step does not have a batch.");
-        }        
+        }
     }
 
     if (eta.display) {
