Horde: Only output one span for each issue in reports.

diff --git a/Engine/Source/Programs/Horde/Horde.Build/Notifications/Sinks/SlackNotificationSink.cs b/Engine/Source/Programs/Horde/Horde.Build/Notifications/Sinks/SlackNotificationSink.cs
--- a/Engine/Source/Programs/Horde/Horde.Build/Notifications/Sinks/SlackNotificationSink.cs
+++ b/Engine/Source/Programs/Horde/Horde.Build/Notifications/Sinks/SlackNotificationSink.cs
@@ -1275,19 +1275,19 @@ namespace Horde.Build.Notifications.Sinks
 							continue;
 						}
 
-						List<(IIssue, IIssueSpan)> pairs = new List<(IIssue, IIssueSpan)>();
+						Dictionary<IIssue, IIssueSpan> pairs = new Dictionary<IIssue, IIssueSpan>();
 						foreach (IIssueSpan span in group)
 						{
 							if (issueIdToInfo.TryGetValue(span.IssueId, out IIssue? issue))
 							{
-								pairs.Add((issue, span));
+								pairs[issue] = span;
 							}
 						}
 
 						if (pairs.Count > 0)
 						{
 							StringBuilder body = new StringBuilder($"Issues affecting *{templateConfig.Name}*:");
-							foreach ((IIssue issue, IIssueSpan span) in pairs)
+							foreach ((IIssue issue, IIssueSpan span) in pairs.OrderBy(x => x.Key.Id))
 							{
 								body.Append('\n');
 								body.Append(await FormatIssueAsync(issue, span, report));
