Horde: Only auto-scale enabled agents

diff --git a/Engine/Source/Programs/Horde/Horde.Build.Tests/Fleet/AutoscaleServiceV2Test.cs b/Engine/Source/Programs/Horde/Horde.Build.Tests/Fleet/AutoscaleServiceV2Test.cs
--- a/Engine/Source/Programs/Horde/Horde.Build.Tests/Fleet/AutoscaleServiceV2Test.cs
+++ b/Engine/Source/Programs/Horde/Horde.Build.Tests/Fleet/AutoscaleServiceV2Test.cs
@@ -1,7 +1,9 @@
 // Copyright Epic Games, Inc. All Rights Reserved.
 
+extern alias HordeAgent;
 using System;
 using System.Collections.Generic;
+using System.Linq;
 using System.Threading;
 using System.Threading.Tasks;
 using Horde.Build.Agents;
@@ -94,6 +96,21 @@ namespace Horde.Build.Tests.Fleet
 			Assert.AreEqual(1, noOpSpy.CallCount);
 			Assert.IsTrue(noOpSpy.PoolIdsSeen.Contains(pool5.Id));
 		}
+		
+		[TestMethod]
+		public async Task OnlyEnabledAgentsAreAutoScaled()
+		{
+			using AutoscaleServiceV2 service = GetAutoscaleService(_fleetManagerSpy);
+			IPool pool = await PoolService.CreatePoolAsync("testPool", null, true, 0, 0, sizeStrategy: PoolSizeStrategy.LeaseUtilization);
+			await CreateAgentAsync(pool, true);
+			await CreateAgentAsync(pool, true);
+			await CreateAgentAsync(pool, false);
+			
+			Dictionary<PoolSizeStrategy, List<PoolSizeData>> poolSizeDatas = await service.GetPoolSizeDataByStrategyType();
+			List<IAgent> agents = poolSizeDatas.Values.SelectMany(x => x.SelectMany(y => y.Agents)).ToList();
+
+			Assert.AreEqual(2, agents.Count);
+		}
 
 		[TestMethod]
 		public async Task ScaleOutCooldown()
diff --git a/Engine/Source/Programs/Horde/Horde.Build.Tests/TestSetup.cs b/Engine/Source/Programs/Horde/Horde.Build.Tests/TestSetup.cs
--- a/Engine/Source/Programs/Horde/Horde.Build.Tests/TestSetup.cs
+++ b/Engine/Source/Programs/Horde/Horde.Build.Tests/TestSetup.cs
@@ -286,9 +286,9 @@ namespace Horde.Build.Tests
 		}
 		
 		private static int s_agentIdCounter = 1;
-		public async Task<IAgent> CreateAgentAsync(IPool pool)
+		public async Task<IAgent> CreateAgentAsync(IPool pool, bool enabled = true)
 		{
-			IAgent agent = await AgentService.CreateAgentAsync("TestAgent" + s_agentIdCounter++, true, null, new List<StringId<IPool>> { pool.Id });
+			IAgent agent = await AgentService.CreateAgentAsync("TestAgent" + s_agentIdCounter++, enabled, null, new List<StringId<IPool>> { pool.Id });
 			agent = await AgentService.CreateSessionAsync(agent, AgentStatus.Ok, new List<string>(), new Dictionary<string, int>(), null);
 			return agent;
 		}
diff --git a/Engine/Source/Programs/Horde/Horde.Build/Agents/AgentCollection.cs b/Engine/Source/Programs/Horde/Horde.Build/Agents/AgentCollection.cs
--- a/Engine/Source/Programs/Horde/Horde.Build/Agents/AgentCollection.cs
+++ b/Engine/Source/Programs/Horde/Horde.Build/Agents/AgentCollection.cs
@@ -194,7 +194,7 @@ namespace Horde.Build.Agents
 		}
 
 		/// <inheritdoc/>
-		public async Task<List<IAgent>> FindAsync(PoolId? poolId, DateTime? modifiedAfter, AgentStatus? status, int? index, int? count)
+		public async Task<List<IAgent>> FindAsync(PoolId? poolId, DateTime? modifiedAfter, AgentStatus? status, bool? enabled, int? index, int? count)
 		{
 			FilterDefinitionBuilder<AgentDocument> filterBuilder = new FilterDefinitionBuilder<AgentDocument>();
 
@@ -214,6 +214,11 @@ namespace Horde.Build.Agents
 			{
 				filter &= filterBuilder.Eq(x => x.Status, status.Value);
 			}
+			
+			if (enabled != null)
+			{
+				filter &= filterBuilder.Eq(x => x.Enabled, enabled.Value);
+			}
 
 			IFindFluent<AgentDocument, AgentDocument> search = _agents.Find(filter);
 			if (index != null)
diff --git a/Engine/Source/Programs/Horde/Horde.Build/Agents/AgentService.cs b/Engine/Source/Programs/Horde/Horde.Build/Agents/AgentService.cs
--- a/Engine/Source/Programs/Horde/Horde.Build/Agents/AgentService.cs
+++ b/Engine/Source/Programs/Horde/Horde.Build/Agents/AgentService.cs
@@ -209,7 +209,7 @@ namespace Horde.Build.Agents
 		/// <returns>List of agents matching the given criteria</returns>
 		public Task<List<IAgent>> FindAgentsAsync(PoolId? poolId, DateTime? modifiedAfter, int? index, int? count)
 		{
-			return Agents.FindAsync(poolId, modifiedAfter, null, index, count);
+			return Agents.FindAsync(poolId, modifiedAfter, null, null, index, count);
 		}
 
 		/// <summary>
diff --git a/Engine/Source/Programs/Horde/Horde.Build/Agents/Fleet/AutoscaleServiceV2.cs b/Engine/Source/Programs/Horde/Horde.Build/Agents/Fleet/AutoscaleServiceV2.cs
--- a/Engine/Source/Programs/Horde/Horde.Build/Agents/Fleet/AutoscaleServiceV2.cs
+++ b/Engine/Source/Programs/Horde/Horde.Build/Agents/Fleet/AutoscaleServiceV2.cs
@@ -88,7 +88,7 @@ namespace Horde.Build.Agents.Fleet
 
 		internal async Task<Dictionary<PoolSizeStrategy, List<PoolSizeData>>> GetPoolSizeDataByStrategyType()
 		{
-			List<IAgent> agents = await _agentCollection.FindAsync(status: AgentStatus.Ok);
+			List<IAgent> agents = await _agentCollection.FindAsync(status: AgentStatus.Ok, enabled: true);
 			List<IAgent> GetAgentsInPool(PoolId poolId) => agents.FindAll(a => a.GetPools().Any(p => p == poolId));
 			List<IPool> pools = await _poolCollection.GetAsync();
 
diff --git a/Engine/Source/Programs/Horde/Horde.Build/Agents/IAgentCollection.cs b/Engine/Source/Programs/Horde/Horde.Build/Agents/IAgentCollection.cs
--- a/Engine/Source/Programs/Horde/Horde.Build/Agents/IAgentCollection.cs
+++ b/Engine/Source/Programs/Horde/Horde.Build/Agents/IAgentCollection.cs
@@ -1,4 +1,4 @@
-// Copyright Epic Games, Inc. All Rights Reserved.
+﻿// Copyright Epic Games, Inc. All Rights Reserved.
 
 using System;
 using System.Collections.Generic;
@@ -58,10 +58,11 @@ namespace Horde.Build.Agents
 		/// <param name="poolId">The pool ID in string form containing the agent</param>
 		/// <param name="modifiedAfter">If set, only returns agents modified after this time</param>
 		/// <param name="status">Status to look for</param>
+		/// <param name="enabled">Enabled/disabled status to look for</param>
 		/// <param name="index">Index of the first result</param>
 		/// <param name="count">Number of results to return</param>
 		/// <returns>List of agents matching the given criteria</returns>
-		Task<List<IAgent>> FindAsync(PoolId? poolId = null, DateTime? modifiedAfter = null, AgentStatus? status = null, int? index = null, int? count = null);
+		Task<List<IAgent>> FindAsync(PoolId? poolId = null, DateTime? modifiedAfter = null, AgentStatus? status = null, bool? enabled = null, int? index = null, int? count = null);
 
 		/// <summary>
 		/// Finds all the expried agents
