Horde Storage - Allow multiple authentication schemes. Moved disable auth into a boolean Enabled flag instead.
This allows us to support authentication against multiple okta auth servers (useful while moving auth server).
Updated Okta library to support overriding scheme.

diff --git a/Engine/Source/Programs/Horde/HordeStorage/Callisto/appsettings.Development.json b/Engine/Source/Programs/Horde/HordeStorage/Callisto/appsettings.Development.json
--- a/Engine/Source/Programs/Horde/HordeStorage/Callisto/appsettings.Development.json
+++ b/Engine/Source/Programs/Horde/HordeStorage/Callisto/appsettings.Development.json
@@ -15,7 +15,7 @@
   },
 
   "Auth": {
-    "Method": "Disabled"
+    "Enabled": false
   },
 
   "Jupiter": {
diff --git a/Engine/Source/Programs/Horde/HordeStorage/Callisto/appsettings.Testing.json b/Engine/Source/Programs/Horde/HordeStorage/Callisto/appsettings.Testing.json
--- a/Engine/Source/Programs/Horde/HordeStorage/Callisto/appsettings.Testing.json
+++ b/Engine/Source/Programs/Horde/HordeStorage/Callisto/appsettings.Testing.json
@@ -14,9 +14,9 @@
     "Enrich": [ "FromLogContext" ]
   },
 
-  "Auth": {
-    "Method": "Disabled"
-  },
+	"Auth": {
+		"Enabled": false
+	},
 
   "Jupiter": {
     "CurrentSite": "Test"
diff --git a/Engine/Source/Programs/Horde/HordeStorage/Horde.Storage/Horde.Storage.Startup.cs b/Engine/Source/Programs/Horde/HordeStorage/Horde.Storage/Horde.Storage.Startup.cs
--- a/Engine/Source/Programs/Horde/HordeStorage/Horde.Storage/Horde.Storage.Startup.cs
+++ b/Engine/Source/Programs/Horde/HordeStorage/Horde.Storage/Horde.Storage.Startup.cs
@@ -168,13 +168,6 @@ namespace Horde.Storage
 
         protected override void OnAddService(IServiceCollection services)
         {
-            // For requests served on the high-perf HTTP port there is not authn/authz.
-            // Suppressing the check below removes the warning for that case.
-            services.Configure<RouteOptions>(options =>
-            {
-                options.SuppressCheckForUnhandledSecurityMetadata = true;
-            });
-
             services.AddHttpClient();
 
             services.AddOptions<HordeStorageSettings>().Bind(Configuration.GetSection("Horde_Storage")).ValidateDataAnnotations();
diff --git a/Engine/Source/Programs/Horde/HordeStorage/Horde.Storage/Implementation/Blob/IBlobService.cs b/Engine/Source/Programs/Horde/HordeStorage/Horde.Storage/Implementation/Blob/IBlobService.cs
--- a/Engine/Source/Programs/Horde/HordeStorage/Horde.Storage/Implementation/Blob/IBlobService.cs
+++ b/Engine/Source/Programs/Horde/HordeStorage/Horde.Storage/Implementation/Blob/IBlobService.cs
@@ -358,7 +358,7 @@ public class BlobService : IBlobService
         HttpRequestMessage request = new HttpRequestMessage(httpMethod, uri);
         if (!string.IsNullOrEmpty(token))
         {
-            request.Headers.Add("Authorization", "Bearer " + token);
+            request.Headers.Add("Authorization", $"{_serviceCredentials.GetAuthenticationScheme()} {token}");
         }
 
         return request;
diff --git a/Engine/Source/Programs/Horde/HordeStorage/Horde.Storage/Implementation/RelayStore.cs b/Engine/Source/Programs/Horde/HordeStorage/Horde.Storage/Implementation/RelayStore.cs
--- a/Engine/Source/Programs/Horde/HordeStorage/Horde.Storage/Implementation/RelayStore.cs
+++ b/Engine/Source/Programs/Horde/HordeStorage/Horde.Storage/Implementation/RelayStore.cs
@@ -29,7 +29,7 @@ namespace Horde.Storage.Implementation
             HttpRequestMessage request = new HttpRequestMessage(method, uri);
             if (!string.IsNullOrEmpty(token))
             {
-                request.Headers.Add("Authorization", "Bearer " + token);
+                request.Headers.Add("Authorization", $"{_serviceCredentials.GetAuthenticationScheme()} {token}");
             }
 
             return request;
diff --git a/Engine/Source/Programs/Horde/HordeStorage/Horde.Storage/Implementation/Replication/RefsReplicator.cs b/Engine/Source/Programs/Horde/HordeStorage/Horde.Storage/Implementation/Replication/RefsReplicator.cs
--- a/Engine/Source/Programs/Horde/HordeStorage/Horde.Storage/Implementation/Replication/RefsReplicator.cs
+++ b/Engine/Source/Programs/Horde/HordeStorage/Horde.Storage/Implementation/Replication/RefsReplicator.cs
@@ -379,7 +379,7 @@ namespace Horde.Storage.Implementation
             HttpRequestMessage request = new HttpRequestMessage(httpMethod, uri);
             if (!string.IsNullOrEmpty(token))
             {
-                request.Headers.Add("Authorization", "Bearer " + token);
+                request.Headers.Add("Authorization", $"{_serviceCredentials.GetAuthenticationScheme()} {token}");
             }
 
             return request;
diff --git a/Engine/Source/Programs/Horde/HordeStorage/Horde.Storage/Implementation/Replication/ReplicationSettings.cs b/Engine/Source/Programs/Horde/HordeStorage/Horde.Storage/Implementation/Replication/ReplicationSettings.cs
--- a/Engine/Source/Programs/Horde/HordeStorage/Horde.Storage/Implementation/Replication/ReplicationSettings.cs
+++ b/Engine/Source/Programs/Horde/HordeStorage/Horde.Storage/Implementation/Replication/ReplicationSettings.cs
@@ -117,5 +117,10 @@ namespace Horde.Storage
         /// The scope to request
         /// </summary>
         public string OAuthScope { get; set; } = "cache_access";
+
+        /// <summary>
+        /// The authentication scheme to use for this token
+        /// </summary>
+        public string SchemeName { get; set; } = "Bearer";
     }
 }
diff --git a/Engine/Source/Programs/Horde/HordeStorage/Horde.Storage/Implementation/Replication/Replicator.cs b/Engine/Source/Programs/Horde/HordeStorage/Horde.Storage/Implementation/Replication/Replicator.cs
--- a/Engine/Source/Programs/Horde/HordeStorage/Horde.Storage/Implementation/Replication/Replicator.cs
+++ b/Engine/Source/Programs/Horde/HordeStorage/Horde.Storage/Implementation/Replication/Replicator.cs
@@ -431,7 +431,7 @@ namespace Horde.Storage.Implementation
             HttpRequestMessage request = new HttpRequestMessage(httpMethod, uri);
             if (!string.IsNullOrEmpty(token))
             {
-                request.Headers.Add("Authorization", "Bearer " + token);
+                request.Headers.Add("Authorization", $"{_serviceCredentials.GetAuthenticationScheme()} {token}");
             }
 
             return request;
diff --git a/Engine/Source/Programs/Horde/HordeStorage/Horde.Storage/Implementation/ServiceCredentials.cs b/Engine/Source/Programs/Horde/HordeStorage/Horde.Storage/Implementation/ServiceCredentials.cs
--- a/Engine/Source/Programs/Horde/HordeStorage/Horde.Storage/Implementation/ServiceCredentials.cs
+++ b/Engine/Source/Programs/Horde/HordeStorage/Horde.Storage/Implementation/ServiceCredentials.cs
@@ -13,14 +13,18 @@ namespace Horde.Storage.Implementation
         IAuthenticator? GetAuthenticator();
 
         string? GetToken();
+
+        string GetAuthenticationScheme();
     }
 
     public class ServiceCredentials : IServiceCredentials
     {
         private readonly ClientCredentialOAuthAuthenticator? _authenticator;
+        private readonly IOptionsMonitor<ServiceCredentialSettings> _settings;
 
         public ServiceCredentials(IOptionsMonitor<ServiceCredentialSettings> settings, ISecretResolver secretResolver)
         {
+            _settings = settings;
             if (settings.CurrentValue.OAuthLoginUrl != null)
             {
                 string? clientId = secretResolver.Resolve(settings.CurrentValue.OAuthClientId);
@@ -35,7 +39,7 @@ namespace Horde.Storage.Implementation
                     throw new ArgumentException("ClientSecret must be set when using a service credential");
                 }
 
-                _authenticator = new ClientCredentialOAuthAuthenticator(settings.CurrentValue.OAuthLoginUrl, clientId, clientSecret, settings.CurrentValue.OAuthScope);
+                _authenticator = new ClientCredentialOAuthAuthenticator(settings.CurrentValue.OAuthLoginUrl, clientId, clientSecret, settings.CurrentValue.OAuthScope, settings.CurrentValue.SchemeName);
             }
         }
 
@@ -48,5 +52,10 @@ namespace Horde.Storage.Implementation
         {
             return _authenticator?.Authenticate();
         }
+
+        public string GetAuthenticationScheme()
+        {
+            return _settings.CurrentValue.SchemeName;
+        }
     }
 }
diff --git a/Engine/Source/Programs/Horde/HordeStorage/Horde.Storage/appsettings.Development.json b/Engine/Source/Programs/Horde/HordeStorage/Horde.Storage/appsettings.Development.json
--- a/Engine/Source/Programs/Horde/HordeStorage/Horde.Storage/appsettings.Development.json
+++ b/Engine/Source/Programs/Horde/HordeStorage/Horde.Storage/appsettings.Development.json
@@ -74,7 +74,7 @@
 
 	// disable auth and allow any namespace
 	"Auth": {
-		"Method": "Disabled"
+		"Enabled": false
 	},
 
 	"Namespaces": {
diff --git a/Engine/Source/Programs/Horde/HordeStorage/Horde.Storage/appsettings.Testing.json b/Engine/Source/Programs/Horde/HordeStorage/Horde.Storage/appsettings.Testing.json
--- a/Engine/Source/Programs/Horde/HordeStorage/Horde.Storage/appsettings.Testing.json
+++ b/Engine/Source/Programs/Horde/HordeStorage/Horde.Storage/appsettings.Testing.json
@@ -71,9 +71,9 @@
   "Azure": {
     "ConnectionString": "DefaultEndpointsProtocol=https;BlobEndpoint=http://localhost:11002/blobaccount;AccountName=blobaccount;AccountKey=iU6uTvlF1ysppmft+NO5lAD0E3hwrAORr5Rb5xcBWUgEz/OicrSkFxwZYMNK5XL29/wXZKGOoOVSW040nAOfPg=="
   },
-  "Auth": {
-    "Method": "Disabled"
-  },
+	"Auth": {
+		"Enabled": false
+	},
 
   "Namespaces": {
     "Policies": {
@@ -107,10 +107,7 @@
   "AllowedHosts": "*",
 
   "Jupiter": {
-    "CurrentSite": "test", 
-    "DisableAuthOnPorts": [
-      57003
-    ]
+    "CurrentSite": "test",
   },
 
   "Kestrel": {
@@ -121,9 +118,6 @@
       "Http2": {
         "Url": "http://localhost:57002",
         "Protocols": "Http2"
-      },
-      "AuthDisabledHttp": {
-        "Url": "http://localhost:57003"
       }
     }
   }
diff --git a/Engine/Source/Programs/Horde/HordeStorage/Jupiter.Common/BaseStartup.cs b/Engine/Source/Programs/Horde/HordeStorage/Jupiter.Common/BaseStartup.cs
--- a/Engine/Source/Programs/Horde/HordeStorage/Jupiter.Common/BaseStartup.cs
+++ b/Engine/Source/Programs/Horde/HordeStorage/Jupiter.Common/BaseStartup.cs
@@ -7,6 +7,7 @@ using System.IO;
 using System.Net.Mime;
 using System.Reflection;
 using System.Text;
+using System.Threading;
 using System.Threading.Tasks;
 using Amazon;
 using Datadog.Trace;
@@ -107,69 +108,71 @@ namespace Jupiter
 
             services.AddHttpContextAccessor();
 
-            // we change the name of the jwt scheme as okta uses the same name and does not allow us to reconfigure it.
-            const string JwtScheme = "JWTBearer";
-            List<string> defaultSchemes = new List<string>();
+            List<string> availableSchemes = new List<string>();
 
             AuthenticationBuilder authenticationBuilder = services.AddAuthentication(options =>
                 {
-                    switch (Auth.Method)
+                    if (Auth.Enabled)
                     {
-                        case AuthMethod.JWTBearer:
-                            options.DefaultAuthenticateScheme = JwtScheme;
-                            options.DefaultChallengeScheme = JwtScheme;
-                            break;
-                        case AuthMethod.Okta:
-                            options.DefaultAuthenticateScheme = OktaDefaults.ApiAuthenticationScheme;
-                            options.DefaultChallengeScheme = OktaDefaults.ApiAuthenticationScheme;
-                            break;
-                        case AuthMethod.Disabled:
-                            options.DefaultAuthenticateScheme = DisabledAuthenticationHandler.AuthenticateScheme;
-                            options.DefaultChallengeScheme = DisabledAuthenticationHandler.AuthenticateScheme;
-                            break;
-                        default:
-                            throw new NotImplementedException($"Method {Auth.Method} not implemented");
+                        options.DefaultAuthenticateScheme = Auth.DefaultScheme;
+                        options.DefaultChallengeScheme = Auth.DefaultScheme;
                     }
-                });
-            if (Auth.Method == AuthMethod.JWTBearer)
-            {
-                defaultSchemes.Add(JwtScheme);
-                authenticationBuilder.AddJwtBearer(JwtScheme, options =>
-                {
-                    options.Authority = Auth.JwtAuthority;
-                    options.Audience = Auth.JwtAudience;
-                });
-            }
+                    else
+                    {
+                        options.DefaultAuthenticateScheme = DisabledAuthenticationHandler.AuthenticateScheme;
+                        options.DefaultChallengeScheme = DisabledAuthenticationHandler.AuthenticateScheme;
+                    }
+                }
+            );
 
-            if (Auth.Method == AuthMethod.Okta)
+            if (Auth.Enabled)
             {
-                defaultSchemes.Add(OktaDefaults.ApiAuthenticationScheme);
-                authenticationBuilder.AddOktaWebApi(new OktaWebApiOptions
+                foreach (KeyValuePair<string, AuthSchemeEntry> schemeEntry in Auth.Schemes)
                 {
-                    OktaDomain = Auth.OktaDomain,
-                    AuthorizationServerId = Auth.AuthorizationServerId,
-                    Audience = Auth.JwtAudience,
-                });
-            }
+                    AuthSchemeEntry scheme = schemeEntry.Value;
 
-            if (Auth.Method == AuthMethod.Disabled)
+                    switch (scheme.Implementation)
+                    {
+                        case SchemeImplementations.JWTBearer:
+                            availableSchemes.Add(scheme.Name);
+                            authenticationBuilder.AddJwtBearer(scheme.Name, options =>
+                            {
+                                options.Authority = scheme.JwtAuthority;
+                                options.Audience = scheme.JwtAudience;
+                            });
+                            break;
+                        case SchemeImplementations.Okta:
+                            availableSchemes.Add(scheme.Name);
+                            authenticationBuilder.AddOktaWebApi(scheme.Name, new OktaWebApiOptions
+                            {
+                                OktaDomain = scheme.OktaDomain,
+                                AuthorizationServerId = scheme.OktaAuthorizationServerId,
+                                Audience = scheme.JwtAudience,
+                            });
+                            break;
+                        default:
+                            throw new NotSupportedException($"Unknown implementation type {scheme.Implementation}");
+                    }
+                }
+            }
+            else
             {
-                defaultSchemes.Add(DisabledAuthenticationHandler.AuthenticateScheme);
+                availableSchemes.Add(DisabledAuthenticationHandler.AuthenticateScheme);
                 authenticationBuilder.AddTestAuth(options => { });
             }
 
-            defaultSchemes.Add(ServiceAccountAuthHandler.AuthenticationScheme);
+            availableSchemes.Add(ServiceAccountAuthHandler.AuthenticationScheme);
             authenticationBuilder.AddScheme<ServiceAccountAuthOptions, ServiceAccountAuthHandler>(ServiceAccountAuthHandler.AuthenticationScheme, options => { });
 
             services.AddAuthorization(options =>
             {
                 options.AddPolicy(NamespaceAccessRequirement.Name, policy =>
                 {
-                    policy.AuthenticationSchemes = defaultSchemes;
+                    policy.AuthenticationSchemes = availableSchemes;
                     policy.Requirements.Add(new NamespaceAccessRequirement());
                 });
 
-                OnAddAuthorization(options, defaultSchemes);
+                OnAddAuthorization(options, availableSchemes);
             });
             services.AddSingleton<IAuthorizationHandler, NamespaceAuthorizationHandler>();
 
@@ -410,28 +413,116 @@ namespace Jupiter
         }
     }
 
-    public enum AuthMethod
+    public enum SchemeImplementations
     {
         JWTBearer,
-        Okta,
-        Disabled
+        Okta
     };
 
-    public class AuthSettings
+    public class AuthSchemeEntry: IValidatableObject
     {
+        /// <summary>
+        /// The name of the authentication scheme. This will need to be specified by any client in the Authorization header together with the token
+        /// </summary>
+        [Required]
+        [Key]
+        public string Name { get; set; } = "Bearer";
+
+        /// <summary>
+        /// The implementation to use, this controls which other configuration values needs to be set. For most servers JWTBearer should work fine.
+        /// </summary>
         [Required] 
-        public AuthMethod Method { get; set; } = AuthMethod.JWTBearer;
+        public SchemeImplementations Implementation { get; set; } = SchemeImplementations.JWTBearer;
 
-        [Required]
+        /// <summary>
+        /// The Okta domain (url to your okta server - do not include the authorization server id)
+        /// </summary>
         public string OktaDomain { get; set; } = "";
-        public string AuthorizationServerId { get; set; } = OktaWebOptions.DefaultAuthorizationServerId;
+        /// <summary>
+        /// The Okta AuthorizationServerId, this is used if you have more then one authorization server within your Okta server. We recommend using a separate authorization server for each major set of systems to reduce blast radius of security issues.
+        /// </summary>
+        public string OktaAuthorizationServerId { get; set; } = OktaWebOptions.DefaultAuthorizationServerId;
 
-        [Required]
+        /// <summary>
+        /// The JWT Authority (url to your IdP)
+        /// </summary>
         public string JwtAuthority { get; set; } = "";
 
+        /// <summary>
+        /// The audience for the token, this is usually defined by your IdP 
+        /// </summary>
         [Required]
         public string JwtAudience { get; set; } = "";
 
+        public IEnumerable<ValidationResult> Validate(ValidationContext validationContext)
+        {
+            List<ValidationResult> validationResults = new List<ValidationResult>();
+            if (Implementation == SchemeImplementations.JWTBearer)
+            {
+                if (string.IsNullOrEmpty(JwtAuthority))
+                {
+                    validationResults.Add(new ValidationResult("JWT Authority must be specified when using JWTBearer implementation"));
+                }
+                if (string.IsNullOrEmpty(JwtAudience))
+                {
+                    validationResults.Add(new ValidationResult("JWT Audience must be specified when using JWTBearer implementation"));
+                }
+            } 
+            else if (Implementation == SchemeImplementations.Okta)
+            {
+                if (string.IsNullOrEmpty(OktaDomain))
+                {
+                    validationResults.Add(new ValidationResult("Okta Domain must be specified when using Okta implementation"));
+                }
+                if (string.IsNullOrEmpty(JwtAudience))
+                {
+                    validationResults.Add(new ValidationResult("JWT Audience must be specified when using Okta implementation"));
+                }
+            }
+            else
+            {
+                throw new NotSupportedException($"Unknown auth implementation {Implementation}");
+            }
+
+            return validationResults;
+        }
+    }
+
+    public class AuthSettings : IValidatableObject
+    {
+        /// <summary>
+        /// The name of the scheme to use by default
+        /// </summary>
+        public string DefaultScheme { get; set; } = "Bearer";
+
+        /// <summary>
+        /// Used to disable authentication, not recommended to set for anything other then local use cases
+        /// </summary>
+        public bool Enabled { get; set; } = true;
+
+        [System.Diagnostics.CodeAnalysis.SuppressMessage("Usage", "CA2227:Collection properties should be read only", Justification = "Used by the configuration system")]
+        public Dictionary<string, AuthSchemeEntry> Schemes { get; set; } = new Dictionary<string, AuthSchemeEntry>();
+
+        public IEnumerable<ValidationResult> Validate(ValidationContext validationContext)
+        {
+            List<ValidationResult> validationResults = new List<ValidationResult>();
+            if (!Enabled)
+            {
+                return validationResults;
+            }
+
+            if (Schemes.Count == 0)
+            {
+                validationResults.Add(new ValidationResult("You must have at least one scheme when authentication is enabled"));
+            }
+
+            if (!Schemes.ContainsKey(DefaultScheme))
+            {
+                validationResults.Add(new ValidationResult($"Expected to find a scheme with the name {DefaultScheme} as its set as the default scheme"));
+            }
+
+            return validationResults;
+        }
     }
 
     public class JupiterSettings
diff --git a/Engine/Source/Programs/Horde/HordeStorage/Jupiter.Common/Implementation/ClientCredentialOAuthAuthenticator.cs b/Engine/Source/Programs/Horde/HordeStorage/Jupiter.Common/Implementation/ClientCredentialOAuthAuthenticator.cs
--- a/Engine/Source/Programs/Horde/HordeStorage/Jupiter.Common/Implementation/ClientCredentialOAuthAuthenticator.cs
+++ b/Engine/Source/Programs/Horde/HordeStorage/Jupiter.Common/Implementation/ClientCredentialOAuthAuthenticator.cs
@@ -9,12 +9,13 @@ namespace Jupiter.Implementation
 {
     public class ClientCredentialOAuthAuthenticator: IAuthenticator
     {
-        public ClientCredentialOAuthAuthenticator(Uri authUrl, string clientId, string clientSecret, string scope)
+        public ClientCredentialOAuthAuthenticator(Uri authUrl, string clientId, string clientSecret, string scope, string schemeName)
         {
             _authUrl = authUrl;
             _clientId = clientId;
             _clientSecret = clientSecret;
             _scope = scope;
+            _schemeName = schemeName;
         }
 
         private string? _accessToken;
@@ -23,6 +24,7 @@ namespace Jupiter.Implementation
         private readonly string _clientId;
         private readonly string _clientSecret;
         private readonly string _scope;
+        private readonly string _schemeName;
         private DateTime _expiresAt;
 
         public void Authenticate(IRestClient client, IRestRequest request)
@@ -32,7 +34,7 @@ namespace Jupiter.Implementation
                 PreAuthenticate();
             }
 
-            request.AddHeader("Authorization", $"Bearer {_accessToken}");
+            request.AddHeader("Authorization", $"{_schemeName} {_accessToken}");
         }
 
         public string? Authenticate()
diff --git a/Engine/Source/Programs/Horde/HordeStorage/Jupiter.Common/Implementation/NamespacePolicyResolver.cs b/Engine/Source/Programs/Horde/HordeStorage/Jupiter.Common/Implementation/NamespacePolicyResolver.cs
--- a/Engine/Source/Programs/Horde/HordeStorage/Jupiter.Common/Implementation/NamespacePolicyResolver.cs
+++ b/Engine/Source/Programs/Horde/HordeStorage/Jupiter.Common/Implementation/NamespacePolicyResolver.cs
@@ -27,7 +27,7 @@ namespace Jupiter.Common
             {
                 return new NamespacePolicy()
                 {
-                    // we grant the full storage access claim access to the internal namespace, this is typically set of admins and service accounts that can access everything
+                    // we expect the full storage access claim access to the internal namespace, this is typically set of admins and service accounts that can access everything
                     Claims = new string[] {"Storage=full" },
                     StoragePool = ""
                 };
diff --git a/Engine/Source/Programs/Horde/HordeStorage/Jupiter.Common/Jupiter.Common.csproj b/Engine/Source/Programs/Horde/HordeStorage/Jupiter.Common/Jupiter.Common.csproj
--- a/Engine/Source/Programs/Horde/HordeStorage/Jupiter.Common/Jupiter.Common.csproj
+++ b/Engine/Source/Programs/Horde/HordeStorage/Jupiter.Common/Jupiter.Common.csproj
@@ -29,7 +29,7 @@
     <PackageReference Include="Microsoft.Extensions.Diagnostics.HealthChecks" Version="6.0.0" />
     <PackageReference Include="Microsoft.Extensions.Hosting" Version="6.0.0" />
     <PackageReference Include="NetEscapades.Configuration.Yaml" Version="2.1.0" />
-    <PackageReference Include="Okta.AspNetCore" Version="3.5.1" />
+    <PackageReference Include="Okta.AspNetCore" Version="4.2.0" />
     <PackageReference Include="RestSharp" Version="106.11.7" />
     <PackageReference Include="RestSharp.Serializers.NewtonsoftJson" Version="106.11.7" />
     <PackageReference Include="Serilog" Version="2.10.0" />
