Horde Storage - Fixed issues with new auth scheme configuration. Wasnt being validated properly and name needed to be specified twice.

diff --git a/Engine/Source/Programs/Horde/HordeStorage/Jupiter.Common/BaseStartup.cs b/Engine/Source/Programs/Horde/HordeStorage/Jupiter.Common/BaseStartup.cs
--- a/Engine/Source/Programs/Horde/HordeStorage/Jupiter.Common/BaseStartup.cs
+++ b/Engine/Source/Programs/Horde/HordeStorage/Jupiter.Common/BaseStartup.cs
@@ -65,7 +65,7 @@ namespace Jupiter
             // send log4net logs to serilog and configure aws to log to log4net (they lack a serilog implementation)
             AWSConfigs.LoggingConfig.LogTo = LoggingOptions.Log4Net;
 
-            services.AddOptions<AuthSettings>().Bind(Configuration.GetSection("Auth")).ValidateDataAnnotations();
+            services.AddOptions<AuthSettings>().Bind(Configuration.GetSection("Auth")).ValidateDataAnnotations().ValidateOnStart();
             Configuration.GetSection("Auth").Bind(Auth);
 
             services.AddOptions<ServiceAccountAuthOptions>().Bind(Configuration.GetSection("ServiceAccounts")).ValidateDataAnnotations();
@@ -129,21 +129,22 @@ namespace Jupiter
             {
                 foreach (KeyValuePair<string, AuthSchemeEntry> schemeEntry in Auth.Schemes)
                 {
+                    string name = schemeEntry.Key;
                     AuthSchemeEntry scheme = schemeEntry.Value;
 
                     switch (scheme.Implementation)
                     {
                         case SchemeImplementations.JWTBearer:
-                            availableSchemes.Add(scheme.Name);
-                            authenticationBuilder.AddJwtBearer(scheme.Name, options =>
+                            availableSchemes.Add(name);
+                            authenticationBuilder.AddJwtBearer(name, options =>
                             {
                                 options.Authority = scheme.JwtAuthority;
                                 options.Audience = scheme.JwtAudience;
                             });
                             break;
                         case SchemeImplementations.Okta:
-                            availableSchemes.Add(scheme.Name);
-                            authenticationBuilder.AddOktaWebApi(scheme.Name, new OktaWebApiOptions
+                            availableSchemes.Add(name);
+                            authenticationBuilder.AddOktaWebApi(name, new OktaWebApiOptions
                             {
                                 OktaDomain = scheme.OktaDomain,
                                 AuthorizationServerId = scheme.OktaAuthorizationServerId,
@@ -421,13 +422,6 @@ namespace Jupiter
 
     public class AuthSchemeEntry: IValidatableObject
     {
-        /// <summary>
-        /// The name of the authentication scheme. This will need to be specified by any client in the Authorization header together with the token
-        /// </summary>
-        [Required]
-        [Key]
-        public string Name { get; set; } = "Bearer";
-
         /// <summary>
         /// The implementation to use, this controls which other configuration values needs to be set. For most servers JWTBearer should work fine.
         /// </summary>
