Horde: Add a link to completed jobs of each type in issue reports.

diff --git a/Engine/Source/Programs/Horde/Horde.Build/Notifications/Sinks/SlackNotificationSink.cs b/Engine/Source/Programs/Horde/Horde.Build/Notifications/Sinks/SlackNotificationSink.cs
--- a/Engine/Source/Programs/Horde/Horde.Build/Notifications/Sinks/SlackNotificationSink.cs
+++ b/Engine/Source/Programs/Horde/Horde.Build/Notifications/Sinks/SlackNotificationSink.cs
@@ -1286,7 +1286,16 @@ namespace Horde.Build.Notifications.Sinks
 
 						if (pairs.Count > 0)
 						{
-							StringBuilder body = new StringBuilder($"Issues affecting *{templateConfig.Name}*:");
+							string template = templateConfig.Name;
+
+							JobsTab? tab = report.Stream.Config.Tabs.OfType<JobsTab>().FirstOrDefault(x => x.Templates != null && x.Templates.Contains(templateConfig.Id));
+							if (tab != null)
+							{
+								Uri templateUrl = new Uri(_settings.DashboardUrl, $"stream/ue5-main?tab=General&template={templateConfig.Id}");
+								template = $"<{templateUrl}|{template}>";
+							}
+
+							StringBuilder body = new StringBuilder($"Issues affecting *{template}*:");
 							foreach ((IIssue issue, IIssueSpan span) in pairs.OrderBy(x => x.Key.Id))
 							{
 								body.Append('\n');
