Horde Storage - Upgraded Microsoft.IdentityModel.Tokens to 6.17 to avoid issues loading it as a result of the Okta.AspNet upgrade

diff --git a/Engine/Source/Programs/Horde/HordeStorage/Jupiter.Common/Jupiter.Common.csproj b/Engine/Source/Programs/Horde/HordeStorage/Jupiter.Common/Jupiter.Common.csproj
--- a/Engine/Source/Programs/Horde/HordeStorage/Jupiter.Common/Jupiter.Common.csproj
+++ b/Engine/Source/Programs/Horde/HordeStorage/Jupiter.Common/Jupiter.Common.csproj
@@ -28,6 +28,7 @@
     <PackageReference Include="Microsoft.AspNetCore.Mvc.NewtonsoftJson" Version="6.0.0" />
     <PackageReference Include="Microsoft.Extensions.Diagnostics.HealthChecks" Version="6.0.0" />
     <PackageReference Include="Microsoft.Extensions.Hosting" Version="6.0.0" />
+    <PackageReference Include="Microsoft.IdentityModel.Tokens" Version="6.17.0" />
     <PackageReference Include="NetEscapades.Configuration.Yaml" Version="2.1.0" />
     <PackageReference Include="Okta.AspNetCore" Version="4.2.0" />
     <PackageReference Include="RestSharp" Version="106.11.7" />
