Horde Storage - Further fixed to assembly version mismatches as a result of the Okta.AspNet upgrade

diff --git a/Engine/Source/Programs/Horde/HordeStorage/Jupiter.Common/Jupiter.Common.csproj b/Engine/Source/Programs/Horde/HordeStorage/Jupiter.Common/Jupiter.Common.csproj
--- a/Engine/Source/Programs/Horde/HordeStorage/Jupiter.Common/Jupiter.Common.csproj
+++ b/Engine/Source/Programs/Horde/HordeStorage/Jupiter.Common/Jupiter.Common.csproj
@@ -28,9 +28,16 @@
     <PackageReference Include="Microsoft.AspNetCore.Mvc.NewtonsoftJson" Version="6.0.0" />
     <PackageReference Include="Microsoft.Extensions.Diagnostics.HealthChecks" Version="6.0.0" />
     <PackageReference Include="Microsoft.Extensions.Hosting" Version="6.0.0" />
+
+    <PackageReference Include="Okta.AspNetCore" Version="4.2.0" />
+    <!-- Adding explicit version of identity model to workaround deletion of a internal symbol used by Okta.AspNetCore -->
     <PackageReference Include="Microsoft.IdentityModel.Tokens" Version="6.17.0" />
+    <PackageReference Include="Microsoft.IdentityModel.Logging" Version="6.17.0" />
+    <PackageReference Include="Microsoft.IdentityModel.Protocols.OpenIdConnect" Version="6.17.0" />
+    <PackageReference Include="Microsoft.IdentityModel.Protocols" Version="6.17.0" />
+    <PackageReference Include="Microsoft.IdentityModel.JsonWebTokens" Version="6.17.0" />
+
     <PackageReference Include="NetEscapades.Configuration.Yaml" Version="2.1.0" />
-    <PackageReference Include="Okta.AspNetCore" Version="4.2.0" />
     <PackageReference Include="RestSharp" Version="106.11.7" />
     <PackageReference Include="RestSharp.Serializers.NewtonsoftJson" Version="106.11.7" />
     <PackageReference Include="Serilog" Version="2.10.0" />
