Horde: Fix errors about sending Slack messages with empty sections.

diff --git a/Engine/Source/Programs/Horde/Horde.Build/Notifications/NotificationService.cs b/Engine/Source/Programs/Horde/Horde.Build/Notifications/NotificationService.cs
--- a/Engine/Source/Programs/Horde/Horde.Build/Notifications/NotificationService.cs
+++ b/Engine/Source/Programs/Horde/Horde.Build/Notifications/NotificationService.cs
@@ -301,10 +301,12 @@ namespace Horde.Build.Notifications
 		private async ValueTask TickEveryTwelveHoursAsync(CancellationToken stoppingToken)
 		{
 			List<JobScheduledNotification> jobScheduledNotifications = await GetAllQueuedNotificationsAsync<JobScheduledNotification>();
-
-			foreach (INotificationSink sink in _sinks)
+			if (jobScheduledNotifications.Count > 0)
 			{
-				await sink.NotifyJobScheduledAsync(jobScheduledNotifications);
+				foreach (INotificationSink sink in _sinks)
+				{
+					await sink.NotifyJobScheduledAsync(jobScheduledNotifications);
+				}
 			}
 		}
 		
diff --git a/Engine/Source/Programs/Horde/Horde.Build/Notifications/Sinks/SlackNotificationSink.cs b/Engine/Source/Programs/Horde/Horde.Build/Notifications/Sinks/SlackNotificationSink.cs
--- a/Engine/Source/Programs/Horde/Horde.Build/Notifications/Sinks/SlackNotificationSink.cs
+++ b/Engine/Source/Programs/Horde/Horde.Build/Notifications/Sinks/SlackNotificationSink.cs
@@ -353,7 +353,10 @@ namespace Horde.Build.Notifications.Sinks
 
 			attachment.Blocks.Add(new HeaderBlock($"Jobs scheduled in empty pool", true));
 			attachment.Blocks.Add(new SectionBlock($"One or more jobs were scheduled in an auto-scaled pool but with no current agents online."));
-			attachment.Blocks.Add(new SectionBlock(sb.ToString()));
+			if (sb.Length > 0)
+			{
+				attachment.Blocks.Add(new SectionBlock(sb.ToString()));
+			}
 			
 			await SendMessageAsync(recipient, attachments: new[] { attachment });
 		}
@@ -1291,7 +1294,7 @@ namespace Horde.Build.Notifications.Sinks
 							JobsTab? tab = report.Stream.Config.Tabs.OfType<JobsTab>().FirstOrDefault(x => x.Templates != null && x.Templates.Contains(templateConfig.Id));
 							if (tab != null)
 							{
-								Uri templateUrl = new Uri(_settings.DashboardUrl, $"stream/ue5-main?tab=General&template={templateConfig.Id}");
+								Uri templateUrl = new Uri(_settings.DashboardUrl, $"stream/{report.Stream.Id}?tab={tab.Title}&template={templateConfig.Id}");
 								template = $"<{templateUrl}|{template}>";
 							}
 
