Horde: Fix symbol issue handler incorrectly ignoring generic errors.

diff --git a/Engine/Source/Programs/Horde/Horde.Build.Tests/IssueServiceTests.cs b/Engine/Source/Programs/Horde/Horde.Build.Tests/IssueServiceTests.cs
--- a/Engine/Source/Programs/Horde/Horde.Build.Tests/IssueServiceTests.cs
+++ b/Engine/Source/Programs/Horde/Horde.Build.Tests/IssueServiceTests.cs
@@ -399,6 +399,35 @@ namespace Horde.Build.Tests
 			}
 		}
 
+		[TestMethod]
+		public async Task DefaultIssueTest3()
+		{
+			// #1
+			// Scenario: Warning in first step
+			// Expected: Default issues is created
+			{
+				string[] lines =
+				{
+					@"< enterprise.max.test_maxscript_datasmith_export.from_maxscript_to_unreal[vray_materials.ms] >",
+					@"  [ :ERROR: ] [2022.14.06-12:08:55] [log_parser] keyword found 'logpython: error:'",
+					@"   > [2022.06.14-12.08.44:338][  1]LogPython: Error: One or more instance material is loaded but not expected : ['VRayMtl__2_.VRayMtl__2_', 'VRayMtl__3_.VRayMtl__3_', 'VRayMtl__1_.VRayMtl__1_', 'VRayMtl__4_.VRayMtl__4_']",
+					@"     [2022.06.14-12.08.44:338][  1]LogPython: Error: One or more instance material expected but not loaded: ['VRayMtl_2.VRayMtl_2', 'VRayMtl_3.VRayMtl_3', 'VRayMtl_4.VRayMtl_4', 'VRayMtl_7.VRayMtl_7']",
+					@"   > File ""D:\build\++UE5\Sync\Engine\Saved\pydrover\session[2022.14.06-11.50.11]\from_maxscript_to_unreal_4d5057656576\ue_log[2022.14.06-12.08.31].txt"", line 1085",
+					@"< enterprise.max.test_maxscript_datasmith_export.from_maxscript_to_unreal[vray_materials.ms] >"
+				};
+
+				IJob job = CreateJob(_mainStreamId, 105, "Test Build", _graph);
+				await ParseEventsAsync(job, 0, 0, lines);
+				await UpdateCompleteStep(job, 0, 0, JobStepOutcome.Warnings);
+
+				List<IIssue> issues = await IssueCollection.FindIssuesAsync();
+				Assert.AreEqual(1, issues.Count);
+				Assert.AreEqual(IssueSeverity.Error, issues[0].Severity);
+
+				Assert.AreEqual("Errors in Update Version Files", issues[0].Summary);
+			}
+		}
+
 		[TestMethod]
 		public async Task AutoSdkWarningTest()
 		{
diff --git a/Engine/Source/Programs/Horde/Horde.Build/Horde.Build.csproj b/Engine/Source/Programs/Horde/Horde.Build/Horde.Build.csproj
--- a/Engine/Source/Programs/Horde/Horde.Build/Horde.Build.csproj
+++ b/Engine/Source/Programs/Horde/Horde.Build/Horde.Build.csproj
@@ -35,7 +35,7 @@
     <PackageReference Include="JetBrains.Profiler.SelfApi" Version="2.1.2" />
     <PackageReference Include="Microsoft.AspNetCore.Authentication.Cookies" Version="2.2.0" />
     <PackageReference Include="Microsoft.AspNetCore.Authentication.JwtBearer" Version="6.0.0" />
-    <PackageReference Include="Microsoft.AspNetCore.Authentication.OpenIdConnect" Version="3.1.5" />
+    <PackageReference Include="Microsoft.AspNetCore.Authentication.OpenIdConnect" Version="3.1.23" />
     <PackageReference Include="Microsoft.AspNetCore.DataProtection.StackExchangeRedis" Version="3.1.7" />
     <PackageReference Include="Microsoft.AspNetCore.SpaServices.Extensions" Version="3.1.17" />
     <PackageReference Include="Microsoft.Extensions.Caching.Memory" Version="6.0.0" />
diff --git a/Engine/Source/Programs/Horde/Horde.Build/Issues/Handlers/SymbolIssueHandler.cs b/Engine/Source/Programs/Horde/Horde.Build/Issues/Handlers/SymbolIssueHandler.cs
--- a/Engine/Source/Programs/Horde/Horde.Build/Issues/Handlers/SymbolIssueHandler.cs
+++ b/Engine/Source/Programs/Horde/Horde.Build/Issues/Handlers/SymbolIssueHandler.cs
@@ -134,9 +134,12 @@ namespace Horde.Build.Issues.Handlers
 							stepEvent.Ignored = true;
 						}
 					}
-					else if(eventId == KnownLogEvents.Generic || eventId == KnownLogEvents.AutomationTool)
+					else if (eventId == KnownLogEvents.Generic || eventId == KnownLogEvents.AutomationTool)
 					{
-						stepEvent.Ignored = true;
+						if (hasMatches)
+						{
+							stepEvent.Ignored = true;
+						}
 					}
 				}
 			}
