Horde: Ignore spans in streams that have been deleted when updating issue derived data.

diff --git a/Engine/Source/Programs/Horde/Horde.Build.Tests/IssueServiceTests.cs b/Engine/Source/Programs/Horde/Horde.Build.Tests/IssueServiceTests.cs
--- a/Engine/Source/Programs/Horde/Horde.Build.Tests/IssueServiceTests.cs
+++ b/Engine/Source/Programs/Horde/Horde.Build.Tests/IssueServiceTests.cs
@@ -29,6 +29,7 @@ using MongoDB.Bson;
 using MongoDB.Driver;
 using Moq;
 using Horde.Build.Issues.Handlers;
+using System.Threading;
 
 namespace Horde.Build.Tests
 {
@@ -290,6 +291,37 @@ namespace Horde.Build.Tests
 			}
 		}
 
+		[TestMethod]
+		public async Task DeleteStreamTest()
+		{
+			await IssueService.StartAsync(CancellationToken.None);
+
+			// #1
+			// Scenario: Warning in first step
+			// Expected: Default issues is created
+			{
+				IJob job = CreateJob(_mainStreamId, 105, "Test Build", _graph);
+				await AddEvent(job, 0, 0, new { level = nameof(LogLevel.Warning), message = "" }, EventSeverity.Warning);
+				await UpdateCompleteStep(job, 0, 0, JobStepOutcome.Warnings);
+
+				List<IIssue> issues = await IssueCollection.FindIssuesAsync();
+				Assert.AreEqual(1, issues.Count);
+
+				Assert.AreEqual("Warnings in Update Version Files", issues[0].Summary);
+			}
+
+			// #2
+			// Scenario: Stream is deleted
+			// Expected: Issue is closed
+			{
+				await StreamCollection.DeleteAsync(_mainStreamId);
+				await Clock.AdvanceAsync(TimeSpan.FromHours(1.0));
+
+				List<IIssue> issues = await IssueCollection.FindIssuesAsync();
+				Assert.AreEqual(0, issues.Count);
+			}
+		}
+
 		[TestMethod]
 		public async Task DefaultIssueTest()
 		{
diff --git a/Engine/Source/Programs/Horde/Horde.Build/Issues/IssueService.cs b/Engine/Source/Programs/Horde/Horde.Build/Issues/IssueService.cs
--- a/Engine/Source/Programs/Horde/Horde.Build/Issues/IssueService.cs
+++ b/Engine/Source/Programs/Horde/Horde.Build/Issues/IssueService.cs
@@ -290,6 +290,26 @@ namespace Horde.Build.Issues
 				{
 					await _issueCollection.TryUpdateIssueAsync(openIssue, newResolvedById: IIssue.ResolvedByTimeoutId);
 					openIssues.RemoveAt(idx--);
+					continue;
+				}
+			}
+
+			// Update any issues that are listed as valid for a stream that no longer exists
+			HashSet<StreamId> validStreamIds = new HashSet<StreamId>(cachedStreams.Select(x => x.Id));
+			for (int idx = 0; idx < openIssues.Count; idx++)
+			{
+				IIssue? openIssue = openIssues[idx];
+				if (openIssue.Streams.Any(x => !validStreamIds.Contains(x.StreamId)))
+				{
+					openIssue = await UpdateIssueDerivedDataAsync(openIssue);
+					if (openIssue != null && openIssue.ResolvedAt == null)
+					{
+						openIssues[idx] = openIssue;
+					}
+					else
+					{
+						openIssues.RemoveAt(idx--);
+					}
 				}
 			}
 
@@ -875,6 +895,16 @@ namespace Horde.Build.Issues
 				// Find all the spans that are attached to the issue
 				List<IIssueSpan> spans = await _issueCollection.FindSpansAsync(issue.Id);
 
+				// Remove any spans for streams that have been deleted
+				for (int idx = spans.Count - 1; idx >= 0; idx--)
+				{
+					IStream? stream = await _streams.GetStreamAsync(spans[idx].StreamId);
+					if (stream == null || stream.Deleted)
+					{
+						spans.RemoveAt(idx);
+					}
+				}
+
 				// Update the suspects for this issue
 				List<NewIssueSuspectData> newSuspects = new List<NewIssueSuspectData>();
 				if (spans.Count > 0)
@@ -944,7 +974,7 @@ namespace Horde.Build.Issues
 				{
 					if (spans.Count == 0)
 					{
-						newVerifiedAt = issue.VerifiedAt ?? DateTime.UtcNow;
+						newVerifiedAt = issue.VerifiedAt ?? _clock.UtcNow;
 					}
 					else
 					{
diff --git a/Engine/Source/Programs/Horde/Horde.Build/Streams/IStream.cs b/Engine/Source/Programs/Horde/Horde.Build/Streams/IStream.cs
--- a/Engine/Source/Programs/Horde/Horde.Build/Streams/IStream.cs
+++ b/Engine/Source/Programs/Horde/Horde.Build/Streams/IStream.cs
@@ -588,6 +588,11 @@ namespace Horde.Build.Streams
 		/// </summary>
 		public StreamConfig Config { get; }
 
+		/// <summary>
+		/// Whether this stream has been deleted
+		/// </summary>
+		public bool Deleted { get; }
+
 		/// <summary>
 		/// Order to display on the dashboard's drop-down list
 		/// </summary>
