Horde: Fix link to jobs page from Slack notifications.

diff --git a/Engine/Source/Programs/Horde/Horde.Build/Notifications/Sinks/SlackNotificationSink.cs b/Engine/Source/Programs/Horde/Horde.Build/Notifications/Sinks/SlackNotificationSink.cs
--- a/Engine/Source/Programs/Horde/Horde.Build/Notifications/Sinks/SlackNotificationSink.cs
+++ b/Engine/Source/Programs/Horde/Horde.Build/Notifications/Sinks/SlackNotificationSink.cs
@@ -1291,7 +1291,7 @@ namespace Horde.Build.Notifications.Sinks
 						{
 							string template = templateConfig.Name;
 
-							JobsTab? tab = report.Stream.Config.Tabs.OfType<JobsTab>().FirstOrDefault(x => x.Templates != null && x.Templates.Contains(templateConfig.Id));
+							JobsTab? tab = report.Stream.Config.Tabs.OfType<CreateJobsTabRequest>().FirstOrDefault(x => x.Templates != null && x.Templates.Contains(templateConfig.Id));
 							if (tab != null)
 							{
 								Uri templateUrl = new Uri(_settings.DashboardUrl, $"stream/{report.Stream.Id}?tab={tab.Title}&template={templateConfig.Id}");
diff --git a/Engine/Source/Programs/Horde/Horde.Build/Streams/StreamMessages.cs b/Engine/Source/Programs/Horde/Horde.Build/Streams/StreamMessages.cs
--- a/Engine/Source/Programs/Horde/Horde.Build/Streams/StreamMessages.cs
+++ b/Engine/Source/Programs/Horde/Horde.Build/Streams/StreamMessages.cs
@@ -799,7 +799,7 @@ namespace Horde.Build.Streams
 		/// <summary>
 		/// List of job template names to show on this page.
 		/// </summary>
-		public List<string>? Templates { get; set; }
+		public List<TemplateRefId>? Templates { get; set; }
 
 		/// <summary>
 		/// Columns to display for different types of aggregates
