Horde: Build fix.

diff --git a/Engine/Source/Programs/Horde/Horde.Build/Notifications/Sinks/SlackNotificationSink.cs b/Engine/Source/Programs/Horde/Horde.Build/Notifications/Sinks/SlackNotificationSink.cs
--- a/Engine/Source/Programs/Horde/Horde.Build/Notifications/Sinks/SlackNotificationSink.cs
+++ b/Engine/Source/Programs/Horde/Horde.Build/Notifications/Sinks/SlackNotificationSink.cs
@@ -1291,7 +1291,7 @@ namespace Horde.Build.Notifications.Sinks
 						{
 							string template = templateConfig.Name;
 
-							JobsTab? tab = report.Stream.Config.Tabs.OfType<CreateJobsTabRequest>().FirstOrDefault(x => x.Templates != null && x.Templates.Contains(templateConfig.Id));
+							CreateJobsTabRequest? tab = report.Stream.Config.Tabs.OfType<CreateJobsTabRequest>().FirstOrDefault(x => x.Templates != null && x.Templates.Contains(templateConfig.Id));
 							if (tab != null)
 							{
 								Uri templateUrl = new Uri(_settings.DashboardUrl, $"stream/{report.Stream.Id}?tab={tab.Title}&template={templateConfig.Id}");
diff --git a/Engine/Source/Programs/Horde/Horde.Build/Streams/StreamMessages.cs b/Engine/Source/Programs/Horde/Horde.Build/Streams/StreamMessages.cs
--- a/Engine/Source/Programs/Horde/Horde.Build/Streams/StreamMessages.cs
+++ b/Engine/Source/Programs/Horde/Horde.Build/Streams/StreamMessages.cs
@@ -918,7 +918,7 @@ namespace Horde.Build.Streams
 		/// <summary>
 		/// List of templates to show on the page
 		/// </summary>
-		public List<string>? Templates { get; set; }
+		public List<TemplateRefId>? Templates { get; set; }
 
 		/// <summary>
 		/// Names of jobs to include on this page. If there is only one name specified, the name column does not need to be displayed.
@@ -938,7 +938,7 @@ namespace Horde.Build.Streams
 		/// <param name="templates">Templates to include on this page</param>
 		/// <param name="jobNames">List of job names to include on this page</param>
 		/// <param name="columns">List of columns to display</param>
-		public GetJobsTabResponse(string title, bool bShowNames, List<string>? templates, List<string>? jobNames, List<GetJobsTabColumnResponse>? columns)
+		public GetJobsTabResponse(string title, bool bShowNames, List<TemplateRefId>? templates, List<string>? jobNames, List<GetJobsTabColumnResponse>? columns)
 			: base(title)
 		{
 			ShowNames = bShowNames;
diff --git a/Engine/Source/Programs/Horde/Horde.Build/Streams/StreamTab.cs b/Engine/Source/Programs/Horde/Horde.Build/Streams/StreamTab.cs
--- a/Engine/Source/Programs/Horde/Horde.Build/Streams/StreamTab.cs
+++ b/Engine/Source/Programs/Horde/Horde.Build/Streams/StreamTab.cs
@@ -252,7 +252,7 @@ namespace Horde.Build.Streams
 			: base(request)
 		{
 			ShowNames = request.ShowNames;
-			Templates = request.Templates?.ConvertAll(x => new TemplateRefId(x));
+			Templates = request.Templates;
 			JobNames = request.JobNames;
 			Columns = request.Columns?.ConvertAll(x => x.ToModel());
 		}
@@ -260,7 +260,7 @@ namespace Horde.Build.Streams
 		/// <inheritdoc/>
 		public override GetStreamTabResponse ToResponse()
 		{
-			return new GetJobsTabResponse(Title, ShowNames, Templates?.ConvertAll(x => x.ToString()), JobNames, Columns?.ConvertAll(x => x.ToResponse()));
+			return new GetJobsTabResponse(Title, ShowNames, Templates, JobNames, Columns?.ConvertAll(x => x.ToResponse()));
 		}
 	}
 }
