Fix tests.

diff --git a/Engine/Source/Programs/Horde/Horde.Build.Tests/Fixture.cs b/Engine/Source/Programs/Horde/Horde.Build.Tests/Fixture.cs
--- a/Engine/Source/Programs/Horde/Horde.Build.Tests/Fixture.cs
+++ b/Engine/Source/Programs/Horde/Horde.Build.Tests/Fixture.cs
@@ -70,7 +70,7 @@ namespace Horde.Build.Tests
 			templates.Add(new TemplateRefConfig { Id = TemplateRefId2, Name = "Test Template" });
 
 			List<CreateStreamTabRequest> tabs = new List<CreateStreamTabRequest>();
-			tabs.Add(new CreateJobsTabRequest { Title = "foo", Templates = new List<string> { TemplateRefId1.ToString(), TemplateRefId2.ToString() } });
+			tabs.Add(new CreateJobsTabRequest { Title = "foo", Templates = new List<TemplateRefId> { TemplateRefId1, TemplateRefId2 } });
 
 			Dictionary<string, CreateAgentTypeRequest> agentTypes = new()
 			{
diff --git a/Engine/Source/Programs/Horde/Horde.Build.Tests/IssueServiceTests.cs b/Engine/Source/Programs/Horde/Horde.Build.Tests/IssueServiceTests.cs
--- a/Engine/Source/Programs/Horde/Horde.Build.Tests/IssueServiceTests.cs
+++ b/Engine/Source/Programs/Horde/Horde.Build.Tests/IssueServiceTests.cs
@@ -122,7 +122,7 @@ namespace Horde.Build.Tests
 			string revision = $"config:{streamId}";
 
 			StreamConfig streamConfig = new StreamConfig { Name = streamName };
-			streamConfig.Tabs.Add(new CreateJobsTabRequest { Title = "General", Templates = new List<string> { "test-template" } });
+			streamConfig.Tabs.Add(new CreateJobsTabRequest { Title = "General", Templates = new List<TemplateRefId> { new TemplateRefId("test-template") } });
 			streamConfig.Templates.Add(new TemplateRefConfig { Id = new TemplateRefId("test-template") });
 			await ConfigCollection.AddConfigAsync(revision, streamConfig);
 
diff --git a/Engine/Source/Programs/Horde/Horde.Build.Tests/JobServiceTests.cs b/Engine/Source/Programs/Horde/Horde.Build.Tests/JobServiceTests.cs
--- a/Engine/Source/Programs/Horde/Horde.Build.Tests/JobServiceTests.cs
+++ b/Engine/Source/Programs/Horde/Horde.Build.Tests/JobServiceTests.cs
@@ -46,7 +46,7 @@ namespace Horde.Build.Tests
 			StreamConfig streamConfig = new StreamConfig();
 			streamConfig.Templates.Add(new TemplateRefConfig { Id = templateRefId1, Name = "Test Template", ChainedJobs = new List<CreateChainedJobTemplateRequest> { new CreateChainedJobTemplateRequest { TemplateId = templateRefId2.ToString(), Trigger = "Setup Build" } } });
 			streamConfig.Templates.Add(new TemplateRefConfig { Id = templateRefId2, Name = "Test Template" });
-			streamConfig.Tabs.Add(new CreateJobsTabRequest { Title = "foo", Templates = new List<string> { templateRefId1.ToString(), templateRefId2.ToString() } });
+			streamConfig.Tabs.Add(new CreateJobsTabRequest { Title = "foo", Templates = new List<TemplateRefId> { templateRefId1, templateRefId2 } });
 
 			StreamId streamId = new StreamId("ue5-main");
 			IStream? stream = await StreamService.GetStreamAsync(streamId);
diff --git a/Engine/Source/Programs/Horde/Horde.Build.Tests/SchedulerTests.cs b/Engine/Source/Programs/Horde/Horde.Build.Tests/SchedulerTests.cs
--- a/Engine/Source/Programs/Horde/Horde.Build.Tests/SchedulerTests.cs
+++ b/Engine/Source/Programs/Horde/Horde.Build.Tests/SchedulerTests.cs
@@ -62,7 +62,7 @@ namespace Horde.Build.Tests
 
 			StreamConfig config = new StreamConfig();
 			config.Name = "//UE5/Main";
-			config.Tabs.Add(new CreateJobsTabRequest { Title = "foo", Templates = new List<string> { TemplateRefId.ToString() } });
+			config.Tabs.Add(new CreateJobsTabRequest { Title = "foo", Templates = new List<TemplateRefId> { TemplateRefId } });
 			config.Templates.Add(new TemplateRefConfig { Id = TemplateRefId, Name = "Test", Schedule = schedule });
 
 			return (await CreateOrReplaceStreamAsync(StreamId, stream, ProjectId, config))!; 
@@ -427,7 +427,7 @@ namespace Horde.Build.Tests
 
 			StreamConfig config = new StreamConfig();
 			config.Name = "//UE5/Main";
-			config.Tabs.Add(new CreateJobsTabRequest { Title = "foo", Templates = new List<string> { newTemplateRefId1.ToString(), newTemplateRefId2.ToString() } });
+			config.Tabs.Add(new CreateJobsTabRequest { Title = "foo", Templates = new List<TemplateRefId> { newTemplateRefId1, newTemplateRefId2 } });
 
 			stream = (await CreateOrReplaceStreamAsync(StreamId, stream, ProjectId, config))!;
 
@@ -507,7 +507,7 @@ namespace Horde.Build.Tests
 
 			StreamConfig config = new StreamConfig();
 			config.Name = "//UE5/Main";
-			config.Tabs.Add(new CreateJobsTabRequest { Title = "foo", Templates = new List<string> { newTemplateRefId1.ToString(), newTemplateRefId2.ToString() } });
+			config.Tabs.Add(new CreateJobsTabRequest { Title = "foo", Templates = new List<TemplateRefId> { newTemplateRefId1, newTemplateRefId2 } });
 			config.Templates.Add(newTemplate1);
 			config.Templates.Add(newTemplate2);
 
