Horde: Replace logic for ignoring systemic issues with code to ignore specific event ids in situations that warrant it. Systemic issues are now treated as real systemic isssues rather than false positives.

diff --git a/Engine/Source/Programs/Horde/Horde.Build/Issues/Handlers/ContentIssueHandler.cs b/Engine/Source/Programs/Horde/Horde.Build/Issues/Handlers/ContentIssueHandler.cs
--- a/Engine/Source/Programs/Horde/Horde.Build/Issues/Handlers/ContentIssueHandler.cs
+++ b/Engine/Source/Programs/Horde/Horde.Build/Issues/Handlers/ContentIssueHandler.cs
@@ -35,6 +35,16 @@ namespace Horde.Build.Issues.Handlers
 			return eventId == KnownLogEvents.Engine_AssetLog;
 		}
 
+		/// <summary>
+		/// Determines if an event should be masked by this 
+		/// </summary>
+		/// <param name="eventId"></param>
+		/// <returns></returns>
+		static bool IsMaskedEventId(EventId eventId)
+		{
+			return eventId == KnownLogEvents.ExitCode;
+		}
+
 		/// <summary>
 		/// Adds all the assets from the given log event
 		/// </summary>
@@ -57,14 +67,23 @@ namespace Horde.Build.Issues.Handlers
 		/// <inheritdoc/>
 		public override void TagEvents(IJob job, INode node, IReadOnlyNodeAnnotations annotations, IReadOnlyList<IssueEvent> stepEvents)
 		{
+			bool hasMatches = false;
 			foreach (IssueEvent stepEvent in stepEvents)
 			{
-				if (stepEvent.EventId != null && IsMatchingEventId(stepEvent.EventId.Value))
+				if (stepEvent.EventId != null)
 				{
-					HashSet<string> newAssetNames = new HashSet<string>();
-					GetAssetNames(stepEvent.EventData, newAssetNames);
+					if (IsMatchingEventId(stepEvent.EventId.Value))
+					{
+						HashSet<string> newAssetNames = new HashSet<string>();
+						GetAssetNames(stepEvent.EventData, newAssetNames);
 
-					stepEvent.Fingerprint = new NewIssueFingerprint(Type, newAssetNames, null, null);
+						stepEvent.Fingerprint = new NewIssueFingerprint(Type, newAssetNames, null, null);
+						hasMatches = true;
+					}
+					else if (hasMatches && IsMaskedEventId(stepEvent.EventId.Value))
+					{
+						stepEvent.Ignored = true;
+					}
 				}
 			}
 		}
diff --git a/Engine/Source/Programs/Horde/Horde.Build/Issues/Handlers/LocalizationIssueHandler.cs b/Engine/Source/Programs/Horde/Horde.Build/Issues/Handlers/LocalizationIssueHandler.cs
--- a/Engine/Source/Programs/Horde/Horde.Build/Issues/Handlers/LocalizationIssueHandler.cs
+++ b/Engine/Source/Programs/Horde/Horde.Build/Issues/Handlers/LocalizationIssueHandler.cs
@@ -35,6 +35,16 @@ namespace Horde.Build.Issues.Handlers
 			return eventId == KnownLogEvents.Engine_Localization;
 		}
 
+		/// <summary>
+		/// Determines if an event should be masked by this 
+		/// </summary>
+		/// <param name="eventId"></param>
+		/// <returns></returns>
+		static bool IsMaskedEventId(EventId eventId)
+		{
+			return eventId == KnownLogEvents.ExitCode;
+		}
+
 		/// <summary>
 		/// Extracts a list of source files from an event
 		/// </summary>
@@ -60,20 +70,30 @@ namespace Horde.Build.Issues.Handlers
 		/// <inheritdoc/>
 		public override void TagEvents(IJob job, INode node, IReadOnlyNodeAnnotations annotations, IReadOnlyList<IssueEvent> stepEvents)
 		{
+			bool hasMatches = false;
 			foreach (IssueEvent stepEvent in stepEvents)
 			{
-				if (stepEvent.EventId != null && IsMatchingEventId(stepEvent.EventId.Value))
+				if (stepEvent.EventId != null)
 				{
-					HashSet<string> newFileNames = new HashSet<string>(StringComparer.OrdinalIgnoreCase);
-					GetSourceFiles(stepEvent.EventData, newFileNames);
-
-					if (newFileNames.Count == 0)
+					if (IsMatchingEventId(stepEvent.EventId.Value))
 					{
-						stepEvent.Ignored = true;
+						HashSet<string> newFileNames = new HashSet<string>(StringComparer.OrdinalIgnoreCase);
+						GetSourceFiles(stepEvent.EventData, newFileNames);
+
+						if (newFileNames.Count == 0)
+						{
+							stepEvent.Ignored = true;
+						}
+						else
+						{
+							stepEvent.Fingerprint = new NewIssueFingerprint(Type, newFileNames, null, null);
+						}
+
+						hasMatches = true;
 					}
-					else
+					else if (hasMatches && IsMaskedEventId(stepEvent.EventId.Value))
 					{
-						stepEvent.Fingerprint = new NewIssueFingerprint(Type, newFileNames, null, null);
+						stepEvent.Ignored = true;
 					}
 				}
 			}
diff --git a/Engine/Source/Programs/Horde/Horde.Build/Issues/Handlers/SymbolIssueHandler.cs b/Engine/Source/Programs/Horde/Horde.Build/Issues/Handlers/SymbolIssueHandler.cs
--- a/Engine/Source/Programs/Horde/Horde.Build/Issues/Handlers/SymbolIssueHandler.cs
+++ b/Engine/Source/Programs/Horde/Horde.Build/Issues/Handlers/SymbolIssueHandler.cs
@@ -40,6 +40,16 @@ namespace Horde.Build.Issues.Handlers
 			return eventId == KnownLogEvents.Linker_UndefinedSymbol || eventId == KnownLogEvents.Linker_DuplicateSymbol || eventId == KnownLogEvents.Linker;
 		}
 
+		/// <summary>
+		/// Determines if an event should be masked by this 
+		/// </summary>
+		/// <param name="eventId"></param>
+		/// <returns></returns>
+		static bool IsMaskedEventId(EventId eventId)
+		{
+			return eventId == KnownLogEvents.ExitCode;
+		}
+
 		/// <summary>
 		/// Parses symbol names from a log event
 		/// </summary>
@@ -134,12 +144,9 @@ namespace Horde.Build.Issues.Handlers
 							stepEvent.Ignored = true;
 						}
 					}
-					else if (eventId == KnownLogEvents.Generic || eventId == KnownLogEvents.AutomationTool)
+					else if (hasMatches && IsMaskedEventId(eventId))
 					{
-						if (hasMatches)
-						{
-							stepEvent.Ignored = true;
-						}
+						stepEvent.Ignored = true;
 					}
 				}
 			}
diff --git a/Engine/Source/Programs/Horde/Horde.Build/Issues/Handlers/SystemicIssueHandler.cs b/Engine/Source/Programs/Horde/Horde.Build/Issues/Handlers/SystemicIssueHandler.cs
--- a/Engine/Source/Programs/Horde/Horde.Build/Issues/Handlers/SystemicIssueHandler.cs
+++ b/Engine/Source/Programs/Horde/Horde.Build/Issues/Handlers/SystemicIssueHandler.cs
@@ -18,6 +18,8 @@ namespace Horde.Build.Issues.Handlers
 	/// </summary>
 	class SystemicIssueHandler : IssueHandler
 	{
+		const string NodeNameKey = "Node";
+
 		/// <inheritdoc/>
 		public override string Type => "Systemic";
 
@@ -47,34 +49,19 @@ namespace Horde.Build.Issues.Handlers
 		/// <inheritdoc/>
 		public override void TagEvents(IJob job, INode node, IReadOnlyNodeAnnotations annotations, IReadOnlyList<IssueEvent> stepEvents)
 		{
+			NewIssueFingerprint? fingerprint = null;
 			foreach (IssueEvent stepEvent in stepEvents)
 			{
 				if (stepEvent.EventId != null && MatchEvent(stepEvent.EventData))
 				{
+					fingerprint ??= new NewIssueFingerprint(Type, new[] { $"step:{job.StreamId}:{job.TemplateId}:{node.Name}" }, null, new[] { $"{NodeNameKey}={node.Name}" });
+					stepEvent.Fingerprint = fingerprint;
 				}
 			}
 		}
 
 		static bool MatchEvent(ILogEventData eventData)
 		{
-			if (eventData.EventId == KnownLogEvents.ExitCode)
-			{
-				for (int i = 0; i < eventData.Lines.Count; i++)
-				{
-					string message = eventData.Lines[i].Message;
-
-					string[] systemicExitMessages = new string[] { "AutomationTool exiting with ExitCode", "BUILD FAILED", "tool returned code" };
-
-					for (int j = 0; j < systemicExitMessages.Length; j++)
-					{
-						if (message.Contains(systemicExitMessages[j], StringComparison.InvariantCultureIgnoreCase))
-						{
-							return true;
-						}
-					}
-				}
-			}
-
 			return eventData.EventId != null && IsMatchingEventId(eventData.EventId.Value);
 		}
 
@@ -82,7 +69,7 @@ namespace Horde.Build.Issues.Handlers
 		public override string GetSummary(IIssueFingerprint fingerprint, IssueSeverity severity)
 		{
 			string type = (severity == IssueSeverity.Warning) ? "Systemic warnings" : "Systemic errors";
-			string nodeName = fingerprint.Keys.FirstOrDefault() ?? "(unknown)";
+			string nodeName = fingerprint.GetMetadataValues(NodeNameKey).FirstOrDefault() ?? "(unknown)";
 			return $"{type} in {nodeName}";
 		}
 	}
diff --git a/Engine/Source/Programs/Horde/Horde.Build/Issues/IssueService.cs b/Engine/Source/Programs/Horde/Horde.Build/Issues/IssueService.cs
--- a/Engine/Source/Programs/Horde/Horde.Build/Issues/IssueService.cs
+++ b/Engine/Source/Programs/Horde/Horde.Build/Issues/IssueService.cs
@@ -610,20 +610,6 @@ namespace Horde.Build.Issues
 			foreach (IssueHandler handler in _handlers)
 			{
 				handler.TagEvents(job, node, annotations, remainingEvents);
-
-				bool ignoreSystemicEvents = false;
-				foreach (IssueEvent remainingEvent in remainingEvents)
-				{
-					if (remainingEvent.Fingerprint != null && remainingEvent.Severity == EventSeverity.Error)
-					{
-						ignoreSystemicEvents = true;
-					}
-					else if (ignoreSystemicEvents && remainingEvent.IsSystemic())
-					{
-						remainingEvent.Ignored = true;
-					}
-				}
-
 				remainingEvents.RemoveAll(x => x.Ignored || x.Fingerprint != null);
 			}
 
