Horde: Minor formatting changes for Slack messages.

diff --git a/Engine/Source/Programs/Horde/Horde.Build/Notifications/Sinks/SlackNotificationSink.cs b/Engine/Source/Programs/Horde/Horde.Build/Notifications/Sinks/SlackNotificationSink.cs
--- a/Engine/Source/Programs/Horde/Horde.Build/Notifications/Sinks/SlackNotificationSink.cs
+++ b/Engine/Source/Programs/Horde/Horde.Build/Notifications/Sinks/SlackNotificationSink.cs
@@ -1298,7 +1298,7 @@ namespace Horde.Build.Notifications.Sinks
 								template = $"<{templateUrl}|{template}>";
 							}
 
-							StringBuilder body = new StringBuilder($"Issues affecting *{template}*:");
+							StringBuilder body = new StringBuilder($"*{template}*:");
 							foreach ((IIssue issue, IIssueSpan span) in pairs.OrderBy(x => x.Key.Id))
 							{
 								body.Append('\n');
@@ -1355,7 +1355,7 @@ namespace Horde.Build.Notifications.Sinks
 				}
 				if (issue.AcknowledgedAt == null)
 				{
-					status = $"{status} (unacknowledged)";
+					status = $"{status} (pending)";
 				}
 			}
 
