Horde Storage - Removed usage of legacy ref store from BlobStoreConsistencyCheck, while correct to consider those namespaces as well practically it doesnt matter much and this causes issues when the legacy setup is not configured (which it shouldnt be for anyone as it is deprecated)

diff --git a/Engine/Source/Programs/Horde/HordeStorage/Horde.Storage/Implementation/Blob/BlobStoreConsistencyCheckService.cs b/Engine/Source/Programs/Horde/HordeStorage/Horde.Storage/Implementation/Blob/BlobStoreConsistencyCheckService.cs
--- a/Engine/Source/Programs/Horde/HordeStorage/Horde.Storage/Implementation/Blob/BlobStoreConsistencyCheckService.cs
+++ b/Engine/Source/Programs/Horde/HordeStorage/Horde.Storage/Implementation/Blob/BlobStoreConsistencyCheckService.cs
@@ -28,7 +28,6 @@ namespace Horde.Storage.Implementation
         private readonly IOptionsMonitor<HordeStorageSettings> _hordeStorageSettings;
         private readonly IServiceProvider _provider;
         private readonly ILeaderElection _leaderElection;
-        private readonly IRefsStore _refsStore;
         private readonly IReferencesStore _referencesStore;
         private readonly IBlobIndex _blobIndex;
         private readonly ILogger _logger = Log.ForContext<BlobStoreConsistencyCheckService>();
@@ -38,13 +37,12 @@ namespace Horde.Storage.Implementation
             return _settings.CurrentValue.EnableBlobStoreChecks;
         }
 
-        public BlobStoreConsistencyCheckService(IOptionsMonitor<ConsistencyCheckSettings> settings, IOptionsMonitor<HordeStorageSettings> hordeStorageSettings, IServiceProvider provider, ILeaderElection leaderElection, IRefsStore refsStore, IReferencesStore referencesStore, IBlobIndex blobIndex) : base(serviceName: nameof(BlobStoreConsistencyCheckService), TimeSpan.FromSeconds(settings.CurrentValue.ConsistencyCheckPollFrequencySeconds), new ConsistencyState())
+        public BlobStoreConsistencyCheckService(IOptionsMonitor<ConsistencyCheckSettings> settings, IOptionsMonitor<HordeStorageSettings> hordeStorageSettings, IServiceProvider provider, ILeaderElection leaderElection, IReferencesStore referencesStore, IBlobIndex blobIndex) : base(serviceName: nameof(BlobStoreConsistencyCheckService), TimeSpan.FromSeconds(settings.CurrentValue.ConsistencyCheckPollFrequencySeconds), new ConsistencyState())
         {
             _settings = settings;
             _hordeStorageSettings = hordeStorageSettings;
             _provider = provider;
             _leaderElection = leaderElection;
-            _refsStore = refsStore;
             _referencesStore = referencesStore;
             _blobIndex = blobIndex;
         }
@@ -82,8 +80,7 @@ namespace Horde.Storage.Implementation
                     continue;
                 }
 
-                List<NamespaceId> namespaces = await _refsStore.GetNamespaces().ToListAsync();
-                namespaces.AddRange(await _referencesStore.GetNamespaces().ToListAsync());
+                List<NamespaceId> namespaces = await _referencesStore.GetNamespaces().ToListAsync();
 
                 // technically this does not need to be run per namespace but per storage pool
                 await foreach (NamespaceId ns in namespaces)
