Horde: Enable configuring of auto-scale strategy settings through REST API

Strategy specific settings were not hooked up before.

diff --git a/Engine/Source/Programs/Horde/Horde.Build.Tests/PoolsControllerTest.cs b/Engine/Source/Programs/Horde/Horde.Build.Tests/PoolsControllerTest.cs
--- a/Engine/Source/Programs/Horde/Horde.Build.Tests/PoolsControllerTest.cs
+++ b/Engine/Source/Programs/Horde/Horde.Build.Tests/PoolsControllerTest.cs
@@ -57,13 +57,16 @@ namespace Horde.Build.Tests
 	        await PoolsController.UpdatePoolAsync(pool1.Id.ToString(), new UpdatePoolRequest
 	        {
 		        Name = "Pool1Modified",
-		        SizeStrategy = PoolSizeStrategy.JobQueue
+		        SizeStrategy = PoolSizeStrategy.JobQueue,
+		        JobQueueSettings = new JobQueueSettingsMessage { ScaleOutFactor = 25.0, ScaleInFactor = 0.3 }
 	        });
-
+	        
 	        ActionResult<object> getResult = await PoolsController.GetPoolAsync(pool1.Id.ToString());
 	        GetPoolResponse response = (getResult.Value! as GetPoolResponse)!;
 	        Assert.AreEqual("Pool1Modified", response.Name);
 	        Assert.AreEqual(PoolSizeStrategy.JobQueue, response.SizeStrategy);
+	        Assert.AreEqual(25.0, response.JobQueueSettings!.ScaleOutFactor);
+	        Assert.AreEqual(0.3, response.JobQueueSettings!.ScaleInFactor);
         }
     }
 }
\ No newline at end of file
diff --git a/Engine/Source/Programs/Horde/Horde.Build/Agents/Pools/PoolMessages.cs b/Engine/Source/Programs/Horde/Horde.Build/Agents/Pools/PoolMessages.cs
--- a/Engine/Source/Programs/Horde/Horde.Build/Agents/Pools/PoolMessages.cs
+++ b/Engine/Source/Programs/Horde/Horde.Build/Agents/Pools/PoolMessages.cs
@@ -25,9 +25,9 @@ namespace Horde.Build.Agents.Pools
 		/// Convert this public REST API object to an internal representation
 		/// </summary>
 		/// <returns></returns>
-		public Horde.Build.Agents.Fleet.LeaseUtilizationSettings Convert()
+		public LeaseUtilizationSettings Convert()
 		{
-			return new Horde.Build.Agents.Fleet.LeaseUtilizationSettings();
+			return new LeaseUtilizationSettings();
 		}
 	}
 	
@@ -48,6 +48,13 @@ namespace Horde.Build.Agents.Pools
 		/// </summary>
 		public double ScaleInFactor { get; set; }
 
+		/// <summary>
+		/// Constructor
+		/// </summary>
+		public JobQueueSettingsMessage()
+		{
+		}
+
 		/// <summary>
 		/// Construct a public REST API representation from the internal one
 		/// </summary>
@@ -193,12 +200,12 @@ namespace Horde.Build.Agents.Pools
 		/// <summary>
 		/// Settings for lease utilization pool sizing strategy (if used)
 		/// </summary>
-		public LeaseUtilizationSettings? LeaseUtilizationSettings { get; set; }
+		public LeaseUtilizationSettingsMessage? LeaseUtilizationSettings { get; set; }
 		
 		/// <summary>
 		/// Settings for job queue pool sizing strategy (if used) 
 		/// </summary>
-		public JobQueueSettings? JobQueueSettings { get; set; }
+		public JobQueueSettingsMessage? JobQueueSettings { get; set; }
 
 		/// <summary>
 		/// The minimum nunmber of agents to retain in this pool
diff --git a/Engine/Source/Programs/Horde/Horde.Build/Agents/Pools/PoolService.cs b/Engine/Source/Programs/Horde/Horde.Build/Agents/Pools/PoolService.cs
--- a/Engine/Source/Programs/Horde/Horde.Build/Agents/Pools/PoolService.cs
+++ b/Engine/Source/Programs/Horde/Horde.Build/Agents/Pools/PoolService.cs
@@ -119,6 +119,8 @@ namespace Horde.Build.Agents.Pools
 		/// <param name="scaleOutCooldown">Cooldown time between scale-out events</param>
 		/// <param name="scaleInCooldown">Cooldown time between scale-in events</param>
 		/// <param name="sizeStrategy">New pool sizing strategy for the pool</param>
+		/// <param name="leaseUtilizationSettings">Settings for lease utilization-based strategy</param>
+		/// <param name="jobQueueSettings">Settings for job queue-based strategy</param>
 		/// <param name="useDefaultStrategy">Whether to use the default strategy</param>
 		/// <returns>Async task object</returns>
 		public async Task<IPool?> UpdatePoolAsync(
@@ -132,6 +134,8 @@ namespace Horde.Build.Agents.Pools
 			TimeSpan? scaleOutCooldown = null,
 			TimeSpan? scaleInCooldown = null,
 			PoolSizeStrategy? sizeStrategy = null,
+			LeaseUtilizationSettings? leaseUtilizationSettings = null,
+			JobQueueSettings? jobQueueSettings = null,
 			bool? useDefaultStrategy = null)
 		{
 			for (; pool != null; pool = await _pools.GetAsync(pool.Id))
@@ -147,6 +151,8 @@ namespace Horde.Build.Agents.Pools
 					scaleOutCooldown: scaleOutCooldown,
 					scaleInCooldown: scaleInCooldown,
 					sizeStrategy: sizeStrategy,
+					leaseUtilizationSettings: leaseUtilizationSettings,
+					jobQueueSettings: jobQueueSettings,
 					useDefaultStrategy: useDefaultStrategy);
 				
 				if (newPool != null)
diff --git a/Engine/Source/Programs/Horde/Horde.Build/Agents/Pools/PoolsController.cs b/Engine/Source/Programs/Horde/Horde.Build/Agents/Pools/PoolsController.cs
--- a/Engine/Source/Programs/Horde/Horde.Build/Agents/Pools/PoolsController.cs
+++ b/Engine/Source/Programs/Horde/Horde.Build/Agents/Pools/PoolsController.cs
@@ -147,7 +147,8 @@ namespace Horde.Build.Agents.Pools
 			TimeSpan? scaleInCooldown = update.ScaleInCooldown == null ? null : TimeSpan.FromSeconds(update.ScaleInCooldown.Value);
 
 			await _poolService.UpdatePoolAsync(pool, update.Name, update.Condition, update.EnableAutoscaling,
-				update.MinAgents, update.NumReserveAgents, update.Properties, scaleOutCooldown, scaleInCooldown, update.SizeStrategy, update.UseDefaultStrategy);
+				update.MinAgents, update.NumReserveAgents, update.Properties, scaleOutCooldown, scaleInCooldown, update.SizeStrategy,
+				update.LeaseUtilizationSettings?.Convert(), update.JobQueueSettings?.Convert(), update.UseDefaultStrategy);
 			return new OkResult();
 		}
 
