Horde Storage - Reworked how access control is managed within Horde Storage to align more with how Horde Build works. You now map ACLs to actions (from claims) which controls what you can do in each namespace. This allows us to grant read or write access to individual namespaces as well as moving more of the mapping from claims into configuration - thus supporting whatever claim setup licensees might have from their IdP

diff --git a/Engine/Source/Programs/Horde/HordeStorage/Callisto/Controllers/TransactionLogController.cs b/Engine/Source/Programs/Horde/HordeStorage/Callisto/Controllers/TransactionLogController.cs
--- a/Engine/Source/Programs/Horde/HordeStorage/Callisto/Controllers/TransactionLogController.cs
+++ b/Engine/Source/Programs/Horde/HordeStorage/Callisto/Controllers/TransactionLogController.cs
@@ -48,7 +48,11 @@ namespace Callisto.Controllers
             // filter namespaces down to only the namespaces the user has access to
             namespaces = namespaces.Where(ns =>
             {
-                Task<AuthorizationResult> authorizationResult = _authorizationService.AuthorizeAsync(User, ns, NamespaceAccessRequirement.Name);
+                Task<AuthorizationResult> authorizationResult = _authorizationService.AuthorizeAsync(User, new NamespaceAccessRequest()
+                {
+                    Namespace = ns,
+                    Actions = new [] { AclAction.ReadTransactionLog }
+                }, NamespaceAccessRequirement.Name);
                 return authorizationResult.Result.Succeeded;
             }).ToArray();
 
@@ -70,7 +74,11 @@ namespace Callisto.Controllers
             [FromQuery] int maxOffsetsAttempted = 100
         )
         {
-            AuthorizationResult authorizationResult = await _authorizationService.AuthorizeAsync(User, ns, NamespaceAccessRequirement.Name);
+            AuthorizationResult authorizationResult = await _authorizationService.AuthorizeAsync(User, new NamespaceAccessRequest()
+            {
+                Namespace = ns,
+                Actions = new[] { AclAction.ReadTransactionLog }
+            }, NamespaceAccessRequirement.Name);
 
             if (!authorizationResult.Succeeded)
             {
@@ -129,7 +137,11 @@ namespace Callisto.Controllers
             [Required] NamespaceId ns
         )
         {
-            AuthorizationResult authorizationResult = await _authorizationService.AuthorizeAsync(User, ns, NamespaceAccessRequirement.Name);
+            AuthorizationResult authorizationResult = await _authorizationService.AuthorizeAsync(User, new NamespaceAccessRequest()
+            {
+                Namespace = ns,
+                Actions = new[] { AclAction.WriteTransactionLog }
+            }, NamespaceAccessRequirement.Name);
 
             if (!authorizationResult.Succeeded)
             {
@@ -154,7 +166,11 @@ namespace Callisto.Controllers
             [Required] [FromBody] TransactionEvent request
         )
         {
-            AuthorizationResult authorizationResult = await _authorizationService.AuthorizeAsync(User, ns, NamespaceAccessRequirement.Name);
+            AuthorizationResult authorizationResult = await _authorizationService.AuthorizeAsync(User, new NamespaceAccessRequest()
+            {
+                Namespace = ns,
+                Actions = new[] { AclAction.WriteTransactionLog }
+            }, NamespaceAccessRequirement.Name);
 
             if (!authorizationResult.Succeeded)
             {
@@ -198,7 +214,11 @@ namespace Callisto.Controllers
             [Required] NamespaceId ns
         )
         {
-            AuthorizationResult authorizationResult = await _authorizationService.AuthorizeAsync(User, ns, NamespaceAccessRequirement.Name);
+            AuthorizationResult authorizationResult = await _authorizationService.AuthorizeAsync(User, new NamespaceAccessRequest()
+            {
+                Namespace = ns,
+                Actions = new[] { AclAction.WriteTransactionLog }
+            }, NamespaceAccessRequirement.Name);
 
             if (!authorizationResult.Succeeded)
             {
diff --git a/Engine/Source/Programs/Horde/HordeStorage/Horde.Storage/Controllers/AdminController.cs b/Engine/Source/Programs/Horde/HordeStorage/Horde.Storage/Controllers/AdminController.cs
--- a/Engine/Source/Programs/Horde/HordeStorage/Horde.Storage/Controllers/AdminController.cs
+++ b/Engine/Source/Programs/Horde/HordeStorage/Horde.Storage/Controllers/AdminController.cs
@@ -9,6 +9,7 @@ using System.Threading;
 using System.Threading.Tasks;
 using EpicGames.Horde.Storage;
 using Horde.Storage.Implementation;
+using Jupiter;
 using Jupiter.Implementation;
 using Microsoft.AspNetCore.Authorization;
 using Microsoft.AspNetCore.Mvc;
@@ -20,17 +21,20 @@ namespace Horde.Storage.Controllers
     [ApiController]
     [Route("api/v1/admin")]
     [InternalApiFilter]
+    [Authorize]
     public class AdminController : Controller
     {
         private readonly LastAccessService _lastAccessService;
         private readonly IRefCleanup _refCleanup;
         private readonly IConfiguration _configuration;
+        private readonly RequestHelper _requestHelper;
 
-        public AdminController(LastAccessService lastAccessService, IRefCleanup refCleanup, IConfiguration configuration)
+        public AdminController(LastAccessService lastAccessService, IRefCleanup refCleanup, IConfiguration configuration, RequestHelper requestHelper)
         {
             _lastAccessService = lastAccessService;
             _refCleanup = refCleanup;
             _configuration = configuration;
+            _requestHelper = requestHelper;
         }
 
         /// <summary>
@@ -41,10 +45,15 @@ namespace Horde.Storage.Controllers
         /// </remarks>
         /// <returns></returns>
         [HttpPost("startLastAccessRollup")]
-        [Authorize("Admin")]
         [ProducesResponseType(type: typeof(UpdatedRecordsResponse), 200)]
         public async Task<IActionResult> StartLastAccessRollup()
         {
+            ActionResult? result = await _requestHelper.HasAccessForGlobalOperations(User, new [] { AclAction.AdminAction });
+            if (result != null)
+            {
+                return result;
+            }
+
             Task<List<(RefRecord, DateTime)>>? updateRecordsTask = _lastAccessService.ProcessLastAccessRecords();
             List<(RefRecord, DateTime)>? updatedRecords = null;
             if (updateRecordsTask != null)
@@ -68,6 +77,12 @@ namespace Horde.Storage.Controllers
         [HttpPost("refCleanup/{ns}")]
         public async Task<IActionResult> RefCleanup([FromRoute] [Required] NamespaceId ns)
         {
+            ActionResult? result = await _requestHelper.HasAccessToNamespace(User, Request, ns, new [] { AclAction.AdminAction });
+            if (result != null)
+            {
+                return result;
+            }
+
             int countOfDeletedRecords = await _refCleanup.Cleanup(ns, CancellationToken.None);
             return Ok(new RemovedRefRecordsResponse(countOfDeletedRecords));
         }
@@ -77,8 +92,14 @@ namespace Horde.Storage.Controllers
         /// </summary>
         /// <returns></returns>
         [HttpGet("settings")]
-        public IActionResult Settings()
+        public async Task<IActionResult> Settings()
         {
+            ActionResult? result = await _requestHelper.HasAccessForGlobalOperations(User, new [] { AclAction.AdminAction });
+            if (result != null)
+            {
+                return result;
+            }
+
             Dictionary<string, Dictionary<string, object>> settings = new Dictionary<string, Dictionary<string, object>>();
 
             Dictionary<string, object> ResolveSection(IConfigurationSection section)
diff --git a/Engine/Source/Programs/Horde/HordeStorage/Horde.Storage/Controllers/AuthController.cs b/Engine/Source/Programs/Horde/HordeStorage/Horde.Storage/Controllers/AuthController.cs
--- a/Engine/Source/Programs/Horde/HordeStorage/Horde.Storage/Controllers/AuthController.cs
+++ b/Engine/Source/Programs/Horde/HordeStorage/Horde.Storage/Controllers/AuthController.cs
@@ -1,31 +1,41 @@
 ﻿// Copyright Epic Games, Inc. All Rights Reserved.
 
+using System.Collections.Generic;
 using Microsoft.AspNetCore.Mvc;
 using System.ComponentModel.DataAnnotations;
 using System.Threading.Tasks;
 using Microsoft.AspNetCore.Authorization;
 using EpicGames.Horde.Storage;
+using Jupiter;
+using Jupiter.Common;
+using Microsoft.Extensions.Options;
+using Newtonsoft.Json;
+using Newtonsoft.Json.Converters;
 
 namespace Horde.Storage.Controllers
 {
     [ApiController]
     [Route("api/v1/auth")]
+    [Authorize]
     public class AuthController : ControllerBase
     {
         private readonly RequestHelper _requestHelper;
+        private readonly INamespacePolicyResolver _namespacePolicyResolver;
+        private readonly IOptionsMonitor<AuthSettings> _authSettings;
 
-        public AuthController(RequestHelper requestHelper)
+        public AuthController(RequestHelper requestHelper, INamespacePolicyResolver namespacePolicyResolver, IOptionsMonitor<AuthSettings> authSettings)
         {
             _requestHelper = requestHelper;
+            _namespacePolicyResolver = namespacePolicyResolver;
+            _authSettings = authSettings;
         }
 
         [HttpGet("{ns}")]
-        [Authorize("Any")]
         public async Task<IActionResult> Verify(
             [FromRoute] [Required] NamespaceId ns
             )
         {
-            ActionResult? result = await _requestHelper.HasAccessToNamespace(User, Request, ns);
+            ActionResult? result = await _requestHelper.HasAccessToNamespace(User, Request, ns, new [] { AclAction.ReadObject });
             if (result != null)
             {
                 return result;
@@ -33,5 +43,34 @@ namespace Horde.Storage.Controllers
 
             return Ok();
         }
+
+        [HttpGet("{ns}/actions")]
+        public IActionResult Actions(
+            [FromRoute] [Required] NamespaceId ns
+        )
+        {
+            NamespacePolicy policy = _namespacePolicyResolver.GetPoliciesForNs(ns);
+
+            List<AclAction> allowedActions = new List<AclAction>();
+            foreach (AclEntry acl in policy.Acls)
+            {
+                allowedActions.AddRange(acl.Resolve(User));
+            }
+
+            // the root and namespace acls are combined, namespace acls can not override what we define in the root
+            foreach (AclEntry acl in _authSettings.CurrentValue.Acls)
+            {
+                allowedActions.AddRange(acl.Resolve(User));
+            }
+
+            return Ok(new ActionsResult {Actions = allowedActions});
+        }
+    }
+
+    public class ActionsResult
+    {
+        [JsonProperty (ItemConverterType = typeof(StringEnumConverter))]
+        [System.Diagnostics.CodeAnalysis.SuppressMessage("Usage", "CA2227:Collection properties should be read only", Justification = "Used by serialization")]
+        public List<AclAction> Actions { get; set; } = new List<AclAction>();
     }
 }
diff --git a/Engine/Source/Programs/Horde/HordeStorage/Horde.Storage/Controllers/CompressedBlobController.cs b/Engine/Source/Programs/Horde/HordeStorage/Horde.Storage/Controllers/CompressedBlobController.cs
--- a/Engine/Source/Programs/Horde/HordeStorage/Horde.Storage/Controllers/CompressedBlobController.cs
+++ b/Engine/Source/Programs/Horde/HordeStorage/Horde.Storage/Controllers/CompressedBlobController.cs
@@ -26,6 +26,7 @@ namespace Horde.Storage.Controllers
     using BlobNotFoundException = Horde.Storage.Implementation.BlobNotFoundException;
 
     [ApiController]
+	[Authorize]
     [Route("api/v1/compressed-blobs")]
     public class CompressedBlobController : ControllerBase
     {
@@ -47,7 +48,6 @@ namespace Horde.Storage.Controllers
         }
 
         [HttpGet("{ns}/{id}")]
-        [Authorize("Storage.read")]
         [ProducesResponseType(type: typeof(byte[]), 200)]
         [ProducesResponseType(type: typeof(ValidationProblemDetails), 400)]
         [Produces(CustomMediaTypeNames.UnrealCompressedBuffer, MediaTypeNames.Application.Octet)]
@@ -56,7 +56,7 @@ namespace Horde.Storage.Controllers
             [Required] NamespaceId ns,
             [Required] ContentId id)
         {
-            ActionResult? result = await _requestHelper.HasAccessToNamespace(User, Request, ns);
+            ActionResult? result = await _requestHelper.HasAccessToNamespace(User, Request, ns, new [] { AclAction.ReadObject });
             if (result != null)
             {
                 return result;
@@ -85,13 +85,12 @@ namespace Horde.Storage.Controllers
         }
         
         [HttpHead("{ns}/{id}")]
-        [Authorize("Storage.read")]
         [ProducesDefaultResponseType]
         public async Task<IActionResult> Head(
             [Required] NamespaceId ns,
             [Required] ContentId id)
         {
-            ActionResult? result = await _requestHelper.HasAccessToNamespace(User, Request, ns);
+            ActionResult? result = await _requestHelper.HasAccessToNamespace(User, Request, ns, new [] { AclAction.ReadObject });
             if (result != null)
             {
                 return result;
@@ -122,13 +121,12 @@ namespace Horde.Storage.Controllers
         }
 
         [HttpPost("{ns}/exists")]
-        [Authorize("Storage.read")]
         [ProducesDefaultResponseType]
         public async Task<IActionResult> ExistsMultiple(
             [Required] NamespaceId ns,
             [Required] [FromQuery] List<ContentId> id)
         {
-            ActionResult? result = await _requestHelper.HasAccessToNamespace(User, Request, ns);
+            ActionResult? result = await _requestHelper.HasAccessToNamespace(User, Request, ns, new [] { AclAction.ReadObject });
             if (result != null)
             {
                 return result;
@@ -165,13 +163,12 @@ namespace Horde.Storage.Controllers
         }
 
         [HttpPost("{ns}/exist")]
-        [Authorize("Storage.read")]
         [ProducesDefaultResponseType]
         public async Task<IActionResult> ExistsBody(
             [Required] NamespaceId ns,
             [FromBody] ContentId[] bodyIds)
         {
-            ActionResult? result = await _requestHelper.HasAccessToNamespace(User, Request, ns);
+            ActionResult? result = await _requestHelper.HasAccessToNamespace(User, Request, ns, new [] { AclAction.ReadObject });
             if (result != null)
             {
                 return result;
@@ -252,14 +249,13 @@ namespace Horde.Storage.Controllers
         }
 
         [HttpPut("{ns}/{id}")]
-        [Authorize("Storage.write")]
         [DisableRequestSizeLimit]
         [RequiredContentType(CustomMediaTypeNames.UnrealCompressedBuffer)]
         public async Task<IActionResult> Put(
             [Required] NamespaceId ns,
             [Required] ContentId id)
         {
-            ActionResult? result = await _requestHelper.HasAccessToNamespace(User, Request, ns);
+            ActionResult? result = await _requestHelper.HasAccessToNamespace(User, Request, ns, new [] { AclAction.WriteObject });
             if (result != null)
             {
                 return result;
@@ -288,13 +284,12 @@ namespace Horde.Storage.Controllers
         }
 
         [HttpPost("{ns}")]
-        [Authorize("Storage.write")]
         [DisableRequestSizeLimit]
         [RequiredContentType(CustomMediaTypeNames.UnrealCompressedBuffer)]
         public async Task<IActionResult> Post(
             [Required] NamespaceId ns)
         {
-            ActionResult? result = await _requestHelper.HasAccessToNamespace(User, Request, ns);
+            ActionResult? result = await _requestHelper.HasAccessToNamespace(User, Request, ns, new [] { AclAction.WriteObject });
             if (result != null)
             {
                 return result;
@@ -369,7 +364,6 @@ namespace Horde.Storage.Controllers
         }
 
         /*[HttpDelete("{ns}/{id}")]
-        [Authorize("Storage.delete")]
         public async Task<IActionResult> Delete(
             [Required] string ns,
             [Required] BlobIdentifier id)
@@ -387,7 +381,6 @@ namespace Horde.Storage.Controllers
 
         
         [HttpDelete("{ns}")]
-        [Authorize("Admin")]
         public async Task<IActionResult> DeleteNamespace(
             [Required] string ns)
         {
diff --git a/Engine/Source/Programs/Horde/HordeStorage/Horde.Storage/Controllers/ContentIdController.cs b/Engine/Source/Programs/Horde/HordeStorage/Horde.Storage/Controllers/ContentIdController.cs
--- a/Engine/Source/Programs/Horde/HordeStorage/Horde.Storage/Controllers/ContentIdController.cs
+++ b/Engine/Source/Programs/Horde/HordeStorage/Horde.Storage/Controllers/ContentIdController.cs
@@ -15,6 +15,7 @@ namespace Horde.Storage.Controllers
     [FormatFilter]
     [Produces(MediaTypeNames.Application.Json, MediaTypeNames.Application.Octet, CustomMediaTypeNames.UnrealCompactBinary)]
     [Route("api/v1/content-id")]
+    [Authorize]
     public class ContentIdController : ControllerBase
     {
         private readonly RequestHelper _requestHelper;
@@ -34,10 +35,9 @@ namespace Horde.Storage.Controllers
         [HttpGet("{ns}/{contentId}.{format?}", Order = 500)]
         [ProducesDefaultResponseType]
         [ProducesResponseType(type: typeof(ProblemDetails), 400)]
-        [Authorize("Cache.read")]
         public async Task<IActionResult> Resolve(NamespaceId ns, ContentId contentId)
         {
-            ActionResult? result = await _requestHelper.HasAccessToNamespace(User, Request, ns);
+            ActionResult? result = await _requestHelper.HasAccessToNamespace(User, Request, ns, new [] { AclAction.ReadObject });
             if (result != null)
             {
                 return result;
@@ -62,10 +62,9 @@ namespace Horde.Storage.Controllers
         [HttpPut("{ns}/{contentId}/update/{blobIdentifier}/{contentWeight}", Order = 500)]
         [ProducesDefaultResponseType]
         [ProducesResponseType(type: typeof(ProblemDetails), 400)]
-        [Authorize("Cache.write")]
         public async Task<IActionResult> UpdateContentIdMapping(NamespaceId ns, ContentId contentId, BlobIdentifier blobIdentifier, int contentWeight)
         {
-            ActionResult? result = await _requestHelper.HasAccessToNamespace(User, Request, ns);
+            ActionResult? result = await _requestHelper.HasAccessToNamespace(User, Request, ns, new [] { AclAction.WriteObject });
             if (result != null)
             {
                 return result;
diff --git a/Engine/Source/Programs/Horde/HordeStorage/Horde.Storage/Controllers/DDCRefController.cs b/Engine/Source/Programs/Horde/HordeStorage/Horde.Storage/Controllers/DDCRefController.cs
--- a/Engine/Source/Programs/Horde/HordeStorage/Horde.Storage/Controllers/DDCRefController.cs
+++ b/Engine/Source/Programs/Horde/HordeStorage/Horde.Storage/Controllers/DDCRefController.cs
@@ -45,22 +45,21 @@ namespace Horde.Storage.Controllers
     [ApiController]
     [FormatFilter]
     [Route("api/v1/c")]
+    [Authorize]
     public class DDCRefController : ControllerBase
     {
         private readonly IRefsStore _refsStore;
         private readonly IDiagnosticContext _diagnosticContext;
-        private readonly IAuthorizationService _authorizationService;
 		private readonly RequestHelper _requestHelper;
 
         private readonly ILogger _logger = Log.ForContext<DDCRefController>();
         private readonly IDDCRefService _ddcRefService;
 
-        public DDCRefController(IDDCRefService ddcRefService, IRefsStore refsStore, IDiagnosticContext diagnosticContext, IAuthorizationService authorizationService, RequestHelper requestHelper)
+        public DDCRefController(IDDCRefService ddcRefService, IRefsStore refsStore, IDiagnosticContext diagnosticContext, RequestHelper requestHelper)
         {
             _refsStore = refsStore;
             _ddcRefService = ddcRefService;
             _diagnosticContext = diagnosticContext;
-            _authorizationService = authorizationService;
             _requestHelper = requestHelper;
         }
 
@@ -71,7 +70,6 @@ namespace Horde.Storage.Controllers
         [HttpGet("ddc")]
         [ProducesDefaultResponseType]
         [ProducesResponseType(type: typeof(ProblemDetails), 400)]
-        [Authorize("Cache.read")]
         public async Task<IActionResult> GetNamespaces()
         {
             NamespaceId[] namespaces = await _refsStore.GetNamespaces().ToArrayAsync();
@@ -80,7 +78,7 @@ namespace Horde.Storage.Controllers
             List<NamespaceId> validNamespaces = new List<NamespaceId>();
             foreach (NamespaceId ns in namespaces)
             {
-                ActionResult? result = await _requestHelper.HasAccessToNamespace(User, Request, ns);
+                ActionResult? result = await _requestHelper.HasAccessToNamespace(User, Request, ns, new [] { AclAction.ReadObject });
                 if (result == null)
                 {
                     validNamespaces.Add(ns);
@@ -104,7 +102,6 @@ namespace Horde.Storage.Controllers
         [ProducesResponseType(type: typeof(RefResponse), 200)]
         [ProducesResponseType(type: typeof(ValidationProblemDetails), 400)]
         [Produces(MediaTypeNames.Application.Json, MediaTypeNames.Application.Octet, CustomMediaTypeNames.UnrealCompactBinary)]
-        [Authorize("Cache.read")]
         public async Task<IActionResult> Get(
             [FromRoute] [Required] NamespaceId ns,
             [FromRoute] [Required] BucketId bucket,
@@ -112,7 +109,7 @@ namespace Horde.Storage.Controllers
             [FromQuery] string[] fields,
             [FromRoute] string? format = null)
         {
-            ActionResult? result = await _requestHelper.HasAccessToNamespace(User, Request, ns);
+            ActionResult? result = await _requestHelper.HasAccessToNamespace(User, Request, ns, new [] { AclAction.ReadObject });
             if (result != null)
             {
                 return result;
@@ -194,13 +191,12 @@ namespace Horde.Storage.Controllers
         [HttpHead("ddc/{ns}/{bucket}/{key}", Order = 500)]
         [ProducesResponseType(type: typeof(RefResponse), 200)]
         [ProducesResponseType(type: typeof(ValidationProblemDetails), 400)]
-        [Authorize("Cache.read")]
         public async Task<IActionResult> Head(
             [FromRoute] [Required] NamespaceId ns,
             [FromRoute] [Required] BucketId bucket,
             [FromRoute] [Required] KeyId key)
         {
-            ActionResult? result = await _requestHelper.HasAccessToNamespace(User, Request, ns);
+            ActionResult? result = await _requestHelper.HasAccessToNamespace(User, Request, ns, new [] { AclAction.ReadObject });
             if (result != null)
             {
                 return result;
@@ -240,14 +236,13 @@ namespace Horde.Storage.Controllers
         [Consumes(MediaTypeNames.Application.Json)]
         [ProducesResponseType(type: typeof(PutRequestResponse), 200)]
         [ProducesResponseType(type: typeof(ValidationProblemDetails), 400)]
-        [Authorize("Cache.write")]
         public async Task<IActionResult> PutStructured(
             [FromRoute] [Required] NamespaceId ns,
             [FromRoute] [Required] BucketId bucket,
             [FromRoute] [Required] KeyId key,
             [FromBody] RefRequest refRequest)
         {
-            ActionResult? result = await _requestHelper.HasAccessToNamespace(User, Request, ns);
+            ActionResult? result = await _requestHelper.HasAccessToNamespace(User, Request, ns, new [] { AclAction.WriteObject });
             if (result != null)
             {
                 return result;
@@ -279,19 +274,15 @@ namespace Horde.Storage.Controllers
         // TODO: Investigate if we can resolve the conflict between this and the other put endpoint in open api
         [ApiExplorerSettings(IgnoreApi = true)]
         [DisableRequestSizeLimit]
-        [Authorize("Cache.write")]
         public async Task<IActionResult> PutBlob(
             [FromRoute] [Required] NamespaceId ns,
             [FromRoute] [Required] BucketId bucket,
             [FromRoute] [Required] KeyId key)
         {
+            ActionResult? result = await _requestHelper.HasAccessToNamespace(User, Request, ns, new [] { AclAction.WriteObject });
+            if (result != null)
             {
-                AuthorizationResult authorizationResult = await _authorizationService.AuthorizeAsync(User, ns, NamespaceAccessRequirement.Name);
-
-                if (!authorizationResult.Succeeded)
-                {
-                    return Forbid();
-                }
+                return result;
             }
 
             byte[] blob;
@@ -348,12 +339,11 @@ namespace Horde.Storage.Controllers
         /// <param name="ns">Namespace. Each namespace is completely separated from each other. Use for different types of data that is never expected to be similar (between two different games for instance)</param>
         [HttpDelete("ddc/{ns}", Order = 500)]
         [ProducesResponseType(204)]
-        [Authorize("admin")]
         public async Task<IActionResult> DeleteNamespace(
             [FromRoute] [Required] NamespaceId ns
         )
         {
-            ActionResult? result = await _requestHelper.HasAccessToNamespace(User, Request, ns);
+            ActionResult? result = await _requestHelper.HasAccessToNamespace(User, Request, ns, new [] { AclAction.DeleteNamespace });
             if (result != null)
             {
                 return result;
@@ -378,12 +368,11 @@ namespace Horde.Storage.Controllers
         /// <param name="bucket">The category/type of record you are caching. Is a clustered key together with the actual key, but all records in the same bucket can be dropped easily.</param>
         [HttpDelete("ddc/{ns}/{bucket}", Order = 500)]
         [ProducesResponseType(204)]
-        [Authorize("Cache.delete")]
         public async Task<IActionResult> DeleteBucket(
             [FromRoute] [Required] NamespaceId ns,
             [FromRoute] [Required] BucketId bucket)
         {
-            ActionResult? result = await _requestHelper.HasAccessToNamespace(User, Request, ns);
+            ActionResult? result = await _requestHelper.HasAccessToNamespace(User, Request, ns, new [] { AclAction.DeleteBucket });
             if (result != null)
             {
                 return result;
@@ -410,13 +399,12 @@ namespace Horde.Storage.Controllers
         [HttpDelete("ddc/{ns}/{bucket}/{key}", Order = 500)]
         [ProducesResponseType(204)]
         [ProducesResponseType(400)]
-        [Authorize("Cache.delete")]
         public async Task<IActionResult> Delete(
             [FromRoute] [Required] NamespaceId ns,
             [FromRoute] [Required] BucketId bucket,
             [FromRoute] [Required] KeyId key)
         {
-            ActionResult? result = await _requestHelper.HasAccessToNamespace(User, Request, ns);
+            ActionResult? result = await _requestHelper.HasAccessToNamespace(User, Request, ns, new [] { AclAction.DeleteObject });
             if (result != null)
             {
                 return result;
@@ -477,17 +465,18 @@ namespace Horde.Storage.Controllers
         [HttpPost("ddc-rpc")]
         public async Task<IActionResult> Post([FromBody] DdcBatchCall batch)
         {
-            string OpToPolicy(DdcBatchOp.DdcOperation op)
+            AclAction MapToAclAction(DdcBatchOp.DdcOperation op)
             {
                 switch (op)
                 {
                     case DdcBatchOp.DdcOperation.GET:
-                    case DdcBatchOp.DdcOperation.HEAD:
-                        return "Cache.read";
+                        return AclAction.ReadObject;
                     case DdcBatchOp.DdcOperation.PUT:
-                        return "Cache.write";
+                        return AclAction.WriteObject;
                     case DdcBatchOp.DdcOperation.DELETE:
-                        return "Cache.delete";
+                        return AclAction.DeleteObject;
+                    case DdcBatchOp.DdcOperation.HEAD:
+                        return AclAction.ReadObject;
                     default:
                         throw new ArgumentOutOfRangeException(nameof(op), op, null);
                 }
@@ -503,12 +492,9 @@ namespace Horde.Storage.Controllers
             {
                 DdcBatchOp op = batch.Operations[index];
 
-                AuthorizationResult authorizationResultNamespace = await _authorizationService.AuthorizeAsync(User, op.Namespace, NamespaceAccessRequirement.Name);
-                AuthorizationResult authorizationResultOp = await _authorizationService.AuthorizeAsync(User, op, OpToPolicy(op.Op));
-                ActionResult? accessToNamespaceResult = await _requestHelper.HasAccessToNamespace(User, Request, op.Namespace!.Value);
-                bool authorized = authorizationResultNamespace.Succeeded && authorizationResultOp.Succeeded && accessToNamespaceResult == null;
+                ActionResult? result = await _requestHelper.HasAccessToNamespace(User, Request, op.Namespace!.Value, new [] { MapToAclAction(op.Op) });
 
-                if (!authorized)
+                if (result != null)
                 {
                     tasks[index] = Task.FromResult((object?)new ProblemDetails { Title = "Forbidden" });
                 }
@@ -609,7 +595,6 @@ namespace Horde.Storage.Controllers
         /// <param name="batch">Spec for which objects to return</param>
         /// <returns></returns>
         [HttpPost("ddc-rpc/batchGet")]
-        [Authorize("Cache.read")]
         public async Task<IActionResult> BatchGet([FromBody] BatchGetOp batch)
         {
             if (batch.Namespace == null)
@@ -617,7 +602,7 @@ namespace Horde.Storage.Controllers
                 throw new InvalidOperationException("A valid namespace must be specified");
             }
 
-            ActionResult? result = await _requestHelper.HasAccessToNamespace(User, Request, batch.Namespace.Value);
+            ActionResult? result = await _requestHelper.HasAccessToNamespace(User, Request, batch.Namespace.Value, new [] { AclAction.ReadObject });
             if (result != null)
             {
                 return result;
diff --git a/Engine/Source/Programs/Horde/HordeStorage/Horde.Storage/Controllers/DebugController.cs b/Engine/Source/Programs/Horde/HordeStorage/Horde.Storage/Controllers/DebugController.cs
--- a/Engine/Source/Programs/Horde/HordeStorage/Horde.Storage/Controllers/DebugController.cs
+++ b/Engine/Source/Programs/Horde/HordeStorage/Horde.Storage/Controllers/DebugController.cs
@@ -15,7 +15,7 @@ namespace Horde.Storage.Controllers
         /// </summary>
         /// <returns></returns>
         [HttpGet("getBytes")]
-        [Authorize("Any")]
+        [Authorize]
         public IActionResult GetBytes([FromQuery] int length = 1)
         {
             return GenerateByteResponse(length);
diff --git a/Engine/Source/Programs/Horde/HordeStorage/Horde.Storage/Controllers/ObjectController.cs b/Engine/Source/Programs/Horde/HordeStorage/Horde.Storage/Controllers/ObjectController.cs
--- a/Engine/Source/Programs/Horde/HordeStorage/Horde.Storage/Controllers/ObjectController.cs
+++ b/Engine/Source/Programs/Horde/HordeStorage/Horde.Storage/Controllers/ObjectController.cs
@@ -25,6 +25,7 @@ namespace Horde.Storage.Controllers
 
     [ApiController]
     [Route("api/v1/objects", Order = 0)]
+    [Authorize]
     [Produces(CustomMediaTypeNames.UnrealCompactBinary, MediaTypeNames.Application.Json)]
     public class ObjectController : ControllerBase
     {
@@ -46,13 +47,12 @@ namespace Horde.Storage.Controllers
         }
 
         [HttpGet("{ns}/{id}")]
-        [Authorize("Storage.read")]
         [ProducesDefaultResponseType]
         public async Task<IActionResult> Get(
             [Required] NamespaceId ns,
             [Required] BlobIdentifier id)
         {
-            ActionResult? result = await _requestHelper.HasAccessToNamespace(User, Request, ns);
+            ActionResult? result = await _requestHelper.HasAccessToNamespace(User, Request, ns, new [] { AclAction.ReadObject });
             if (result != null)
             {
                 return result;
@@ -71,13 +71,12 @@ namespace Horde.Storage.Controllers
         }
 
         [HttpHead("{ns}/{id}")]
-        [Authorize("Storage.read")]
         [ProducesDefaultResponseType]
         public async Task<IActionResult> Head(
             [Required] NamespaceId ns,
             [Required] BlobIdentifier id)
         {
-            ActionResult? result = await _requestHelper.HasAccessToNamespace(User, Request, ns);
+            ActionResult? result = await _requestHelper.HasAccessToNamespace(User, Request, ns, new [] { AclAction.ReadObject });
             if (result != null)
             {
                 return result;
@@ -94,13 +93,12 @@ namespace Horde.Storage.Controllers
         }
 
         [HttpPost("{ns}/exists")]
-        [Authorize("Storage.read")]
         [ProducesDefaultResponseType]
         public async Task<IActionResult> ExistsMultiple(
             [Required] NamespaceId ns,
             [Required] [FromQuery] List<BlobIdentifier> id)
         {
-            ActionResult? result = await _requestHelper.HasAccessToNamespace(User, Request, ns);
+            ActionResult? result = await _requestHelper.HasAccessToNamespace(User, Request, ns, new [] { AclAction.ReadObject });
             if (result != null)
             {
                 return result;
@@ -121,13 +119,12 @@ namespace Horde.Storage.Controllers
         }
 
         [HttpPost("{ns}/exist")]
-        [Authorize("Storage.read")]
         [ProducesDefaultResponseType]
         public async Task<IActionResult> ExistsBody(
             [Required] NamespaceId ns,
             [FromBody] BlobIdentifier[] bodyIds)
         {
-            ActionResult? result = await _requestHelper.HasAccessToNamespace(User, Request, ns);
+            ActionResult? result = await _requestHelper.HasAccessToNamespace(User, Request, ns, new [] { AclAction.ReadObject });
             if (result != null)
             {
                 return result;
@@ -148,13 +145,12 @@ namespace Horde.Storage.Controllers
         }
 
         [HttpPut("{ns}/{id}")]
-        [Authorize("Storage.write")]
         [RequiredContentType(CustomMediaTypeNames.UnrealCompactBinary)]
         public async Task<IActionResult> Put(
             [Required] NamespaceId ns,
             [Required] BlobIdentifier id)
         {
-            ActionResult? result = await _requestHelper.HasAccessToNamespace(User, Request, ns);
+            ActionResult? result = await _requestHelper.HasAccessToNamespace(User, Request, ns, new [] { AclAction.WriteObject });
             if (result != null)
             {
                 return result;
@@ -168,12 +164,11 @@ namespace Horde.Storage.Controllers
         }
 
         [HttpGet("{ns}/{id}/references")]
-        [Authorize("Storage.read")]
         public async Task<IActionResult> ResolveReferences(
             [Required] NamespaceId ns,
             [Required] BlobIdentifier id)
         {
-            ActionResult? result = await _requestHelper.HasAccessToNamespace(User, Request, ns);
+            ActionResult? result = await _requestHelper.HasAccessToNamespace(User, Request, ns, new [] { AclAction.ReadObject });
             if (result != null)
             {
                 return result;
@@ -221,12 +216,11 @@ namespace Horde.Storage.Controllers
         }
 
         [HttpDelete("{ns}/{id}")]
-        [Authorize("Storage.delete")]
         public async Task<IActionResult> Delete(
             [Required] NamespaceId ns,
             [Required] BlobIdentifier id)
         {
-            ActionResult? result = await _requestHelper.HasAccessToNamespace(User, Request, ns);
+            ActionResult? result = await _requestHelper.HasAccessToNamespace(User, Request, ns, new [] { AclAction.DeleteObject });
             if (result != null)
             {
                 return result;
@@ -241,11 +235,10 @@ namespace Horde.Storage.Controllers
         }
 
         [HttpDelete("{ns}")]
-        [Authorize("Admin")]
         public async Task<IActionResult> DeleteNamespace(
             [Required] NamespaceId ns)
         {
-            ActionResult? result = await _requestHelper.HasAccessToNamespace(User, Request, ns);
+            ActionResult? result = await _requestHelper.HasAccessToNamespace(User, Request, ns, new [] { AclAction.DeleteNamespace });
             if (result != null)
             {
                 return result;
diff --git a/Engine/Source/Programs/Horde/HordeStorage/Horde.Storage/Controllers/ReferencesController.cs b/Engine/Source/Programs/Horde/HordeStorage/Horde.Storage/Controllers/ReferencesController.cs
--- a/Engine/Source/Programs/Horde/HordeStorage/Horde.Storage/Controllers/ReferencesController.cs
+++ b/Engine/Source/Programs/Horde/HordeStorage/Horde.Storage/Controllers/ReferencesController.cs
@@ -20,6 +20,7 @@ using EpicGames.Horde.Storage;
 using EpicGames.Serialization;
 using Horde.Storage.Implementation;
 using JetBrains.Annotations;
+using Jupiter;
 using Jupiter.Common.Implementation;
 using Jupiter.Implementation;
 using Jupiter.Utils;
@@ -42,6 +43,7 @@ namespace Horde.Storage.Controllers
     [FormatFilter]
     [Produces(MediaTypeNames.Application.Json, MediaTypeNames.Application.Octet, CustomMediaTypeNames.UnrealCompactBinary)]
     [Route("api/v1/refs")]
+    [Authorize]
     public class ReferencesController : ControllerBase
     {
         private readonly IDiagnosticContext _diagnosticContext;
@@ -72,7 +74,6 @@ namespace Horde.Storage.Controllers
         [HttpGet("")]
         [ProducesDefaultResponseType]
         [ProducesResponseType(type: typeof(ProblemDetails), 400)]
-        [Authorize("Cache.read")]
         public async Task<IActionResult> GetNamespaces()
         {
             NamespaceId[] namespaces = await _objectService.GetNamespaces().ToArrayAsync();
@@ -81,7 +82,7 @@ namespace Horde.Storage.Controllers
             List<NamespaceId> namespacesWithAccess = new();
             foreach (NamespaceId ns in namespaces)
             {
-                ActionResult? accessResult = await _requestHelper.HasAccessToNamespace(User, Request, ns);
+                ActionResult? accessResult = await _requestHelper.HasAccessToNamespace(User, Request, ns, new [] { AclAction.ReadObject });
                 if (accessResult == null)
                 {
                     namespacesWithAccess.Add(ns);
@@ -100,14 +101,13 @@ namespace Horde.Storage.Controllers
         /// <param name="format">Optional specifier to set which output format is used json/raw/cb</param>
         [HttpGet("{ns}/{bucket}/{key}.{format?}", Order = 500)]
         [Produces(MediaTypeNames.Application.Json, MediaTypeNames.Application.Octet, CustomMediaTypeNames.UnrealCompactBinary, CustomMediaTypeNames.JupiterInlinedPayload)]
-        [Authorize("Object.read")]
         public async Task<IActionResult> Get(
             [FromRoute] [Required] NamespaceId ns,
             [FromRoute] [Required] BucketId bucket,
             [FromRoute] [Required] IoHashKey key,
             [FromRoute] string? format = null)
         {
-            ActionResult? accessResult = await _requestHelper.HasAccessToNamespace(User, Request, ns);
+            ActionResult? accessResult = await _requestHelper.HasAccessToNamespace(User, Request, ns, new [] { AclAction.ReadObject });
             if (accessResult != null)
             {
                 return accessResult;
@@ -338,38 +338,37 @@ namespace Horde.Storage.Controllers
     /// <param name="key">The unique name of this particular key. `iAmAVeryValidKey`</param>
     /// <param name="fields">The fields to include in the response, omit this to include everything.</param>
     [HttpGet("{ns}/{bucket}/{key}/metadata", Order = 500)]
-        [Authorize("Object.read")]
-        public async Task<IActionResult> GetMetadata(
-            [FromRoute] [Required] NamespaceId ns,
-            [FromRoute] [Required] BucketId bucket,
-            [FromRoute] [Required] IoHashKey key,
-            [FromQuery] string[] fields)
+    public async Task<IActionResult> GetMetadata(
+        [FromRoute] [Required] NamespaceId ns,
+        [FromRoute] [Required] BucketId bucket,
+        [FromRoute] [Required] IoHashKey key,
+        [FromQuery] string[] fields)
+    {
+        ActionResult? accessResult = await _requestHelper.HasAccessToNamespace(User, Request, ns, new [] { AclAction.ReadObject });
+        if (accessResult != null)
         {
-            ActionResult? accessResult = await _requestHelper.HasAccessToNamespace(User, Request, ns);
-            if (accessResult != null)
-            {
-                return accessResult;
-            }
+            return accessResult;
+        }
 
-            try
-            {
-                (ObjectRecord objectRecord, BlobContents? _) = await _objectService.Get(ns, bucket, key, fields);
+        try
+        {
+            (ObjectRecord objectRecord, BlobContents? _) = await _objectService.Get(ns, bucket, key, fields);
 
-                return Ok(new RefMetadataResponse(objectRecord));
-            }
-            catch (NamespaceNotFoundException e)
-            {
-                return NotFound(new ProblemDetails {Title = $"Namespace {e.Namespace} did not exist"});
-            }
-            catch (ObjectNotFoundException e)
-            {
-                return NotFound(new ProblemDetails { Title = $"Object {e.Bucket} {e.Key} did not exist" });
-            }
-            catch (BlobNotFoundException e)
-            {
-                return NotFound(new ProblemDetails { Title = $"Object {e.Blob} in {e.Ns} not found" });
-            }
+            return Ok(new RefMetadataResponse(objectRecord));
+        }
+        catch (NamespaceNotFoundException e)
+        {
+            return NotFound(new ProblemDetails {Title = $"Namespace {e.Namespace} did not exist"});
         }
+        catch (ObjectNotFoundException e)
+        {
+            return NotFound(new ProblemDetails { Title = $"Object {e.Bucket} {e.Key} did not exist" });
+        }
+        catch (BlobNotFoundException e)
+        {
+            return NotFound(new ProblemDetails { Title = $"Object {e.Blob} in {e.Ns} not found" });
+        }
+    }
 
         /// <summary>
         /// Checks if a object exists
@@ -381,13 +380,12 @@ namespace Horde.Storage.Controllers
         [HttpHead("{ns}/{bucket}/{key}", Order = 500)]
         [ProducesResponseType(type: typeof(RefResponse), 200)]
         [ProducesResponseType(type: typeof(ValidationProblemDetails), 400)]
-        [Authorize("Object.read")]
         public async Task<IActionResult> Head(
             [FromRoute] [Required] NamespaceId ns,
             [FromRoute] [Required] BucketId bucket,
             [FromRoute] [Required] IoHashKey key)
         {
-            ActionResult? accessResult = await _requestHelper.HasAccessToNamespace(User, Request, ns);
+            ActionResult? accessResult = await _requestHelper.HasAccessToNamespace(User, Request, ns, new [] { AclAction.ReadObject });
             if (accessResult != null)
             {
                 return accessResult;
@@ -446,13 +444,12 @@ namespace Horde.Storage.Controllers
         }
 
         [HttpGet("{ns}/exists")]
-        [Authorize("Object.read")]
         [ProducesDefaultResponseType]
         public async Task<IActionResult> ExistsMultiple(
             [FromRoute] [Required] NamespaceId ns,
             [FromQuery] [Required] List<string> names)
         {
-            ActionResult? accessResult = await _requestHelper.HasAccessToNamespace(User, Request, ns);
+            ActionResult? accessResult = await _requestHelper.HasAccessToNamespace(User, Request, ns, new [] { AclAction.ReadObject });
             if (accessResult != null)
             {
                 return accessResult;
@@ -512,13 +509,12 @@ namespace Horde.Storage.Controllers
 
         [HttpPut("{ns}/{bucket}/{key}.{format?}", Order = 500)]
         [DisableRequestSizeLimit]
-        [Authorize("Object.write")]
         public async Task<IActionResult> PutObject(
             [FromRoute] [Required] NamespaceId ns,
             [FromRoute] [Required] BucketId bucket,
             [FromRoute] [Required] IoHashKey key)
         {
-            ActionResult? accessResult = await _requestHelper.HasAccessToNamespace(User, Request, ns);
+            ActionResult? accessResult = await _requestHelper.HasAccessToNamespace(User, Request, ns, new [] { AclAction.WriteObject });
             if (accessResult != null)
             {
                 return accessResult;
@@ -607,14 +603,13 @@ namespace Horde.Storage.Controllers
         }
 
         [HttpPost("{ns}/{bucket}/{key}/finalize/{hash}.{format?}")]
-        [Authorize("Object.write")]
         public async Task<IActionResult> FinalizeObject(
             [FromRoute] [Required] NamespaceId ns,
             [FromRoute] [Required] BucketId bucket,
             [FromRoute] [Required] IoHashKey key,
             [FromRoute] [Required] BlobIdentifier hash)
         {
-            ActionResult? accessResult = await _requestHelper.HasAccessToNamespace(User, Request, ns);
+            ActionResult? accessResult = await _requestHelper.HasAccessToNamespace(User, Request, ns, new [] { AclAction.WriteObject });
             if (accessResult != null)
             {
                 return accessResult;
@@ -634,7 +629,24 @@ namespace Horde.Storage.Controllers
             [FromRoute] [Required] NamespaceId ns,
             [FromBody] [Required] BatchOps ops)
         {
-            ActionResult? accessResult = await _requestHelper.HasAccessToNamespace(User, Request, ns);
+            AclAction ActionForOp(BatchOps.BatchOp.Operation op)
+            {
+                switch (op)
+                {
+                    case BatchOps.BatchOp.Operation.GET:
+                        return AclAction.ReadObject;
+                    case BatchOps.BatchOp.Operation.PUT:
+                        return AclAction.WriteObject;
+                    case BatchOps.BatchOp.Operation.HEAD:
+                        return AclAction.ReadObject;
+                    default:
+                        throw new ArgumentOutOfRangeException(nameof(op), op, null);
+                }
+            }
+
+            AclAction[] requiredActions = ops.Ops.Select(op => ActionForOp(op.Op)).ToArray();
+
+            ActionResult? accessResult = await _requestHelper.HasAccessToNamespace(User, Request, ns, requiredActions);
             if (accessResult != null)
             {
                 return accessResult;
@@ -801,12 +813,11 @@ namespace Horde.Storage.Controllers
         /// <param name="ns">Namespace. Each namespace is completely separated from each other. Use for different types of data that is never expected to be similar (between two different games for instance)</param>
         [HttpDelete("{ns}", Order = 500)]
         [ProducesResponseType(204)]
-        [Authorize("admin")]
         public async Task<IActionResult> DeleteNamespace(
             [FromRoute] [Required] NamespaceId ns
         )
         {
-            ActionResult? accessResult = await _requestHelper.HasAccessToNamespace(User, Request, ns);
+            ActionResult? accessResult = await _requestHelper.HasAccessToNamespace(User, Request, ns, new [] { AclAction.DeleteNamespace });
             if (accessResult != null)
             {
                 return accessResult;
@@ -831,12 +842,11 @@ namespace Horde.Storage.Controllers
         /// <param name="bucket">The category/type of record you are caching. Is a clustered key together with the actual key, but all records in the same bucket can be dropped easily.</param>
         [HttpDelete("{ns}/{bucket}", Order = 500)]
         [ProducesResponseType(200)]
-        [Authorize("Object.delete")]
         public async Task<IActionResult> DeleteBucket(
             [FromRoute] [Required] NamespaceId ns,
             [FromRoute] [Required] BucketId bucket)
         {
-            ActionResult? accessResult = await _requestHelper.HasAccessToNamespace(User, Request, ns);
+            ActionResult? accessResult = await _requestHelper.HasAccessToNamespace(User, Request, ns, new [] { AclAction.DeleteBucket });
             if (accessResult != null)
             {
                 return accessResult;
@@ -864,13 +874,12 @@ namespace Horde.Storage.Controllers
         [HttpDelete("{ns}/{bucket}/{key}", Order = 500)]
         [ProducesResponseType(200)]
         [ProducesResponseType(404)]
-        [Authorize("Object.delete")]
         public async Task<IActionResult> Delete(
             [FromRoute] [Required] NamespaceId ns,
             [FromRoute] [Required] BucketId bucket,
             [FromRoute] [Required] IoHashKey key)
         {
-            ActionResult? accessResult = await _requestHelper.HasAccessToNamespace(User, Request, ns);
+            ActionResult? accessResult = await _requestHelper.HasAccessToNamespace(User, Request, ns, new [] { AclAction.DeleteObject });
             if (accessResult != null)
             {
                 return accessResult;
diff --git a/Engine/Source/Programs/Horde/HordeStorage/Horde.Storage/Controllers/ReplicationController.cs b/Engine/Source/Programs/Horde/HordeStorage/Horde.Storage/Controllers/ReplicationController.cs
--- a/Engine/Source/Programs/Horde/HordeStorage/Horde.Storage/Controllers/ReplicationController.cs
+++ b/Engine/Source/Programs/Horde/HordeStorage/Horde.Storage/Controllers/ReplicationController.cs
@@ -7,6 +7,7 @@ using System.Linq;
 using System.Threading.Tasks;
 using EpicGames.Horde.Storage;
 using Horde.Storage.Implementation;
+using Jupiter;
 using Jupiter.Implementation;
 using Microsoft.AspNetCore.Authorization;
 using Microsoft.AspNetCore.Mvc;
@@ -16,7 +17,7 @@ namespace Horde.Storage.Controllers
     [ApiController]
     [FormatFilter]
     [Route("api/v1/g")]
-    [Authorize("Admin")]
+    [Authorize]
     [InternalApiFilter]
     public class ReplicationController : ControllerBase
     {
@@ -35,7 +36,7 @@ namespace Horde.Storage.Controllers
             [Required] NamespaceId ns
         )
         {
-            ActionResult? result = await _requestHelper.HasAccessToNamespace(User, Request, ns);
+            ActionResult? result = await _requestHelper.HasAccessToNamespace(User, Request, ns, new [] { AclAction.ReadTransactionLog });
             if (result != null)
             {
                 return result;
@@ -78,7 +79,7 @@ namespace Horde.Storage.Controllers
             [Required] long offset
         )
         {
-            ActionResult? result = await _requestHelper.HasAccessToNamespace(User, Request, ns);
+            ActionResult? result = await _requestHelper.HasAccessToNamespace(User, Request, ns, new [] { AclAction.WriteTransactionLog });
             if (result != null)
             {
                 return result;
@@ -116,7 +117,7 @@ namespace Horde.Storage.Controllers
             [Required] [FromBody] NewReplicationState replicationState
         )
         {
-            ActionResult? result = await _requestHelper.HasAccessToNamespace(User, Request, ns);
+            ActionResult? result = await _requestHelper.HasAccessToNamespace(User, Request, ns, new [] { AclAction.WriteTransactionLog });
             if (result != null)
             {
                 return result;
@@ -159,7 +160,7 @@ namespace Horde.Storage.Controllers
             [Required] NamespaceId ns
         )
         {
-            ActionResult? result = await _requestHelper.HasAccessToNamespace(User, Request, ns);
+            ActionResult? result = await _requestHelper.HasAccessToNamespace(User, Request, ns, new [] { AclAction.WriteTransactionLog });
             if (result != null)
             {
                 return result;
diff --git a/Engine/Source/Programs/Horde/HordeStorage/Horde.Storage/Controllers/ReplicationLogController.cs b/Engine/Source/Programs/Horde/HordeStorage/Horde.Storage/Controllers/ReplicationLogController.cs
--- a/Engine/Source/Programs/Horde/HordeStorage/Horde.Storage/Controllers/ReplicationLogController.cs
+++ b/Engine/Source/Programs/Horde/HordeStorage/Horde.Storage/Controllers/ReplicationLogController.cs
@@ -10,6 +10,7 @@ using EpicGames.Horde.Storage;
 using EpicGames.Serialization;
 using Horde.Storage.Implementation;
 using Horde.Storage.Implementation.TransactionLog;
+using Jupiter;
 using Jupiter.Implementation;
 using Microsoft.AspNetCore.Authorization;
 using Microsoft.AspNetCore.Mvc;
@@ -22,6 +23,7 @@ namespace Horde.Storage.Controllers
     [ApiController]
     [Route("api/v1/replication-log")]
     [InternalApiFilter]
+    [Authorize]
     public class ReplicationLogController : ControllerBase
     {
         private readonly IServiceProvider _provider;
@@ -40,12 +42,11 @@ namespace Horde.Storage.Controllers
         [HttpGet("snapshots/{ns}")]
         [ProducesDefaultResponseType]
         [ProducesResponseType(type: typeof(ProblemDetails), 400)]
-        [Authorize("replication-log.read")]
         public async Task<IActionResult> GetSnapshots(
             [Required] NamespaceId ns
         )
         {
-            ActionResult? result = await _requestHelper.HasAccessToNamespace(User, Request, ns);
+            ActionResult? result = await _requestHelper.HasAccessToNamespace(User, Request, ns, new [] { AclAction.ReadTransactionLog });
             if (result != null)
             {
                 return result;
@@ -58,12 +59,11 @@ namespace Horde.Storage.Controllers
         [HttpPost("snapshots/{ns}/create")]
         [ProducesDefaultResponseType]
         [ProducesResponseType(type: typeof(ProblemDetails), 400)]
-        [Authorize("Admin")]
         public async Task<IActionResult> CreateSnapshot(
             [Required] NamespaceId ns
         )
         {
-            ActionResult? result = await _requestHelper.HasAccessToNamespace(User, Request, ns);
+            ActionResult? result = await _requestHelper.HasAccessToNamespace(User, Request, ns, new [] { AclAction.WriteTransactionLog });
             if (result != null)
             {
                 return result;
@@ -77,7 +77,6 @@ namespace Horde.Storage.Controllers
         [HttpGet("incremental/{ns}")]
         [ProducesDefaultResponseType]
         [ProducesResponseType(type: typeof(ProblemDetails), 400)]
-        [Authorize("replication-log.read")]
         public async Task<IActionResult> GetIncrementalEvents(
             [Required] NamespaceId ns,
             [FromQuery] string? lastBucket,
@@ -85,7 +84,7 @@ namespace Horde.Storage.Controllers
             [FromQuery] int count = 100
         )
         {
-            ActionResult? result = await _requestHelper.HasAccessToNamespace(User, Request, ns);
+            ActionResult? result = await _requestHelper.HasAccessToNamespace(User, Request, ns, new [] { AclAction.ReadTransactionLog });
             if (result != null)
             {
                 return result;
diff --git a/Engine/Source/Programs/Horde/HordeStorage/Horde.Storage/Controllers/RequestHelper.cs b/Engine/Source/Programs/Horde/HordeStorage/Horde.Storage/Controllers/RequestHelper.cs
--- a/Engine/Source/Programs/Horde/HordeStorage/Horde.Storage/Controllers/RequestHelper.cs
+++ b/Engine/Source/Programs/Horde/HordeStorage/Horde.Storage/Controllers/RequestHelper.cs
@@ -26,10 +26,14 @@ public class RequestHelper
         _settings = settings;
     }
 
-    public async Task<ActionResult?> HasAccessToNamespace(ClaimsPrincipal user, HttpRequest request, NamespaceId ns)
+    public async Task<ActionResult?> HasAccessToNamespace(ClaimsPrincipal user, HttpRequest request, NamespaceId ns, AclAction[] aclActions)
     {
         using IScope _ = Tracer.Instance.StartActive("authorize");
-        AuthorizationResult authorizationResult = await _authorizationService.AuthorizeAsync(user, ns, NamespaceAccessRequirement.Name);
+        AuthorizationResult authorizationResult = await _authorizationService.AuthorizeAsync(user, new NamespaceAccessRequest
+        {
+            Namespace = ns,
+            Actions = aclActions
+        }, NamespaceAccessRequirement.Name);
 
         if (!authorizationResult.Succeeded)
         {
@@ -59,4 +63,20 @@ public class RequestHelper
         // restricted namespace in corp or internal port, this is okay
         return null;
     }
+
+    public async Task<ActionResult?> HasAccessForGlobalOperations(ClaimsPrincipal user, AclAction[] aclActions)
+    {
+        using IScope _ = Tracer.Instance.StartActive("authorize");
+        AuthorizationResult authorizationResult = await _authorizationService.AuthorizeAsync(user, new GlobalAccessRequest
+        {
+            Actions = aclActions
+        }, GlobalAccessRequirement.Name);
+
+        if (!authorizationResult.Succeeded)
+        {
+            return new ForbidResult();
+        }
+
+        return null;
+    }
 }
diff --git a/Engine/Source/Programs/Horde/HordeStorage/Horde.Storage/Controllers/StatusController.cs b/Engine/Source/Programs/Horde/HordeStorage/Horde.Storage/Controllers/StatusController.cs
--- a/Engine/Source/Programs/Horde/HordeStorage/Horde.Storage/Controllers/StatusController.cs
+++ b/Engine/Source/Programs/Horde/HordeStorage/Horde.Storage/Controllers/StatusController.cs
@@ -15,6 +15,7 @@ namespace Horde.Storage.Controllers
 {
     [ApiController]
     [Route("api/v1/status")]
+    [Authorize]
     public class StatusController : Controller
     {
         private readonly VersionFile _versionFile;
@@ -38,7 +39,6 @@ namespace Horde.Storage.Controllers
         /// </remarks>
         /// <returns></returns>
         [HttpGet("")]
-        [Authorize("Any")]
         [ProducesResponseType(type: typeof(StatusResponse), 200)]
         public IActionResult Status()
         {
@@ -77,7 +77,6 @@ namespace Horde.Storage.Controllers
         /// </remarks>
         /// <returns></returns>
         [HttpGet("peers")]
-        [Authorize("Any")]
         [ProducesResponseType(type: typeof(PeersResponse), 200)]
         public IActionResult Peers([FromQuery] bool includeInternalEndpoints = false)
         {
diff --git a/Engine/Source/Programs/Horde/HordeStorage/Horde.Storage/Controllers/StorageController.cs b/Engine/Source/Programs/Horde/HordeStorage/Horde.Storage/Controllers/StorageController.cs
--- a/Engine/Source/Programs/Horde/HordeStorage/Horde.Storage/Controllers/StorageController.cs
+++ b/Engine/Source/Programs/Horde/HordeStorage/Horde.Storage/Controllers/StorageController.cs
@@ -26,31 +26,29 @@ namespace Horde.Storage.Controllers
     [ApiController]
     [Route("api/v1/s", Order = 1)]
     [Route("api/v1/blobs", Order = 0)]
+    [Authorize]
     public class StorageController : ControllerBase
     {
         private readonly IBlobService _storage;
         private readonly IDiagnosticContext _diagnosticContext;
-        private readonly IAuthorizationService _authorizationService;
         private readonly RequestHelper _requestHelper;
         private readonly BufferedPayloadFactory _bufferedPayloadFactory;
 
-        public StorageController(IBlobService storage, IDiagnosticContext diagnosticContext, IAuthorizationService authorizationService, RequestHelper requestHelper, BufferedPayloadFactory bufferedPayloadFactory)
+        public StorageController(IBlobService storage, IDiagnosticContext diagnosticContext, RequestHelper requestHelper, BufferedPayloadFactory bufferedPayloadFactory)
         {
             _storage = storage;
             _diagnosticContext = diagnosticContext;
-            _authorizationService = authorizationService;
             _requestHelper = requestHelper;
             _bufferedPayloadFactory = bufferedPayloadFactory;
         }
 
         [HttpGet("{ns}/{id}")]
-        [Authorize("Storage.read")]
         [ProducesDefaultResponseType]
         public async Task<IActionResult> Get(
             [Required] NamespaceId ns,
             [Required] BlobIdentifier id)
         {
-            ActionResult? result = await _requestHelper.HasAccessToNamespace(User, Request, ns);
+            ActionResult? result = await _requestHelper.HasAccessToNamespace(User, Request, ns, new [] { AclAction.ReadObject });
             if (result != null)
             {
                 return result;
@@ -69,13 +67,12 @@ namespace Horde.Storage.Controllers
         }
         
         [HttpHead("{ns}/{id}")]
-        [Authorize("Storage.read")]
         [ProducesDefaultResponseType]
         public async Task<IActionResult> Head(
             [Required] NamespaceId ns,
             [Required] BlobIdentifier id)
         {
-            ActionResult? result = await _requestHelper.HasAccessToNamespace(User, Request, ns);
+            ActionResult? result = await _requestHelper.HasAccessToNamespace(User, Request, ns, new [] { AclAction.ReadObject });
             if (result != null)
             {
                 return result;
@@ -91,13 +88,12 @@ namespace Horde.Storage.Controllers
         }
 
         [HttpPost("{ns}/exists")]
-        [Authorize("Storage.read")]
         [ProducesDefaultResponseType]
         public async Task<IActionResult> ExistsMultiple(
             [Required] NamespaceId ns,
              [Required] [FromQuery] List<BlobIdentifier> id)
         {
-            ActionResult? result = await _requestHelper.HasAccessToNamespace(User, Request, ns);
+            ActionResult? result = await _requestHelper.HasAccessToNamespace(User, Request, ns, new [] { AclAction.ReadObject });
             if (result != null)
             {
                 return result;
@@ -118,13 +114,12 @@ namespace Horde.Storage.Controllers
         }
 
         [HttpPost("{ns}/exist")]
-        [Authorize("Storage.read")]
         [ProducesDefaultResponseType]
         public async Task<IActionResult> ExistsBody(
             [Required] NamespaceId ns,
             [FromBody] BlobIdentifier[] bodyIds)
         {
-            ActionResult? result = await _requestHelper.HasAccessToNamespace(User, Request, ns);
+            ActionResult? result = await _requestHelper.HasAccessToNamespace(User, Request, ns, new [] { AclAction.ReadObject });
             if (result != null)
             {
                 return result;
@@ -162,14 +157,13 @@ namespace Horde.Storage.Controllers
         }
 
         [HttpPut("{ns}/{id}")]
-        [Authorize("Storage.write")]
         [RequiredContentType(MediaTypeNames.Application.Octet)]
         [DisableRequestSizeLimit]
         public async Task<IActionResult> Put(
             [Required] NamespaceId ns,
             [Required] BlobIdentifier id)
         {
-            ActionResult? result = await _requestHelper.HasAccessToNamespace(User, Request, ns);
+            ActionResult? result = await _requestHelper.HasAccessToNamespace(User, Request, ns, new [] { AclAction.WriteObject });
             if (result != null)
             {
                 return result;
@@ -193,13 +187,12 @@ namespace Horde.Storage.Controllers
         }
 
         [HttpPost("{ns}")]
-        [Authorize("Storage.write")]
         [RequiredContentType(MediaTypeNames.Application.Octet)]
         [DisableRequestSizeLimit]
         public async Task<IActionResult> Post(
             [Required] NamespaceId ns)
         {
-            ActionResult? result = await _requestHelper.HasAccessToNamespace(User, Request, ns);
+            ActionResult? result = await _requestHelper.HasAccessToNamespace(User, Request, ns, new [] { AclAction.WriteObject });
             if (result != null)
             {
                 return result;
@@ -220,12 +213,11 @@ namespace Horde.Storage.Controllers
         }
 
         [HttpDelete("{ns}/{id}")]
-        [Authorize("Storage.delete")]
         public async Task<IActionResult> Delete(
             [Required] NamespaceId ns,
             [Required] BlobIdentifier id)
         {
-            ActionResult? result = await _requestHelper.HasAccessToNamespace(User, Request, ns);
+            ActionResult? result = await _requestHelper.HasAccessToNamespace(User, Request, ns, new [] { AclAction.DeleteObject });
             if (result != null)
             {
                 return result;
@@ -238,11 +230,10 @@ namespace Horde.Storage.Controllers
 
         
         [HttpDelete("{ns}")]
-        [Authorize("Admin")]
         public async Task<IActionResult> DeleteNamespace(
             [Required] NamespaceId ns)
         {
-            ActionResult? result = await _requestHelper.HasAccessToNamespace(User, Request, ns);
+            ActionResult? result = await _requestHelper.HasAccessToNamespace(User, Request, ns, new [] { AclAction.DeleteNamespace });
             if (result != null)
             {
                 return result;
@@ -292,17 +283,17 @@ namespace Horde.Storage.Controllers
         [HttpPost("")]
         public async Task<IActionResult> Post([FromBody] BatchCall batch)
         {
-            string OpToPolicy(BatchOp.Operation op)
+            AclAction MapToAclAction(BatchOp.Operation op)
             {
                 switch (op)
                 {
                     case BatchOp.Operation.GET:
                     case BatchOp.Operation.HEAD:
-                        return "Storage.read";
+                        return AclAction.ReadObject;
                     case BatchOp.Operation.PUT:
-                        return "Storage.write";
+                        return AclAction.WriteObject;
                     case BatchOp.Operation.DELETE:
-                        return "Storage.delete";
+                        return AclAction.DeleteObject;
                     default:
                         throw new ArgumentOutOfRangeException(nameof(op), op, null);
                 }
@@ -322,12 +313,11 @@ namespace Horde.Storage.Controllers
                     throw new Exception(nameof(op.Namespace));
                 }
 
-                AuthorizationResult authorizationResultNamespace = await _authorizationService.AuthorizeAsync(User, op.Namespace, NamespaceAccessRequirement.Name);
-                AuthorizationResult authorizationResultOp = await _authorizationService.AuthorizeAsync(User, OpToPolicy(op.Op));
+                ActionResult? result = await _requestHelper.HasAccessToNamespace(User, Request, op.Namespace!.Value, new[] { MapToAclAction(op.Op) });
 
-                if (!authorizationResultNamespace.Succeeded || !authorizationResultOp.Succeeded)
+                if (result != null)
                 {
-                    return Forbid();
+                    return result;
                 }
 
                 switch (op.Op)
diff --git a/Engine/Source/Programs/Horde/HordeStorage/Horde.Storage/Horde.Storage.Startup.cs b/Engine/Source/Programs/Horde/HordeStorage/Horde.Storage/Horde.Storage.Startup.cs
--- a/Engine/Source/Programs/Horde/HordeStorage/Horde.Storage/Horde.Storage.Startup.cs
+++ b/Engine/Source/Programs/Horde/HordeStorage/Horde.Storage/Horde.Storage.Startup.cs
@@ -51,119 +51,7 @@ namespace Horde.Storage
 
         protected override void OnAddAuthorization(AuthorizationOptions authorizationOptions, List<string> defaultSchemes)
         {
-            authorizationOptions.AddPolicy("Cache.read", policy =>
-            {
-                policy.AuthenticationSchemes = defaultSchemes;
-                policy.BeginAnyContext()
-                    .IncludeClaim("Cache", "read", "readwrite", "full")
-                    .IncludeRole("Cache.read", "Cache.readwrite", "Cache.full")
-                    .AddAny();
-            });
-
-            authorizationOptions.AddPolicy("Cache.write", policy =>
-            {
-                policy.AuthenticationSchemes = defaultSchemes;
-                policy.BeginAnyContext()
-                    .IncludeClaim("Cache", "write", "readwrite", "full")
-                    .IncludeRole("Cache.write", "Cache.readwrite", "Cache.full")
-                    .AddAny();
-            });
-
-            authorizationOptions.AddPolicy("Cache.delete", policy =>
-            {
-                policy.AuthenticationSchemes = defaultSchemes;
-                policy.BeginAnyContext()
-                    .IncludeClaim("Cache", "delete", "full")
-                    .IncludeRole("Cache.delete", "Cache.full")
-                    .AddAny();
-            });
-
-            authorizationOptions.AddPolicy("Object.read", policy =>
-            {
-                policy.AuthenticationSchemes = defaultSchemes;
-                policy.BeginAnyContext()
-                    .IncludeClaim("Cache", "read", "readwrite", "full")
-                    .IncludeRole("Cache.read", "Cache.readwrite", "Cache.full")
-                    .AddAny();
-            });
-
-            authorizationOptions.AddPolicy("Object.write", policy =>
-            {
-                policy.AuthenticationSchemes = defaultSchemes;
-                policy.BeginAnyContext()
-                    .IncludeClaim("Cache", "write", "readwrite", "full")
-                    .IncludeRole("Cache.write", "Cache.readwrite", "Cache.full")
-                    .AddAny();
-            });
-
-            authorizationOptions.AddPolicy("Object.delete", policy =>
-            {
-                policy.AuthenticationSchemes = defaultSchemes;
-                policy.BeginAnyContext()
-                    .IncludeClaim("Cache", "delete", "full")
-                    .IncludeRole("Cache.delete", "Cache.full")
-                    .AddAny();
-            });
-
-            authorizationOptions.AddPolicy("Storage.read", policy =>
-            {
-                policy.AuthenticationSchemes = defaultSchemes;
-                policy.BeginAnyContext()
-                    .IncludeClaim("Storage", "read", "readwrite", "full")
-                    .IncludeRole("Storage.read", "Storage.readwrite", "Storage.full")
-                    .AddAny();
-            });
-
-            authorizationOptions.AddPolicy("Storage.write", policy =>
-            {
-                policy.AuthenticationSchemes = defaultSchemes;
-                policy.BeginAnyContext()
-                    .IncludeClaim("Storage", "write", "readwrite", "full")
-                    .IncludeRole("Storage.write", "Storage.readwrite", "Storage.full")
-                    .AddAny();
-            });
-
-            authorizationOptions.AddPolicy("Storage.delete", policy =>
-            {
-                policy.AuthenticationSchemes = defaultSchemes;
-                policy.BeginAnyContext()
-                    .IncludeClaim("Storage", "delete", "full")
-                    .IncludeRole("Storage.delete", "Storage.full")
-                    .AddAny();
-            });
-
-            authorizationOptions.AddPolicy("replication-log.read", policy =>
-            {
-                policy.AuthenticationSchemes = defaultSchemes;
-                policy.BeginAnyContext()
-                    .IncludeClaim("transactionLog", "read", "readwrite", "full")
-                    .IncludeRole("transactionLog.read", "transactionLog.readwrite", "transactionLog.full")
-                    .AddAny();
-            });
-
-            authorizationOptions.AddPolicy("replication-log.write", policy =>
-            {
-                policy.AuthenticationSchemes = defaultSchemes;
-                policy.BeginAnyContext()
-                    .IncludeClaim("transactionLog", "write", "readwrite", "full")
-                    .IncludeRole("transactionLog.write", "transactionLog.readwrite", "transactionLog.full")
-                    .AddAny();
-            });
-
-            authorizationOptions.AddPolicy("Admin", policy =>
-            {
-                policy.AuthenticationSchemes = defaultSchemes;
-                policy.BeginAnyContext()
-                    .IncludeClaim("Admin")
-                    .IncludeRole("Admin")
-                    .AddAny();
-            });
-
-            // A policy that grants any authenticated user access
-            authorizationOptions.AddPolicy("Any", policy =>
-            {
-                policy.RequireAuthenticatedUser();
-            });
+  
         }
 
         protected override void OnAddService(IServiceCollection services)
diff --git a/Engine/Source/Programs/Horde/HordeStorage/Jupiter.Common/Authentication/Acl.cs b/Engine/Source/Programs/Horde/HordeStorage/Jupiter.Common/Authentication/Acl.cs
new file mode 100644
--- /dev/null
+++ b/Engine/Source/Programs/Horde/HordeStorage/Jupiter.Common/Authentication/Acl.cs
@@ -0,0 +1,101 @@
+// Copyright Epic Games, Inc. All Rights Reserved.
+
+using System;
+using System.Collections.Generic;
+using System.Security.Claims;
+using Microsoft.AspNetCore.Authorization;
+
+namespace Jupiter;
+
+public enum AclAction
+{
+    /// <summary>
+    /// General read access to refs / blobs and so on
+    /// </summary>
+    ReadObject,
+    /// <summary>
+    /// General write access to upload refs / blobs etc
+    /// </summary>
+    WriteObject,
+    /// <summary>
+    /// Access to delete blobs / refs etc
+    /// </summary>
+    DeleteObject,
+
+    /// <summary>
+    /// Access to delete a particular bucket
+    /// </summary>
+    DeleteBucket,
+    /// <summary>
+    /// Access to delete a whole namespace
+    /// </summary>
+    DeleteNamespace,
+
+    /// <summary>
+    /// Access to read the transaction log
+    /// </summary>
+    ReadTransactionLog,
+
+    /// <summary>
+    /// Access to write the transaction log
+    /// </summary>
+    WriteTransactionLog,
+
+    /// <summary>
+    /// Access to perform administrative task
+    /// </summary>
+    AdminAction
+}
+
+public class AclEntry
+{
+    /// <summary>
+    /// Claims required to be present to be allowed to do the actions - if multiple claims are present *all* of them are required
+    /// </summary>
+    [System.Diagnostics.CodeAnalysis.SuppressMessage("Usage", "CA2227:Collection properties should be read only", Justification = "Used by the configuration system")]
+    public List<string> Claims { get; set; } = new List<string>();
+
+    /// <summary>
+    /// The actions granted if the claims match
+    /// </summary>
+    [System.Diagnostics.CodeAnalysis.SuppressMessage("Usage", "CA2227:Collection properties should be read only", Justification = "Used by the configuration system")]
+    public List<AclAction> Actions { get; set; } = new List<AclAction>();
+
+    public IEnumerable<AclAction> Resolve(ClaimsPrincipal user)
+    {
+        List<AclAction> allowedActions = new List<AclAction>();
+        // These are ANDed, e.g. all claims needs to be present
+        foreach (string expectedClaim in Claims)
+        {
+            // if expected claim is * then everyone is allowed to use the namespace
+            if (expectedClaim == "*")
+            {
+                allowedActions.AddRange(Actions);
+                continue;
+            }
+
+            if (expectedClaim.Contains('=', StringComparison.InvariantCultureIgnoreCase))
+            {
+                int separatorIndex = expectedClaim.IndexOf('=', StringComparison.InvariantCultureIgnoreCase);
+                string claimName = expectedClaim.Substring(0, separatorIndex);
+                string claimValue = expectedClaim.Substring(separatorIndex + 1);
+                if (user.HasClaim(claim => claim.Type == claimName && claim.Value == claimValue))
+                {
+                    allowedActions.AddRange(Actions);
+                    continue;
+                }
+            }
+            if (user.HasClaim(claim => claim.Type == expectedClaim))
+            {
+                allowedActions.AddRange(Actions);
+            }
+        }
+
+        return allowedActions;
+    }
+
+    public IEnumerable<AclAction> Resolve(AuthorizationHandlerContext context)
+    {
+        return Resolve(context.User);
+    }
+}
\ No newline at end of file
diff --git a/Engine/Source/Programs/Horde/HordeStorage/Jupiter.Common/Authentication/GlobalAuthorizationHandler.cs b/Engine/Source/Programs/Horde/HordeStorage/Jupiter.Common/Authentication/GlobalAuthorizationHandler.cs
new file mode 100644
--- /dev/null
+++ b/Engine/Source/Programs/Horde/HordeStorage/Jupiter.Common/Authentication/GlobalAuthorizationHandler.cs
@@ -0,0 +1,69 @@
+// Copyright Epic Games, Inc. All Rights Reserved.
+
+using System;
+using System.Collections.Generic;
+using System.Linq;
+using System.Threading.Tasks;
+using Microsoft.AspNetCore.Authorization;
+using Microsoft.Extensions.Options;
+
+namespace Jupiter
+{
+    public class GlobalAccessRequest
+    {
+        public AclAction[] Actions { get; init; } = Array.Empty<AclAction>();
+    }
+
+    /// <summary>
+    /// Verifies that you have access to acl actions that are not tied to a namespace
+    /// </summary>
+    public class GlobalAuthorizationHandler : AuthorizationHandler<GlobalAccessRequirement, GlobalAccessRequest>
+    {
+        private readonly IOptionsMonitor<AuthSettings> _authSettings;
+
+        public GlobalAuthorizationHandler(IOptionsMonitor<AuthSettings> authSettings)
+        {
+            _authSettings = authSettings;
+        }
+
+        protected override Task HandleRequirementAsync(AuthorizationHandlerContext context, GlobalAccessRequirement requirement,
+            GlobalAccessRequest accessRequest)
+        {
+            if (!accessRequest.Actions.Any())
+            {
+                throw new Exception("At least 1 AclAction has to be specified for the namespace access request");
+            }
+
+            if (!_authSettings.CurrentValue.Enabled)
+            {
+                context.Succeed(requirement);
+            }
+
+            List<AclAction> allowedActions = new List<AclAction>();
+            foreach (AclEntry acl in _authSettings.CurrentValue.Acls)
+            {
+                allowedActions.AddRange(acl.Resolve(context));
+            }
+
+            bool haveAccessToActions = true;
+            foreach (AclAction requiredAction in accessRequest.Actions)
+            {
+                if (!allowedActions.Contains(requiredAction))
+                {
+                    haveAccessToActions = false;
+                }
+            }
+            if (haveAccessToActions)
+            {
+                context.Succeed(requirement);
+            }
+
+            return Task.CompletedTask;
+        }
+    }
+
+    public class GlobalAccessRequirement : IAuthorizationRequirement
+    {
+        public const string Name = "GlobalAccess";
+    }
+}
diff --git a/Engine/Source/Programs/Horde/HordeStorage/Jupiter.Common/Authentication/NamespaceAuthorizationHandler.cs b/Engine/Source/Programs/Horde/HordeStorage/Jupiter.Common/Authentication/NamespaceAuthorizationHandler.cs
--- a/Engine/Source/Programs/Horde/HordeStorage/Jupiter.Common/Authentication/NamespaceAuthorizationHandler.cs
+++ b/Engine/Source/Programs/Horde/HordeStorage/Jupiter.Common/Authentication/NamespaceAuthorizationHandler.cs
@@ -1,27 +1,39 @@
 ﻿// Copyright Epic Games, Inc. All Rights Reserved.
 
 using System;
+using System.Collections.Generic;
+using System.Linq;
 using System.Threading.Tasks;
 using EpicGames.Horde.Storage;
 using Jupiter.Common;
 using Jupiter.Implementation;
 using Microsoft.AspNetCore.Authorization;
+using Microsoft.Extensions.Options;
 
 namespace Jupiter
 {
+    public class NamespaceAccessRequest
+    {
+        public NamespaceId Namespace { get; init; }
+        public AclAction[] Actions { get; init; } = Array.Empty<AclAction>();
+    }
+
     // verifies that you have access to a namespace by checking if you have a corresponding claim to that namespace
-    public class NamespaceAuthorizationHandler : AuthorizationHandler<NamespaceAccessRequirement, NamespaceId>
+    public class NamespaceAuthorizationHandler : AuthorizationHandler<NamespaceAccessRequirement, NamespaceAccessRequest>
     {
         private readonly INamespacePolicyResolver _namespacePolicyResolver;
+        private readonly IOptionsMonitor<AuthSettings> _authSettings;
 
-        public NamespaceAuthorizationHandler(INamespacePolicyResolver namespacePolicyResolver)
+        public NamespaceAuthorizationHandler(INamespacePolicyResolver namespacePolicyResolver, IOptionsMonitor<AuthSettings> authSettings)
         {
             _namespacePolicyResolver = namespacePolicyResolver;
+            _authSettings = authSettings;
         }
 
         protected override Task HandleRequirementAsync(AuthorizationHandlerContext context, NamespaceAccessRequirement requirement,
-            NamespaceId namespaceName)
+            NamespaceAccessRequest accessRequest)
         {
+            NamespaceId namespaceName = accessRequest.Namespace;
             if (context.User.HasClaim(claim => claim.Type == "AllNamespaces"))
             {
                 context.Succeed(requirement);
@@ -30,34 +42,37 @@ namespace Jupiter
 
             try
             {
+                if (!accessRequest.Actions.Any())
+                {
+                    throw new Exception("At least 1 AclAction has to be specified for the namespace access request");
+                } 
+                
                 NamespacePolicy policy = _namespacePolicyResolver.GetPoliciesForNs(namespaceName);
 
-                // These are ANDed, e.g. all claims needs to be present
-                foreach (string expectedClaim in policy.Claims)
+                List<AclAction> allowedActions = new List<AclAction>();
+                foreach (AclEntry acl in policy.Acls)
                 {
-                    // if expected claim is * then everyone is allowed to use the namespace
-                    if (expectedClaim == "*")
-                    {
-                        context.Succeed(requirement);
-                        continue;
-                    }
+                    allowedActions.AddRange(acl.Resolve(context));
+                }
 
-                    if (expectedClaim.Contains('=', StringComparison.InvariantCultureIgnoreCase))
-                    {
-                        int separatorIndex = expectedClaim.IndexOf('=', StringComparison.InvariantCultureIgnoreCase);
-                        string claimName = expectedClaim.Substring(0, separatorIndex);
-                        string claimValue = expectedClaim.Substring(separatorIndex + 1);
-                        if (context.User.HasClaim(claim => claim.Type == claimName && claim.Value == claimValue))
-                        {
-                            context.Succeed(requirement);
-                            continue;
-                        }
-                    }
-                    if (context.User.HasClaim(claim => claim.Type == expectedClaim))
+                // the root and namespace acls are combined, namespace acls can not override what we define in the root
+                foreach (AclEntry acl in _authSettings.CurrentValue.Acls)
+                {
+                    allowedActions.AddRange(acl.Resolve(context));
+                }
+
+                bool haveAccessToActions = true;
+                foreach (AclAction requiredAction in accessRequest.Actions)
+                {
+                    if (!allowedActions.Contains(requiredAction))
                     {
-                        context.Succeed(requirement);
+                        haveAccessToActions = false;
                     }
                 }
+                if (haveAccessToActions)
+                {
+                    context.Succeed(requirement);
+                }
             }
             catch (UnknownNamespaceException)
             {
diff --git a/Engine/Source/Programs/Horde/HordeStorage/Jupiter.Common/BaseStartup.cs b/Engine/Source/Programs/Horde/HordeStorage/Jupiter.Common/BaseStartup.cs
--- a/Engine/Source/Programs/Horde/HordeStorage/Jupiter.Common/BaseStartup.cs
+++ b/Engine/Source/Programs/Horde/HordeStorage/Jupiter.Common/BaseStartup.cs
@@ -3,7 +3,9 @@
 using System;
 using System.Collections.Generic;
 using System.ComponentModel.DataAnnotations;
+using System.IdentityModel.Tokens.Jwt;
 using System.IO;
+using System.Linq;
 using System.Net.Mime;
 using System.Reflection;
 using System.Text;
@@ -27,6 +29,7 @@ using Microsoft.Extensions.Diagnostics.HealthChecks;
 using Microsoft.Extensions.Hosting;
 using Microsoft.Extensions.Options;
 using Microsoft.Extensions.Primitives;
+using Microsoft.Net.Http.Headers;
 using Microsoft.OpenApi.Models;
 using Newtonsoft.Json;
 using Newtonsoft.Json.Serialization;
@@ -108,14 +111,26 @@ namespace Jupiter
 
             services.AddHttpContextAccessor();
 
+            const string ForwardingScheme = "ForwardingScheme";
             List<string> availableSchemes = new List<string>();
 
             AuthenticationBuilder authenticationBuilder = services.AddAuthentication(options =>
                 {
                     if (Auth.Enabled)
                     {
-                        options.DefaultAuthenticateScheme = Auth.DefaultScheme;
-                        options.DefaultChallengeScheme = Auth.DefaultScheme;
+                        if (Auth.Schemes.Count > 1)
+                        {
+                            // we have multiple schemes, so we set the default to the forwarding scheme which will use the jwtAuthority to pick the correct scheme for the token
+                            options.DefaultAuthenticateScheme = ForwardingScheme;
+                            options.DefaultChallengeScheme = ForwardingScheme;
+                        }
+                        else
+                        {
+                            // if we only have one scheme we set it to default
+                            options.DefaultAuthenticateScheme = Auth.DefaultScheme;
+                            options.DefaultChallengeScheme = Auth.DefaultScheme;
+
+                        }
                     }
                     else
                     {
@@ -155,6 +170,42 @@ namespace Jupiter
                             throw new NotSupportedException($"Unknown implementation type {scheme.Implementation}");
                     }
                 }
+
+                authenticationBuilder.AddPolicyScheme(ForwardingScheme, ForwardingScheme, options =>
+                {
+                    options.ForwardDefaultSelector = context =>
+                    {
+                        string authorization = context.Request.Headers[HeaderNames.Authorization];
+                        string name = "Bearer";
+                        string tokenName = $"{name} ";
+                        if (string.IsNullOrEmpty(authorization) ||
+                            !authorization.StartsWith(tokenName, StringComparison.InvariantCulture))
+                        {
+                            return Auth.DefaultScheme;
+                        }
+
+                        string token = authorization.Substring(tokenName.Length).Trim();
+                        JwtSecurityTokenHandler jwtHandler = new JwtSecurityTokenHandler();
+
+                        if (!jwtHandler.CanReadToken(token))
+                        {
+                            return Auth.DefaultScheme;
+                        }
+
+                        JwtSecurityToken jwtToken = jwtHandler.ReadJwtToken(token);
+
+                        foreach (KeyValuePair<string, AuthSchemeEntry> entry in Auth.Schemes)
+                        {
+                            if (entry.Value.JwtAuthority == jwtToken.Issuer)
+                            {
+                                return entry.Key;
+                            }
+                        }
+
+                        return Auth.DefaultScheme;
+
+                    };
+                });
             }
             else
             {
@@ -172,10 +223,24 @@ namespace Jupiter
                     policy.AuthenticationSchemes = availableSchemes;
                     policy.Requirements.Add(new NamespaceAccessRequirement());
                 });
+                
+                options.AddPolicy(GlobalAccessRequirement.Name, policy =>
+                {
+                    policy.AuthenticationSchemes = availableSchemes;
+                    policy.Requirements.Add(new GlobalAccessRequirement());
+                });
+
+                // A policy that grants any authenticated user access
+                options.AddPolicy("Any", policy =>
+                {
+                    policy.AuthenticationSchemes = availableSchemes;
+                    policy.RequireAuthenticatedUser();
+                });
 
                 OnAddAuthorization(options, availableSchemes);
             });
             services.AddSingleton<IAuthorizationHandler, NamespaceAuthorizationHandler>();
+            services.AddSingleton<IAuthorizationHandler, GlobalAuthorizationHandler>();
 
             services.Configure<ForwardedHeadersOptions>(options =>
             {
@@ -472,6 +537,10 @@ namespace Jupiter
                 {
                     validationResults.Add(new ValidationResult("JWT Audience must be specified when using Okta implementation"));
                 }
+                if (string.IsNullOrEmpty(JwtAuthority))
+                {
+                    validationResults.Add(new ValidationResult("JWT Authority must be specified when using Okta implementation"));
+                }
             }
             else
             {
@@ -497,6 +566,9 @@ namespace Jupiter
         [System.Diagnostics.CodeAnalysis.SuppressMessage("Usage", "CA2227:Collection properties should be read only", Justification = "Used by the configuration system")]
         public Dictionary<string, AuthSchemeEntry> Schemes { get; set; } = new Dictionary<string, AuthSchemeEntry>();
 
+        [System.Diagnostics.CodeAnalysis.SuppressMessage("Usage", "CA2227:Collection properties should be read only", Justification = "Used by the configuration system")]
+        public List<AclEntry> Acls { get; set; } = new List<AclEntry>();
+
         public IEnumerable<ValidationResult> Validate(ValidationContext validationContext)
         {
             List<ValidationResult> validationResults = new List<ValidationResult>();
@@ -527,7 +599,7 @@ namespace Jupiter
         public long MemoryBufferSize { get; set; } = int.MaxValue;
 
         // enable to unhide potentially personal information, see https://aka.ms/IdentityModel/PII
-        public bool ShowPII { get; set; } = false;
+        public bool ShowPII { get; set; } = true;
         public bool DisableHealthChecks { get; set; } = false;
         public bool HostSwaggerDocumentation { get; set; } = true;
 
@@ -572,7 +644,8 @@ namespace Jupiter
 
     public class NamespacePolicy
     {
-        public string[] Claims { get; set; } = Array.Empty<string>();
+        [System.Diagnostics.CodeAnalysis.SuppressMessage("Usage", "CA2227:Collection properties should be read only", Justification = "Used by the configuration system")]
+        public List<AclEntry> Acls { get; set; } = new List<AclEntry>();
         public string StoragePool { get; set; } = "";
 
         public bool LastAccessTracking { get; set; } = true;
diff --git a/Engine/Source/Programs/Horde/HordeStorage/Jupiter.Common/Implementation/NamespacePolicyResolver.cs b/Engine/Source/Programs/Horde/HordeStorage/Jupiter.Common/Implementation/NamespacePolicyResolver.cs
--- a/Engine/Source/Programs/Horde/HordeStorage/Jupiter.Common/Implementation/NamespacePolicyResolver.cs
+++ b/Engine/Source/Programs/Horde/HordeStorage/Jupiter.Common/Implementation/NamespacePolicyResolver.cs
@@ -1,5 +1,6 @@
 // Copyright Epic Games, Inc. All Rights Reserved.
 
+using System.Collections.Generic;
 using EpicGames.Horde.Storage;
 using Jupiter.Implementation;
 using Microsoft.Extensions.Options;
@@ -23,16 +24,6 @@ namespace Jupiter.Common
 
         public NamespacePolicy GetPoliciesForNs(NamespaceId ns)
         {
-            if (ns == INamespacePolicyResolver.JupiterInternalNamespace)
-            {
-                return new NamespacePolicy()
-                {
-                    // we expect the full storage access claim access to the internal namespace, this is typically set of admins and service accounts that can access everything
-                    Claims = new string[] {"Storage=full" },
-                    StoragePool = ""
-                };
-            }
-
             if (_namespaceSettings.CurrentValue.Policies.TryGetValue(ns.ToString(),
                     out NamespacePolicy? settings))
             {
