Horde: Make ComputeController instanceable again

Internal access modifier on constructor prevented class from getting instantiated by DI.

diff --git a/Engine/Source/Programs/Horde/Horde.Build/Compute/ComputeController.cs b/Engine/Source/Programs/Horde/Horde.Build/Compute/ComputeController.cs
--- a/Engine/Source/Programs/Horde/Horde.Build/Compute/ComputeController.cs
+++ b/Engine/Source/Programs/Horde/Horde.Build/Compute/ComputeController.cs
@@ -28,7 +28,7 @@ namespace Horde.Build.Compute
 		/// </summary>
 		/// <param name="aclService"></param>
 		/// <param name="computeService">The compute service singleton</param>
-		internal ComputeController(AclService aclService, ComputeService computeService)
+		public ComputeController(AclService aclService, ComputeService computeService)
 		{
 			_aclService = aclService;
 			_computeService = computeService;
diff --git a/Engine/Source/Programs/Horde/Horde.Build/Compute/ComputeService.cs b/Engine/Source/Programs/Horde/Horde.Build/Compute/ComputeService.cs
--- a/Engine/Source/Programs/Horde/Horde.Build/Compute/ComputeService.cs
+++ b/Engine/Source/Programs/Horde/Horde.Build/Compute/ComputeService.cs
@@ -62,7 +62,7 @@ namespace Horde.Build.Compute
 	/// <summary>
 	/// Dispatches remote actions. Does not implement any cross-pod communication to satisfy leases; only agents connected to this server instance will be stored.
 	/// </summary>
-	class ComputeService : TaskSourceBase<ComputeTaskMessage>, IHostedService, IDisposable
+	public class ComputeService : TaskSourceBase<ComputeTaskMessage>, IHostedService, IDisposable
 	{
 		[RedisConverter(typeof(QueueKeySerializer))]
 		class QueueKey
@@ -113,6 +113,9 @@ namespace Horde.Build.Compute
 		/// <inheritdoc/>
 		public override TaskSourceFlags Flags => TaskSourceFlags.None;
 
+		/// <summary>
+		/// ID of the default namespace
+		/// </summary>
 		public static NamespaceId DefaultNamespaceId { get; } = new NamespaceId("default");
 
 		readonly IStorageClient _storageClient;
@@ -343,6 +346,7 @@ namespace Horde.Build.Compute
 			return await _messageQueue.WaitForMessagesAsync(GetMessageQueueId(clusterId, channelId), cancellationToken);
 		}
 
+		/// <inheritdoc/>
 		public override async Task OnLeaseStartedAsync(IAgent agent, LeaseId leaseId, ComputeTaskMessage computeTask, ILogger logger)
 		{
 			await base.OnLeaseStartedAsync(agent, leaseId, computeTask, logger);
@@ -353,6 +357,7 @@ namespace Horde.Build.Compute
 			await PostStatusMessageAsync(computeTask, status);
 		}
 
+		/// <inheritdoc/>
 		public override async Task OnLeaseFinishedAsync(IAgent agent, LeaseId leaseId, ComputeTaskMessage computeTask, LeaseOutcome outcome, ReadOnlyMemory<byte> output, ILogger logger)
 		{
 			await base.OnLeaseFinishedAsync(agent, leaseId, computeTask, outcome, output, logger);
