Horde Storage - Fixed bug when requiring more then one claim in a ACL where only having one of the claims was suffiecent to get access to the ACL

diff --git a/Engine/Source/Programs/Horde/HordeStorage/Jupiter.Common/Authentication/Acl.cs b/Engine/Source/Programs/Horde/HordeStorage/Jupiter.Common/Authentication/Acl.cs
--- a/Engine/Source/Programs/Horde/HordeStorage/Jupiter.Common/Authentication/Acl.cs
+++ b/Engine/Source/Programs/Horde/HordeStorage/Jupiter.Common/Authentication/Acl.cs
@@ -64,33 +64,42 @@ public class AclEntry
     public IEnumerable<AclAction> Resolve(ClaimsPrincipal user)
     {
         List<AclAction> allowedActions = new List<AclAction>();
+
+        bool allClaimsFound = true;
         // These are ANDed, e.g. all claims needs to be present
         foreach (string expectedClaim in Claims)
         {
-            // if expected claim is * then everyone is allowed to use the namespace
+            bool claimFound = false;
+            // if expected claim is * then everyone has the associated actions
             if (expectedClaim == "*")
             {
-                allowedActions.AddRange(Actions);
-                continue;
+                claimFound = true;
             }
-
-            if (expectedClaim.Contains('=', StringComparison.InvariantCultureIgnoreCase))
+            else if (expectedClaim.Contains('=', StringComparison.InvariantCultureIgnoreCase))
             {
                 int separatorIndex = expectedClaim.IndexOf('=', StringComparison.InvariantCultureIgnoreCase);
                 string claimName = expectedClaim.Substring(0, separatorIndex);
                 string claimValue = expectedClaim.Substring(separatorIndex + 1);
                 if (user.HasClaim(claim => claim.Type == claimName && claim.Value == claimValue))
                 {
-                    allowedActions.AddRange(Actions);
-                    continue;
+                    claimFound = true;
                 }
             }
-            if (user.HasClaim(claim => claim.Type == expectedClaim))
+            else if (user.HasClaim(claim => claim.Type == expectedClaim))
             {
-                allowedActions.AddRange(Actions);
+                claimFound = true;
+            }
+
+            if (!claimFound)
+            {
+                allClaimsFound = false;
             }
         }
 
+        if (allClaimsFound)
+        {
+            allowedActions.AddRange(Actions);
+        }
         return allowedActions;
     }
 
