Horde Storage - Deleted raw ddc endpoints which are not in use anymore

diff --git a/Engine/Source/Programs/Horde/HordeStorage/Horde.Storage/Endpoints/DDCEndpoints.cs b/Engine/Source/Programs/Horde/HordeStorage/Horde.Storage/Endpoints/DDCEndpoints.cs
deleted file mode 100644
diff --git a/Engine/Source/Programs/Horde/HordeStorage/Horde.Storage/Horde.Storage.Startup.cs b/Engine/Source/Programs/Horde/HordeStorage/Horde.Storage/Horde.Storage.Startup.cs
--- a/Engine/Source/Programs/Horde/HordeStorage/Horde.Storage/Horde.Storage.Startup.cs
+++ b/Engine/Source/Programs/Horde/HordeStorage/Horde.Storage/Horde.Storage.Startup.cs
@@ -22,7 +22,6 @@ using Jupiter;
 using Jupiter.Common.Implementation;
 using Jupiter.Utils;
 using Microsoft.AspNetCore.Authorization;
-using Microsoft.AspNetCore.Routing;
 using Microsoft.Extensions.Configuration;
 using Microsoft.Extensions.DependencyInjection;
 using Microsoft.Extensions.Options;
@@ -103,7 +102,6 @@ namespace Horde.Storage
             services.AddHostedService<BlobCleanupService>(p => p.GetService<BlobCleanupService>()!);
 
             services.AddSingleton(serviceType: typeof(IDDCRefService), typeof(DDCRefService));
-            services.AddSingleton(serviceType: typeof(DDCEndpoints));
 
             services.AddSingleton(serviceType: typeof(IObjectService), typeof(ObjectService));
             services.AddSingleton(serviceType: typeof(IReferencesStore), ObjectStoreFactory);
