Horde Storage - Removed support for the "AllNamespaces" special claim - this can be achived by using the acls in the auth setting instead

diff --git a/Engine/Source/Programs/Horde/HordeStorage/Jupiter.Common/Authentication/DisabledAuthenticationHandler.cs b/Engine/Source/Programs/Horde/HordeStorage/Jupiter.Common/Authentication/DisabledAuthenticationHandler.cs
--- a/Engine/Source/Programs/Horde/HordeStorage/Jupiter.Common/Authentication/DisabledAuthenticationHandler.cs
+++ b/Engine/Source/Programs/Horde/HordeStorage/Jupiter.Common/Authentication/DisabledAuthenticationHandler.cs
@@ -54,7 +54,6 @@ namespace Jupiter
             new Claim("TransactionLog", "full"),
             new Claim("Storage", "full"),
             new Claim("Admin", ""),
-            new Claim("AllNamespaces", "true")
         }, "automatic");
     }
 }
diff --git a/Engine/Source/Programs/Horde/HordeStorage/Jupiter.Common/Authentication/GlobalAuthorizationHandler.cs b/Engine/Source/Programs/Horde/HordeStorage/Jupiter.Common/Authentication/GlobalAuthorizationHandler.cs
--- a/Engine/Source/Programs/Horde/HordeStorage/Jupiter.Common/Authentication/GlobalAuthorizationHandler.cs
+++ b/Engine/Source/Programs/Horde/HordeStorage/Jupiter.Common/Authentication/GlobalAuthorizationHandler.cs
@@ -37,6 +37,7 @@ namespace Jupiter
             if (!_authSettings.CurrentValue.Enabled)
             {
                 context.Succeed(requirement);
+                return Task.CompletedTask;
             }
 
             List<AclAction> allowedActions = new List<AclAction>();
diff --git a/Engine/Source/Programs/Horde/HordeStorage/Jupiter.Common/Authentication/NamespaceAuthorizationHandler.cs b/Engine/Source/Programs/Horde/HordeStorage/Jupiter.Common/Authentication/NamespaceAuthorizationHandler.cs
--- a/Engine/Source/Programs/Horde/HordeStorage/Jupiter.Common/Authentication/NamespaceAuthorizationHandler.cs
+++ b/Engine/Source/Programs/Horde/HordeStorage/Jupiter.Common/Authentication/NamespaceAuthorizationHandler.cs
@@ -34,7 +34,7 @@ namespace Jupiter
             NamespaceAccessRequest accessRequest)
         {
             NamespaceId namespaceName = accessRequest.Namespace;
-            if (context.User.HasClaim(claim => claim.Type == "AllNamespaces"))
+            if (!_authSettings.CurrentValue.Enabled)
             {
                 context.Succeed(requirement);
                 return Task.CompletedTask;
