Horde Storage - Updated documentation around auth setup

diff --git a/Engine/Source/Programs/Horde/HordeStorage/README.md b/Engine/Source/Programs/Horde/HordeStorage/README.md
--- a/Engine/Source/Programs/Horde/HordeStorage/README.md
+++ b/Engine/Source/Programs/Horde/HordeStorage/README.md
@@ -50,7 +50,7 @@ This is the most tested deployment form as this is how we operate it at Epic. We
 The helm chart we install into each regions kubernetes cluster is provided in this repo.
 
 ## On premise
-Horde.Storage can be deployed onprem without using any cloud resources but it does require a kubernetes cluster to run in for production cases. You will also need to setup a Mongo database for it to run again.
+Horde.Storage can be deployed onprem without using any cloud resources but it does require a kubernetes cluster to run in for production cases. You will also need to setup a Mongo database for it to run against.
 Note that if you intend to deployment to multiple on prem sites you will need to manually setup replication of the database used.
 
 # Monitoring
@@ -63,9 +63,83 @@ You can reach the health checks at `/health/live` and `/health/ready` for live a
 # Authentication
 Horde.Storage supports using any OIDC provider that does JWT verfication for authentication. We use Okta at Epic so this is what has been tested but other OIDCs should be compatible as well.
 
-You control the authentication using the `authorizationMap` (helm) `Authorization.NamespaceToClaim` (appSettings) option. This maps a Horde.Storage namespace to which claim is expected to be present to allow access.
-The claim can be set to `*` to allow access for anyone (if a namespace has no entry no one has access).
-You can also use a special `AllNamespaces` claim to grant access to all namespaces when certain claims are present (only used for debugging).
+You configure authentication in two steps, setting up the IdentityProvider (IdP) and then setting up authorization for each namespace.
+
+## IdentityProvider setup
+
+You specify auth schemes to in the `auth` settings. 
+```  
+auth:
+    defaultScheme: Bearer
+    schemes:
+      Bearer: 
+        implementation: "Jwt"
+        jwtAudience: "api://horde"
+        jwtAuthority: "<url-to-your-idp>
+```
+We recommend naming your scheme `Bearer` if its your first and only scheme. You can use multiple schemes to connect against multiple IdPs, this is mostly useful during a migration.
+
+
+## Namespace access
+
+Access to operations within Horde Storage is controlled using a set of actions:
+```
+    - ReadObject
+    - WriteObject
+    - DeleteObject
+    - DeleteBucket
+    - DeleteNamespace
+    - ReadTransactionLog
+    - WriteTransactionLog
+    - AdminAction
+```
+
+These can be assigned either per namespace using the acls in each namespace policy or by assigning them to the acls in the Auth settings (which applies them across all namespaces and for operations not associated with a namespace)
+Example configuration that sets transaction log access for users that have access to it, admin access for admins and then per namespace access.
+```
+auth:
+    acls:
+    - claims: 
+        - groups=app-horde-storage-transactionlog
+        actions:
+        - ReadTransactionLog
+        - WriteTransactionLog
+
+    - claims: 
+        - groups=app-horde-storage-admin
+        actions:
+        - ReadObject
+        - WriteObject
+        - DeleteObject
+        - DeleteBucket
+        - DeleteNamespace
+        - AdminAction
+
+namespace:
+  policy:
+    example-namespace:
+      acls:
+      - actions: 
+        - ReadObject
+        - WriteObject
+        claims: 
+        - ExampleClaim
+     - actions: 
+        - ReadObject
+        - WriteObject
+        claims: 
+        - AnotherClaim
+    open-namespace:
+      acls:
+      - actions: 
+        - ReadObject
+        - WriteObject
+        claims: 
+        - "*"
+```
+
+If you specify multiple claims in the claims array these are ANDed together thus requiring all the claims. A claim statement like `A=B` requires claim `A` to have value `B` (or contain value `B` in the case of a array).
+You can also specify the `*` claim which grants access for any valid token no matter which claims it has, this is mostly for debug / testing scenarios and shouldnt be used for production data.
 
 # Development
 ## Debugging profiles
