Horde: Add AWS Lambda support for agent using the WorkerService

Adds a "one shot mode" to WorkerService that allows it to listen for exactly one lease from the server. This is a hybrid as opposed to sending the ComputeTaskMessage as part of the Lambda invocation request. Through this, the agent will listen for a lease and behave like a normal agent and share the same task queue. Functionality for starting these Lambda based agents is yet to be added.

A test task was added to aid testing in the fake HordeRpcServer.

diff --git a/Engine/Source/Programs/Horde/Horde.Agent.Tests/AwsLambdaListenerCommandTest.cs b/Engine/Source/Programs/Horde/Horde.Agent.Tests/AwsLambdaListenerCommandTest.cs
new file mode 100644
--- /dev/null
+++ b/Engine/Source/Programs/Horde/Horde.Agent.Tests/AwsLambdaListenerCommandTest.cs
@@ -0,0 +1,73 @@
+// Copyright Epic Games, Inc. All Rights Reserved.
+
+using System;
+using System.Threading;
+using System.Threading.Tasks;
+using Google.Protobuf;
+using Horde.Agent.Commands.Compute;
+using Horde.Agent.Services;
+using HordeCommon;
+using HordeCommon.Rpc.Messages;
+using Microsoft.Extensions.Logging;
+using Microsoft.VisualStudio.TestTools.UnitTesting;
+
+namespace Horde.Agent.Tests
+{
+	[TestClass]
+	public class AwsLambdaListenerCommandTest
+	{
+		private readonly ILogger<WorkerService> _workerLogger;
+		private readonly ILogger<GrpcService> _grpcServiceLogger;
+		private readonly ILogger<FakeHordeRpcServer> _hordeRpcServerLogger;
+
+		public AwsLambdaListenerCommandTest()
+		{
+			using ILoggerFactory loggerFactory = LoggerFactory.Create(builder =>
+			{
+				builder.AddConsole(consoleLoggerOptions => consoleLoggerOptions.TimestampFormat = "[HH:mm:ss] ");
+			});
+
+			_workerLogger = loggerFactory.CreateLogger<WorkerService>();
+			_grpcServiceLogger = loggerFactory.CreateLogger<GrpcService>();
+			_hordeRpcServerLogger = loggerFactory.CreateLogger<FakeHordeRpcServer>();
+		}
+		
+		[TestMethod]
+		public async Task ShutdownIfNoLeasesArriveWithinMaxWaitTime()
+		{
+			using CancellationTokenSource cts = new (10000);
+			FakeHordeRpcServer fakeServer = new("testServerName", _hordeRpcServerLogger, cts.Token);
+			WorkerServiceLambda lambdaFunc = new((_) => WorkerServiceTest.GetWorkerService(_grpcServiceLogger, _workerLogger, (a, b, c) => WorkerServiceTest.NullExecutor,(c) => fakeServer.GetClient()));
+
+			AwsLambdaListenResponse response = await InvokeProtoAsync(lambdaFunc, new() { MaxWaitTimeForLeaseMs = 1000 }, cts.Token);
+			Assert.IsFalse(response.DidAcceptLease);
+		}
+		
+		[TestMethod]
+		public async Task HandleOneLeaseAndShutdown()
+		{
+			using CancellationTokenSource cts = new (10000);
+			FakeHordeRpcServer fakeServer = new("testServerName", _hordeRpcServerLogger, cts.Token);
+			WorkerServiceLambda lambdaFunc = new((_) => WorkerServiceTest.GetWorkerService(_grpcServiceLogger, _workerLogger, (a, b, c) => WorkerServiceTest.NullExecutor,(c) => fakeServer.GetClient()));
+
+			fakeServer.AddTestLease("testLeaseId");
+			AwsLambdaListenResponse response = await InvokeProtoAsync(lambdaFunc, new() { MaxWaitTimeForLeaseMs = 5000 }, cts.Token);
+			Assert.AreEqual(LeaseState.Completed, fakeServer.GetLease("testLeaseId").State);
+			Assert.AreEqual(LeaseOutcome.Success, fakeServer.GetLease("testLeaseId").Outcome);
+			Assert.IsTrue(response.DidAcceptLease);
+		}
+
+		/// <summary>
+		/// Helper to serialize/deserialize parameters to <see cref="IAwsLambdaFunction" /> interface.
+		/// </summary>
+		/// <param name="function">AWS Lambda interface to invoke</param>
+		/// <param name="request">Protobuf-based request message</param>
+		/// <param name="cancellationToken"></param>
+		/// <returns>Deserialized Protobuf message</returns>
+		private static async Task<AwsLambdaListenResponse> InvokeProtoAsync(IAwsLambdaFunction function, AwsLambdaListenRequest request, CancellationToken cancellationToken)
+		{
+			ReadOnlyMemory<byte> responseData = await function.OnLambdaInvokeAsync("testRequestId", request.ToByteArray(), cancellationToken);
+			return AwsLambdaListenResponse.Parser.ParseFrom(responseData.Span);
+		}
+	}
+}
diff --git a/Engine/Source/Programs/Horde/Horde.Agent.Tests/WorkerServiceTest.cs b/Engine/Source/Programs/Horde/Horde.Agent.Tests/WorkerServiceTest.cs
--- a/Engine/Source/Programs/Horde/Horde.Agent.Tests/WorkerServiceTest.cs
+++ b/Engine/Source/Programs/Horde/Horde.Agent.Tests/WorkerServiceTest.cs
@@ -3,6 +3,7 @@
 using System;
 using System.Collections.Generic;
 using System.IO;
+using System.Linq;
 using System.Threading;
 using System.Threading.Channels;
 using System.Threading.Tasks;
@@ -29,6 +30,12 @@ namespace Horde.Agent.Tests
 		private readonly ILogger<WorkerService> _workerLogger;
 		private readonly ILogger<GrpcService> _grpcServiceLogger;
 		private readonly ILogger<FakeHordeRpcServer> _hordeRpcServerLogger;
+		
+		internal static IExecutor NullExecutor = new SimpleTestExecutor(async (step, logger, cancellationToken) =>
+		{
+			await Task.Delay(1, cancellationToken);
+			return JobStepOutcome.Success;
+		});
 
 		public WorkerServiceTest()
 		{
@@ -42,12 +49,14 @@ namespace Horde.Agent.Tests
 			_hordeRpcServerLogger = loggerFactory.CreateLogger<FakeHordeRpcServer>();
 		}
 
-		private WorkerService GetWorkerService(
+		internal static WorkerService GetWorkerService(
+			ILogger<GrpcService> grpcServiceLogger,
+			ILogger<WorkerService> workerServiceLogger,
 			Func<IRpcConnection, ExecuteJobTask, BeginBatchResponse, IExecutor>? createExecutor = null,
 			Func<GrpcChannel, HordeRpc.HordeRpcClient>? createHordeRpcClient = null)
 		{
-			AgentSettings settings = new AgentSettings();
-			ServerProfile profile = new ServerProfile();
+			AgentSettings settings = new ();
+			ServerProfile profile = new ();
 			profile.Name = "test";
 			profile.Environment = "test-env";
 			profile.Token = "bogus-token";
@@ -58,12 +67,19 @@ namespace Horde.Agent.Tests
 			settings.Executor = ExecutorType.Test; // Not really used since the executor is overridden in the tests
 			IOptions<AgentSettings> settingsMonitor = new TestOptionsMonitor<AgentSettings>(settings);
 
-			GrpcService grpcService = new (settingsMonitor, _grpcServiceLogger);
-			WorkerService ws = new (_workerLogger, settingsMonitor, grpcService, null!, createExecutor, createHordeRpcClient);
+			GrpcService grpcService = new (settingsMonitor, grpcServiceLogger);
+			WorkerService ws = new (workerServiceLogger, settingsMonitor, grpcService, null!, createExecutor, createHordeRpcClient);
 			ws._rpcConnectionRetryDelay = TimeSpan.FromSeconds(0.1);
 			return ws;
 		}
 
+		private WorkerService GetWorkerService(
+			Func<IRpcConnection, ExecuteJobTask, BeginBatchResponse, IExecutor>? createExecutor = null,
+			Func<GrpcChannel, HordeRpc.HordeRpcClient>? createHordeRpcClient = null)
+		{
+			return GetWorkerService(_grpcServiceLogger, _workerLogger, createExecutor, createHordeRpcClient);
+		}
+
 		[TestMethod]
 		public async Task AbortExecuteStepTest()
 		{
@@ -200,7 +216,7 @@ namespace Horde.Agent.Tests
 			FakeHordeRpcServer fakeServer = new("bogusServerName", _hordeRpcServerLogger, cts.Token);
 			using WorkerService ws = GetWorkerService((a, b, c) => executor,(c) => fakeServer.GetClient());
 
-			Task handleSessionTask = ws.HandleSessionAsync(cts.Token);
+			Task handleSessionTask = ws.HandleSessionAsync(false, false, cts.Token);
 			await fakeServer.CreateSessionReceived.Task.WaitAsync(cts.Token);
 			await fakeServer.UpdateSessionReceived.Task.WaitAsync(cts.Token);
 			await ws.ShutdownAsync(_workerLogger);
@@ -217,7 +233,7 @@ namespace Horde.Agent.Tests
 	{
 		private readonly string _serverName;
 		private bool _isStopping = false;
-		private List<Lease> _leases = new();
+		private Dictionary<string, Lease> _leases = new();
 		private readonly Mock<HordeRpc.HordeRpcClient> _mockClient;
 		private readonly ILogger<FakeHordeRpcServer> _logger;
 		public readonly TaskCompletionSource<bool> CreateSessionReceived = new();
@@ -242,6 +258,27 @@ namespace Horde.Agent.Tests
 				.Returns(() => GetUpdateSessionCall(cancellationToken));
 		}
 
+		public void AddTestLease(string leaseId)
+		{
+			if (_leases.ContainsKey(leaseId))
+			{
+				throw new ArgumentException($"Lease ID {leaseId} already exists");
+			}
+			
+			TestTask testTask = new();
+			_leases[leaseId] = new Lease
+			{
+				Id = leaseId,
+				State = LeaseState.Pending,
+				Payload = Any.Pack(testTask)
+			};
+		}
+
+		public Lease GetLease(string leaseId)
+		{
+			return _leases[leaseId];
+		}
+
 		public HordeRpc.HordeRpcClient GetClient()
 		{
 			return _mockClient.Object;
@@ -289,13 +326,19 @@ namespace Horde.Agent.Tests
 			async Task OnRequest(UpdateSessionRequest request)
 			{
 				UpdateSessionReceived.TrySetResult(true);
-				
-				// TODO: Handle updates of leases sent in request and verify session/agent IDs.
+
+				foreach (Lease agentLease in request.Leases)
+				{
+					Lease serverLease = _leases[agentLease.Id];
+					serverLease.State = agentLease.State;
+					serverLease.Outcome = agentLease.Outcome;
+					serverLease.Output = agentLease.Output;
+				}
 				
 				_logger.LogInformation("OnUpdateSessionRequest: {AgentId} {SessionId} {Status}", request.AgentId, request.SessionId, request.Status);
-				await Task.Delay(1000, cancellationToken);
+				await Task.Delay(100, cancellationToken);
 				UpdateSessionResponse response = new () { ExpiryTime = Timestamp.FromDateTime(DateTime.UtcNow + TimeSpan.FromMinutes(120)) };
-				response.Leases.AddRange(_leases);
+				response.Leases.AddRange(_leases.Values.Where(x => x.State != LeaseState.Completed));
 				await responseStream.Write(response);
 			}
 			
diff --git a/Engine/Source/Programs/Horde/Horde.Agent/Commands/Compute/AwsLambdaClient.cs b/Engine/Source/Programs/Horde/Horde.Agent/Commands/Compute/AwsLambdaClient.cs
--- a/Engine/Source/Programs/Horde/Horde.Agent/Commands/Compute/AwsLambdaClient.cs
+++ b/Engine/Source/Programs/Horde/Horde.Agent/Commands/Compute/AwsLambdaClient.cs
@@ -61,12 +61,12 @@ namespace Horde.Agent.Commands.Compute
 		/// <summary>
 		/// Whether the exception is fatal and requires termination of the process (to adhere to Lambda specs)
 		/// </summary>
-		public bool isFatal { get; }
+		public bool IsFatal { get; }
 		
 		/// <summary>
 		/// HTTP status code in response from AWS Lambda API
 		/// </summary>
-		public HttpStatusCode? statusCode { get; }
+		public HttpStatusCode? StatusCode { get; }
 
 		/// <summary>
 		/// Constructor
@@ -76,15 +76,16 @@ namespace Horde.Agent.Commands.Compute
 		/// <param name="isFatal"></param>
 		public AwsLambdaClientException(string message, bool isFatal = false, HttpStatusCode? statusCode = null) : base(message)
 		{
-			this.isFatal = isFatal;
-			this.statusCode = statusCode;
+			IsFatal = isFatal;
+			StatusCode = statusCode;
 		}
 	}
 
 	/// <summary>
 	/// Client interfacing with the AWS Lambda runtime API
+	/// See https://docs.aws.amazon.com/lambda/latest/dg/runtimes-api.html
 	/// </summary>
-	class AwsLambdaClient
+	public class AwsLambdaClient
 	{
 		private const string EnvVarRuntimeApi = "AWS_LAMBDA_RUNTIME_API";
 		private const string HeaderNameAwsRequestId = "Lambda-Runtime-Aws-Request-Id";
@@ -95,12 +96,24 @@ namespace Horde.Agent.Commands.Compute
 		private readonly ILogger<AwsLambdaClient> _logger;
 		private readonly string _hostPort;
 
+		/// <summary>
+		/// Constructor
+		/// </summary>
+		/// <param name="hostPort">Base URL to the Lambda Runtime API endpoint</param>
+		/// <param name="logger"></param>
 		public AwsLambdaClient(string hostPort, ILogger<AwsLambdaClient> logger)
 		{
 			_hostPort = hostPort;
 			_logger = logger;
 		}
 
+		/// <summary>
+		/// Create an AwsLambdaClient from environment variables. Usually how the client is created when running
+		/// inside a real AWS Lambda environment.
+		/// </summary>
+		/// <param name="logger"></param>
+		/// <returns>A configured AwsLambdaClient instance</returns>
+		/// <exception cref="ArgumentException"></exception>
 		public static AwsLambdaClient InitFromEnv(ILogger<AwsLambdaClient> logger)
 		{
 			string? runtimeApi = Environment.GetEnvironmentVariable(EnvVarRuntimeApi);
@@ -122,7 +135,15 @@ namespace Horde.Agent.Commands.Compute
 			return new HttpClient();
 		}
 		
-		public async Task<NextInvocationResponse> GetNextInvocationAsync(CancellationToken cancellationToken)
+		/// <summary>
+		/// Get the next invocation from the Lambda runtime API
+		///
+		/// Can block for many minutes waiting for a new invocation request to be received. 
+		/// </summary>
+		/// <param name="cancellationToken"></param>
+		/// <returns>A response containing details for invocation</returns>
+		/// <exception cref="AwsLambdaClientException"></exception>
+		private async Task<NextInvocationResponse> GetNextInvocationAsync(CancellationToken cancellationToken)
 		{
 			using HttpClient client = GetHttpClient();
 			client.Timeout = TimeSpan.FromHours(1); // Waiting for the next invocation can block for a while
@@ -158,7 +179,7 @@ namespace Horde.Agent.Commands.Compute
 			return new NextInvocationResponse(requestId, invokedFunctionArn, deadlineMs, data);
 		}
 		
-		public async Task SendInvocationResponseAsync(string awsRequestId, ReadOnlyMemory<byte> data)
+		private async Task SendInvocationResponseAsync(string awsRequestId, ReadOnlyMemory<byte> data)
 		{
 			using HttpClient client = GetHttpClient();
 			client.Timeout = TimeSpan.FromSeconds(30);
@@ -173,17 +194,72 @@ namespace Horde.Agent.Commands.Compute
 			}
 		}
 		
+		/// <summary>
+		/// Notify the Lambda runtime API that an error occurred during initialization. 
+		/// </summary>
+		/// <param name="errorType">Any string but AWS recommends a format of [category.reason]</param>
+		/// <param name="errorMessage">Message describing the error</param>
+		/// <param name="stackTrace">Lines of an (optional) stacktrace</param>
+		/// <returns></returns>
 		public Task SendInitErrorAsync(string errorType, string errorMessage, List<string>? stackTrace = null)
 		{
 			return SendErrorAsync(GetInitErrorUrl(), errorType, errorMessage, stackTrace);
 		}
 		
-		public Task SendInvocationErrorAsync(string awsRequestId, string errorType, string errorMessage, List<string>? stackTrace = null)
+		private Task SendInvocationErrorAsync(string awsRequestId, string errorType, string errorMessage, List<string>? stackTrace = null)
 		{
 			return SendErrorAsync(GetInvocationErrorUrl(awsRequestId), errorType, errorMessage, stackTrace);
 		}
 		
-		public static async Task SendErrorAsync(string url, string errorType, string errorMessage, List<string>? stackTrace)
+		/// <summary>
+		/// Listen for new Lambda invocations
+		///
+		/// Call <paramref name="function"/> for each new invocation received.
+		/// Will block until cancellation token is triggered.
+		/// </summary>
+		/// <param name="function"></param>
+		/// <param name="cancellationToken"></param>
+		/// <returns></returns>
+		public async Task<bool> ListenForInvocationsAsync(IAwsLambdaFunction function, CancellationToken cancellationToken)
+		{
+			while (!cancellationToken.IsCancellationRequested)
+			{
+				NextInvocationResponse? nextInvocationResponse = null;
+				try
+				{
+					_logger.LogDebug("Waiting for next invocation...");
+
+					// Block and wait for the next invocation (can be multiple minutes)
+					nextInvocationResponse = await GetNextInvocationAsync(cancellationToken);
+					ReadOnlyMemory<byte> responseData = await function.OnLambdaInvokeAsync(nextInvocationResponse.requestId, nextInvocationResponse.data, cancellationToken);
+					await SendInvocationResponseAsync(nextInvocationResponse.requestId, responseData);
+				}
+				catch (Exception e)
+				{
+					_logger.LogError(e, "Exception invoking Lambda function from request");
+					
+					if (nextInvocationResponse != null)
+					{
+						try
+						{
+							List<string>? stacktraceLines = e.StackTrace?.Split("\n").ToList();
+							await SendInvocationErrorAsync(nextInvocationResponse.requestId, "general", $"{e.GetType()}: {e.Message}", stacktraceLines);
+						}
+						catch (AwsLambdaClientException sendErrorException)
+						{
+							_logger.LogError(sendErrorException, "Bad response when sending invocation error. isFatal={IsFatal}", sendErrorException.IsFatal);
+							return false;
+						}
+					}
+					
+					await Task.Delay(500, cancellationToken); // Cool down before trying another next invocation call
+				}
+			}
+
+			return true;
+		}
+		
+		private static async Task SendErrorAsync(string url, string errorType, string errorMessage, List<string>? stackTrace)
 		{
 			using HttpClient client = new HttpClient();
 			client.Timeout = TimeSpan.FromHours(1);
@@ -205,4 +281,19 @@ namespace Horde.Agent.Commands.Compute
 			}
 		}
 	}
+
+	/// <summary>
+	/// Represents an AWS Lambda function that can be invoked
+	/// </summary>
+	public interface IAwsLambdaFunction
+	{
+		/// <summary>
+		/// Handles an AWS Lambda function call
+		/// </summary>
+		/// <param name="requestId">Lambda request ID for this invocation</param>
+		/// <param name="requestData">Raw bytes for the request body</param>
+		/// <param name="cancellationToken"></param>
+		/// <returns>Raw bytes for the response body</returns>
+		public Task<ReadOnlyMemory<byte>> OnLambdaInvokeAsync(string requestId, ReadOnlyMemory<byte> requestData, CancellationToken cancellationToken);
+	}
 }
diff --git a/Engine/Source/Programs/Horde/Horde.Agent/Commands/Compute/AwsLambdaCommand.cs b/Engine/Source/Programs/Horde/Horde.Agent/Commands/Compute/AwsLambdaCommand.cs
deleted file mode 100644
diff --git a/Engine/Source/Programs/Horde/Horde.Agent/Commands/Compute/AwsLambdaListenerCommand.cs b/Engine/Source/Programs/Horde/Horde.Agent/Commands/Compute/AwsLambdaListenerCommand.cs
new file mode 100644
--- /dev/null
+++ b/Engine/Source/Programs/Horde/Horde.Agent/Commands/Compute/AwsLambdaListenerCommand.cs
@@ -0,0 +1,222 @@
+// Copyright Epic Games, Inc. All Rights Reserved.
+
+using System;
+using System.Collections.Generic;
+using System.Linq;
+using System.Net.Http;
+using System.Text;
+using System.Threading;
+using System.Threading.Tasks;
+using Amazon.SimpleSystemsManagement;
+using Amazon.SimpleSystemsManagement.Model;
+using EpicGames.Core;
+using EpicGames.Horde.Auth;
+using EpicGames.Horde.Storage;
+using EpicGames.Horde.Storage.Impl;
+using Google.Protobuf;
+using Grpc.Core;
+using Horde.Agent.Services;
+using HordeCommon.Rpc.Messages;
+using Microsoft.Extensions.Logging;
+using Microsoft.Extensions.Options;
+using Serilog;
+using Serilog.Events;
+using Command = EpicGames.Core.Command;
+using ILogger = Microsoft.Extensions.Logging.ILogger;
+
+namespace Horde.Agent.Commands.Compute
+{
+	/// <summary>
+	/// Run WorkerService inside an AWS Lambda function invocation
+	/// Will be run in "one-shot mode" where either one lease will be handled, or max wait time is encountered.
+	/// Since the running time of a single Lambda function invocation is limited to 15 minutes,
+	/// we want to handle only one lease before returning.
+	/// </summary>
+	class WorkerServiceLambda : IAwsLambdaFunction
+	{
+		private readonly Func<string, WorkerService> _createWorkerService;
+		
+		public WorkerServiceLambda(Func<string, WorkerService> createWorkerService)
+		{
+			_createWorkerService = createWorkerService;
+		}
+
+		/// <inheritdoc />
+		public async Task<ReadOnlyMemory<byte>> OnLambdaInvokeAsync(string requestId, ReadOnlyMemory<byte> requestData, CancellationToken cancellationToken)
+		{
+			bool isProtobufCompatibleJson = requestData.Length >= 1 && requestData.Span[0] == '{';
+			AwsLambdaListenRequest request = isProtobufCompatibleJson ?
+				AwsLambdaListenRequest.Parser.ParseJson(Encoding.UTF8.GetString(requestData.Span)) :
+				AwsLambdaListenRequest.Parser.ParseFrom(requestData.Span);
+
+			TimeSpan maxWaitTime = TimeSpan.FromMilliseconds(request.MaxWaitTimeForLeaseMs);
+			using WorkerService service = _createWorkerService(requestId);
+			bool isLeaseActive = false;
+			service.OnLeaseActive += lease => isLeaseActive = true;
+			using CancellationTokenSource maxWaitTimeCts = new();
+
+			Task _ = Task.Run(async () =>
+			{
+				await Task.Delay(maxWaitTime, cancellationToken);
+				if (!isLeaseActive)
+				{
+					maxWaitTimeCts.Cancel();
+				}
+			}, cancellationToken);
+
+			// Combine tokens so either max wait time or original token can cancel
+			using CancellationTokenSource linkedCts = CancellationTokenSource.CreateLinkedTokenSource(cancellationToken, maxWaitTimeCts.Token);
+			try
+			{
+				await service.HandleSessionAsync(true, false, linkedCts.Token);
+			}
+			catch (RpcException e)
+			{
+				// Don't bubble exceptions for cancellations as they can happen naturally from maxWaitTime parameter above
+				if (e.Status.StatusCode != StatusCode.Cancelled)
+				{
+					throw;
+				}
+			}
+			
+			AwsLambdaListenResponse response = new () { DidAcceptLease = service.NumLeasesCompleted > 0 };
+			byte[] responseData = isProtobufCompatibleJson ?
+				Encoding.UTF8.GetBytes(JsonFormatter.Default.Format(response)) :
+				response.ToByteArray();
+
+			return responseData;
+		}
+	}
+
+	/// <summary>
+	/// Run agent inside an AWS Lambda function.
+	/// </summary>
+	[Command("Compute", "AwsLambda", "Listen for AWS Lambda invocations")]
+	class AwsLambdaListenerCommand : Command
+	{
+		/// <inheritdoc/>
+		public override async Task<int> ExecuteAsync(ILogger logger)
+		{
+			LogLevel logLevel;
+			string logLevelStr = Environment.GetEnvironmentVariable("UE_HORDE_LOG_LEVEL") ?? "debug";
+
+			if (Enum.TryParse(logLevelStr, true, out LogLevel logEventLevel))
+			{
+				logLevel = logEventLevel;
+			}
+			else
+			{
+				logger.LogError("Unable to parse log level: {LogLevelStr}", logLevelStr);
+				return 1;
+			}
+
+			Serilog.Core.Logger serilogLogger = new LoggerConfiguration()
+				.MinimumLevel.Debug()
+				.MinimumLevel.Override("Microsoft", LogEventLevel.Warning)
+				.Enrich.FromLogContext()
+				.WriteTo.Console()
+				.CreateLogger();
+			
+			using ILoggerFactory loggerFactory = LoggerFactory.Create(logging => { logging.AddSerilog(serilogLogger); });
+			ILogger<WorkerService> workerServiceLogger = loggerFactory.CreateLogger<WorkerService>();
+			ILogger<GrpcService> grpcServiceLogger = loggerFactory.CreateLogger<GrpcService>();
+			ILogger<AwsLambdaClient> lambdaClientLogger = loggerFactory.CreateLogger<AwsLambdaClient>();
+			
+			AwsLambdaClient lambdaClient = AwsLambdaClient.InitFromEnv(lambdaClientLogger);
+			
+			using CancellationTokenSource cts = new ();
+			try
+			{
+				logger.LogInformation("Initializing AWS Lambda executor...");
+
+				HttpServiceClientOptions clientOptions = await CreateHordeStorageOptions(cts.Token);
+				using HttpClient httpClient = new ();
+				OAuthHandlerFactory oAuthHandlerFactory = new (httpClient);
+				OAuthHandler<HttpStorageClient> oAuthHandler = oAuthHandlerFactory.Create<HttpStorageClient>(clientOptions);
+				using HttpClient client = new (oAuthHandler);
+				client.BaseAddress = clientOptions.Url;
+				IStorageClient storageClient = new HttpStorageClient(client);
+
+				AgentSettings agentSettings = await CreateAgentSettingsAsync(cts.Token);
+				agentSettings.WorkingDir = Environment.GetEnvironmentVariable("UE_HORDE_WORKING_DIR") ?? "/tmp";
+				agentSettings.PerforceExecutor.RunConform = false;
+				
+				OptionsWrapper<AgentSettings> agentSettingsOptions = new (agentSettings);
+				GrpcService grpcService = new (agentSettingsOptions, grpcServiceLogger);
+
+				AppDomain.CurrentDomain.ProcessExit += (s, e) => 
+				{
+					logger.LogInformation("ProcessExit triggered. Shutting down...");
+					cts.Cancel();
+				};
+
+				WorkerServiceLambda lambdaFunction = new ((requestId) =>
+				{
+					agentSettingsOptions.Value.Name = "lambda-" + requestId;
+					return new WorkerService(workerServiceLogger, agentSettingsOptions, grpcService, storageClient);
+				});
+				bool success = await lambdaClient.ListenForInvocationsAsync(lambdaFunction, cts.Token);
+				return success ? 1 : 0;
+			}
+			catch (Exception e)
+			{
+				logger.LogError(e, "Error initializing");
+				try
+				{
+					List<string>? stacktraceLines = e.StackTrace?.Split("\n").ToList();
+					await lambdaClient.SendInitErrorAsync("general", $"{e.GetType()}: {e.Message}", stacktraceLines);
+				}
+				catch (AwsLambdaClientException initException)
+				{
+					logger.LogError(initException, "Bad response when sending init error. isFatal={IsFatal}", initException.IsFatal);
+				}
+				
+				return 1;
+			}
+		}
+
+		private static async Task<string> GetSsmParameterSecretAsync(string name, CancellationToken cancellationToken)
+		{
+			using AmazonSimpleSystemsManagementClient ssmClient = new ();
+			GetParameterRequest request = new () { Name = name, WithDecryption = true };
+			GetParameterResponse response = await ssmClient.GetParameterAsync(request, cancellationToken);
+			return response.Parameter.Value;
+		}
+
+		private static async Task<HttpServiceClientOptions> CreateHordeStorageOptions(CancellationToken cancellationToken)
+		{
+			return new HttpServiceClientOptions
+			{
+				Url = new Uri(GetEnvVar("UE_HORDE_STORAGE_URL")),
+				AuthUrl = new Uri(GetEnvVar("UE_HORDE_STORAGE_OAUTH_URL")),
+				GrantType = GetEnvVar("UE_HORDE_STORAGE_OAUTH_GRANT_TYPE"),
+				ClientId = GetEnvVar("UE_HORDE_STORAGE_OAUTH_CLIENT_ID"),
+				ClientSecret = await GetSsmParameterSecretAsync(GetEnvVar("UE_HORDE_STORAGE_OAUTH_CLIENT_SECRET_ARN"), cancellationToken),
+				Scope = GetEnvVar("UE_HORDE_STORAGE_OAUTH_SCOPE"),
+			};
+		}
+		
+		private static async Task<AgentSettings> CreateAgentSettingsAsync(CancellationToken cancellationToken)
+		{
+			AgentSettings settings = new ();
+			settings.Server = "AwsLambdaDefault"; 
+			settings.ServerProfiles.Add(new ServerProfile
+			{
+				Name = settings.Server,
+				Url = new Uri(GetEnvVar("UE_HORDE_BUILD_URL")),
+				Token = await GetSsmParameterSecretAsync(GetEnvVar("UE_HORDE_BUILD_TOKEN"), cancellationToken)
+			});
+			return settings;
+		}
+
+		private static string GetEnvVar(string name)
+		{
+			string? value = Environment.GetEnvironmentVariable(name);
+			if (value == null)
+			{
+				throw new ArgumentException($"Missing env var {name}");
+			}
+			return value;
+		}
+	}
+}
diff --git a/Engine/Source/Programs/Horde/Horde.Agent/Horde.Agent.csproj b/Engine/Source/Programs/Horde/Horde.Agent/Horde.Agent.csproj
--- a/Engine/Source/Programs/Horde/Horde.Agent/Horde.Agent.csproj
+++ b/Engine/Source/Programs/Horde/Horde.Agent/Horde.Agent.csproj
@@ -127,12 +127,13 @@
     <Protobuf Include="..\Horde.Build\Protos\horde\messages\agent_requirements.proto" ProtoRoot="..\Horde.Build\Protos" Access="Public" />
     <Protobuf Include="..\Horde.Build\Protos\horde\messages\agent_workspace.proto" ProtoRoot="..\Horde.Build\Protos" Access="Public" />
     <Protobuf Include="..\Horde.Build\Protos\horde\messages\lease.proto" ProtoRoot="..\Horde.Build\Protos" Access="Public" />
-    <Protobuf Include="..\Horde.Build\Protos\horde\tasks\aws_lambda_compute_task.proto" ProtoRoot="..\Horde.Build\Protos" Access="Public" />
+    <Protobuf Include="..\Horde.Build\Protos\horde\messages\aws_lambda.proto" ProtoRoot="..\Horde.Build\Protos" Access="Public" />
     <Protobuf Include="..\Horde.Build\Protos\horde\tasks\compute_task.proto" ProtoRoot="..\Horde.Build\Protos" Access="Public" />
     <Protobuf Include="..\Horde.Build\Protos\horde\tasks\conform_task.proto" ProtoRoot="..\Horde.Build\Protos" Access="Public" />
     <Protobuf Include="..\Horde.Build\Protos\horde\tasks\job_task.proto" ProtoRoot="..\Horde.Build\Protos" Access="Public" />
     <Protobuf Include="..\Horde.Build\Protos\horde\tasks\restart_task.proto" ProtoRoot="..\Horde.Build\Protos" Access="Public" />
     <Protobuf Include="..\Horde.Build\Protos\horde\tasks\shutdown_task.proto" ProtoRoot="..\Horde.Build\Protos" Access="Public" />
+    <Protobuf Include="..\Horde.Build\Protos\horde\tasks\test_task.proto" ProtoRoot="..\Horde.Build\Protos" Access="Public" />
     <Protobuf Include="..\Horde.Build\Protos\horde\tasks\upgrade_task.proto" ProtoRoot="..\Horde.Build\Protos" Access="Public" />
   </ItemGroup>
 
diff --git a/Engine/Source/Programs/Horde/Horde.Agent/Services/WorkerService.cs b/Engine/Source/Programs/Horde/Horde.Agent/Services/WorkerService.cs
--- a/Engine/Source/Programs/Horde/Horde.Agent/Services/WorkerService.cs
+++ b/Engine/Source/Programs/Horde/Horde.Agent/Services/WorkerService.cs
@@ -173,13 +173,18 @@ namespace Horde.Agent.Services
 		/// </summary>
 		readonly List<LeaseInfo> _activeLeases = new List<LeaseInfo>();
 
+		/// <summary>
+		/// Number of leases completed
+		/// </summary>
+		public int NumLeasesCompleted { get; private set; }
+
 		/// <summary>
 		/// Whether the agent is currently in an unhealthy state
 		/// </summary>
 		readonly bool _unhealthy = false;
 
 		/// <summary>
-		/// Whether the 
+		/// Whether a shutdown of the worker service has been requested
 		/// </summary>
 		bool _requestShutdown = false;
 
@@ -207,6 +212,16 @@ namespace Horde.Agent.Services
 		/// Number of times UpdateSession has failed
 		/// </summary>
 		int _updateSessionFailures;
+		
+		/// <summary>
+		/// Delegate for lease active events
+		/// </summary>
+		public delegate void LeaseActiveEvent(Lease lease);
+
+		/// <summary>
+		/// Event triggered when a lease is accepted and set to active
+		/// </summary>
+		public event LeaseActiveEvent? OnLeaseActive;
 
 		/// <summary>
 		/// Function for creating a new executor. Primarily to aid testing. 
@@ -324,7 +339,7 @@ namespace Horde.Agent.Services
 				{
 					if (_activeLeases.Count == 0)
 					{
-						await HandleSessionAsync(stoppingToken);
+						await HandleSessionAsync(false, true, stoppingToken);
 					}
 				}
 				catch (Exception ex)
@@ -437,14 +452,22 @@ namespace Horde.Agent.Services
 
 		/// <summary>
 		/// Handles the lifetime of an agent session.
+		///
+		/// <paramref name="shutdownAfterFinishedLease"/> is primarily for executing leases in a one-shot environment,
+		/// such as an AWS Lambda function.
 		/// </summary>
+		/// <param name="shutdownAfterFinishedLease">Stop after one lease has finished</param>
+		/// <param name="ensureSingleInstance">Ensure only one instance of agent is running</param>
 		/// <param name="stoppingToken">Indicates that the service is trying to stop</param>
 		/// <returns>Async task</returns>
-		internal async Task HandleSessionAsync(CancellationToken stoppingToken)
+		internal async Task HandleSessionAsync(bool shutdownAfterFinishedLease, bool ensureSingleInstance, CancellationToken stoppingToken)
 		{
-			// Make sure there's only one instance of the agent running
-			using Mutex singleInstanceMutex = new Mutex(false, "Global\\HordeAgent-DB828ACB-0AA5-4D32-A62A-21D4429B1014");
-			await WaitForMutexAsync(singleInstanceMutex, stoppingToken);
+			using Mutex singleInstanceMutex = new (false, "Global\\HordeAgent-DB828ACB-0AA5-4D32-A62A-21D4429B1014");
+			if (ensureSingleInstance)
+			{
+				// Make sure there's only one instance of the agent running
+				await WaitForMutexAsync(singleInstanceMutex, stoppingToken);				
+			}
 
 			// Terminate any remaining child processes from other instances
 			TerminateProcesses(_logger, stoppingToken);
@@ -496,7 +519,7 @@ namespace Horde.Agent.Services
 
 				// Create a session
 				createSessionResponse = await rpcClient.CreateSessionAsync(sessionRequest, null, null, stoppingToken);
-				_logger.LogInformation("Session started ({SessionId})", createSessionResponse.SessionId);
+				_logger.LogInformation("Session started. AgentName={AgentName} SessionId={SessionId}", _settings.GetAgentName(), createSessionResponse.SessionId);
 			}
 
 			Func<GrpcChannel> createGrpcChannel = () => _grpcService.CreateGrpcChannel(createSessionResponse.Token);
@@ -515,7 +538,18 @@ namespace Horde.Agent.Services
 					Task waitTask = _updateLeasesEvent.Task;
 
 					// Flag for whether the service is stopping
-					bool stopping = stoppingToken.IsCancellationRequested || _requestShutdown;
+					bool stopping = false;
+					if (stoppingToken.IsCancellationRequested)
+					{
+						_logger.LogInformation("Cancellation from token requested");
+						stopping = true;
+					}
+					
+					if (_requestShutdown)
+					{
+						_logger.LogInformation("Shutdown requested");
+						stopping = true;
+					}
 
 					// Build the next update request
 					UpdateSessionRequest updateSessionRequest = new UpdateSessionRequest();
@@ -586,9 +620,17 @@ namespace Horde.Agent.Services
 							// Now reconcile the local state to match what the server reports
 							if (updateSessionResponse != null)
 							{
-								// Remove any leases which have completed
-								_activeLeases.RemoveAll(x => (x.Lease.State == LeaseState.Completed || x.Lease.State == LeaseState.Cancelled) && !updateSessionResponse.Leases.Any(y => y.Id == x.Lease.Id && y.State != LeaseState.Cancelled));
+								bool atLeastOneLeaseFinished = _activeLeases.Any(x => x.Lease.State is LeaseState.Completed or LeaseState.Cancelled);
+								if (atLeastOneLeaseFinished && shutdownAfterFinishedLease)
+								{
+									_logger.LogInformation("At least one lease executed. Requesting shutdown.");
+									_requestShutdown = true;
+								}
 
+								// Remove any leases which have completed
+								int numRemoved = _activeLeases.RemoveAll(x => (x.Lease.State == LeaseState.Completed || x.Lease.State == LeaseState.Cancelled) && !updateSessionResponse.Leases.Any(y => y.Id == x.Lease.Id && y.State != LeaseState.Cancelled));
+								NumLeasesCompleted += numRemoved;
+ 
 								// Create any new leases and cancel any running leases
 								foreach (Lease serverLease in updateSessionResponse.Leases)
 								{
@@ -609,6 +651,7 @@ namespace Horde.Agent.Services
 										LeaseInfo info = new LeaseInfo(serverLease);
 										info.Task = Task.Run(() => HandleLeaseAsync(rpcCon, createSessionResponse.AgentId, info), CancellationToken.None);
 										_activeLeases.Add(info);
+										OnLeaseActive?.Invoke(serverLease);
 									}
 								}
 							}
@@ -829,6 +872,13 @@ namespace Horde.Agent.Services
 				Task<LeaseResult> Handler(ILogger newLogger) => RestartAsync(newLogger);
 				return await HandleLeasePayloadWithLogAsync(rpcConnection, restartTask.LogId, null, null, Handler, leaseInfo.CancellationTokenSource.Token);
 			}
+			
+			TestTask testTask;
+			if (leaseInfo.Lease.Payload.TryUnpack(out testTask))
+			{
+				GlobalTracer.Instance.ActiveSpan?.SetTag("task", "Test");
+				return LeaseResult.Success;
+			}
 
 			_logger.LogError("Invalid lease payload type ({PayloadType})", payload.TypeUrl);
 			return LeaseResult.Failed;
@@ -1769,7 +1819,8 @@ namespace Horde.Agent.Services
 			// Get the IP addresses
 			try
 			{
-				IPHostEntry entry = await Dns.GetHostEntryAsync(Dns.GetHostName());
+				using CancellationTokenSource dnsCts = new(3000);
+				IPHostEntry entry = await Dns.GetHostEntryAsync(Dns.GetHostName(), dnsCts.Token);
 				foreach (IPAddress address in entry.AddressList)
 				{
 					if (address.AddressFamily == System.Net.Sockets.AddressFamily.InterNetwork)
