Horde: Add a fleet manager that can create/terminate instances from scratch on AWS EC2

The current AWS fleet manager was renamed to clarify it only can start existing and stopped instances.

diff --git a/Engine/Source/Programs/Horde/Horde.Build.Tests/Fleet/AutoscaleServiceV2Test.cs b/Engine/Source/Programs/Horde/Horde.Build.Tests/Fleet/AutoscaleServiceV2Test.cs
--- a/Engine/Source/Programs/Horde/Horde.Build.Tests/Fleet/AutoscaleServiceV2Test.cs
+++ b/Engine/Source/Programs/Horde/Horde.Build.Tests/Fleet/AutoscaleServiceV2Test.cs
@@ -25,19 +25,19 @@ namespace Horde.Build.Tests.Fleet
 		public int ExpandPoolAsyncCallCount { get; private set; }
 		public int ShrinkPoolAsyncCallCount { get; private set; }
 		
-		public Task ExpandPoolAsync(IPool pool, IReadOnlyList<IAgent> agents, int count)
+		public Task ExpandPoolAsync(IPool pool, IReadOnlyList<IAgent> agents, int count, CancellationToken cancellationToken)
 		{
 			ExpandPoolAsyncCallCount++;
 			return Task.CompletedTask;
 		}
 
-		public Task ShrinkPoolAsync(IPool pool, IReadOnlyList<IAgent> agents, int count)
+		public Task ShrinkPoolAsync(IPool pool, IReadOnlyList<IAgent> agents, int count, CancellationToken cancellationToken)
 		{
 			ShrinkPoolAsyncCallCount++;
 			return Task.CompletedTask;
 		}
 
-		public Task<int> GetNumStoppedInstancesAsync(IPool pool)
+		public Task<int> GetNumStoppedInstancesAsync(IPool pool, CancellationToken cancellationToken)
 		{
 			throw new NotImplementedException();
 		}
diff --git a/Engine/Source/Programs/Horde/Horde.Build.Tests/TestSetup.cs b/Engine/Source/Programs/Horde/Horde.Build.Tests/TestSetup.cs
--- a/Engine/Source/Programs/Horde/Horde.Build.Tests/TestSetup.cs
+++ b/Engine/Source/Programs/Horde/Horde.Build.Tests/TestSetup.cs
@@ -190,7 +190,7 @@ namespace Horde.Build.Tests
 			services.AddSingleton<AgentService>();
 			services.AddSingleton<AgentSoftwareService>();
 			services.AddSingleton<AutoscaleService>();
-			services.AddSingleton<AwsFleetManager, AwsFleetManager>();
+			services.AddSingleton<AwsReuseFleetManager, AwsReuseFleetManager>();
 			services.AddSingleton<ConsistencyService>();
 			services.AddSingleton<ConfigCollection>();
 			services.AddSingleton<RequestTrackerService>();
@@ -199,7 +199,7 @@ namespace Horde.Build.Tests
 			services.AddSingleton<IDowntimeService, DowntimeServiceStub>();
 			services.AddSingleton<IDogStatsd, NoOpDogStatsd>();
 			services.AddSingleton<IssueService>();
-			services.AddSingleton<IFleetManager, AwsFleetManager>();
+			services.AddSingleton<IFleetManager, AwsReuseFleetManager>();
 			services.AddSingleton<JobService>();
 			services.AddSingleton<LifetimeService>();
 			services.AddSingleton<ILogStorage, NullLogStorage>();
diff --git a/Engine/Source/Programs/Horde/Horde.Build/Agents/Fleet/IFleetManager.cs b/Engine/Source/Programs/Horde/Horde.Build/Agents/Fleet/IFleetManager.cs
--- a/Engine/Source/Programs/Horde/Horde.Build/Agents/Fleet/IFleetManager.cs
+++ b/Engine/Source/Programs/Horde/Horde.Build/Agents/Fleet/IFleetManager.cs
@@ -1,6 +1,7 @@
 // Copyright Epic Games, Inc. All Rights Reserved.
 
 using System.Collections.Generic;
+using System.Threading;
 using System.Threading.Tasks;
 using Horde.Build.Agents.Pools;
 
@@ -17,8 +18,9 @@ namespace Horde.Build.Agents.Fleet
 		/// <param name="pool">Pool to resize</param>
 		/// <param name="agents">Current list of agents in the pool</param>
 		/// <param name="count">Number of agents to add</param>
+		/// <param name="cancellationToken">Cancellation token for the call</param>
 		/// <returns>Async task</returns>
-		Task ExpandPoolAsync(IPool pool, IReadOnlyList<IAgent> agents, int count);
+		Task ExpandPoolAsync(IPool pool, IReadOnlyList<IAgent> agents, int count, CancellationToken cancellationToken = default);
 
 		/// <summary>
 		/// Shrink the given pool
@@ -26,15 +28,17 @@ namespace Horde.Build.Agents.Fleet
 		/// <param name="pool">Pool to resize</param>
 		/// <param name="agents">Current list of agents in the pool</param>
 		/// <param name="count">Number of agents to remove</param>
+		/// <param name="cancellationToken">Cancellation token for the call</param>
 		/// <returns>Async task</returns>
-		Task ShrinkPoolAsync(IPool pool, IReadOnlyList<IAgent> agents, int count);
+		Task ShrinkPoolAsync(IPool pool, IReadOnlyList<IAgent> agents, int count, CancellationToken cancellationToken = default);
 
 		/// <summary>
 		/// Returns the number of stopped instances in the given pool
 		/// </summary>
 		/// <param name="pool">Pool to resize</param>
+		/// <param name="cancellationToken">Cancellation token for the call</param>
 		/// <returns>Async task</returns>
-		Task<int> GetNumStoppedInstancesAsync(IPool pool);
+		Task<int> GetNumStoppedInstancesAsync(IPool pool, CancellationToken cancellationToken = default);
 	}
 }
 
diff --git a/Engine/Source/Programs/Horde/Horde.Build/Agents/Fleet/Providers/AwsFleetManager.cs b/Engine/Source/Programs/Horde/Horde.Build/Agents/Fleet/Providers/AwsFleetManager.cs
deleted file mode 100644
diff --git a/Engine/Source/Programs/Horde/Horde.Build/Agents/Fleet/Providers/AwsReuseFleetManager.cs b/Engine/Source/Programs/Horde/Horde.Build/Agents/Fleet/Providers/AwsReuseFleetManager.cs
new file mode 100644
--- /dev/null
+++ b/Engine/Source/Programs/Horde/Horde.Build/Agents/Fleet/Providers/AwsReuseFleetManager.cs
@@ -0,0 +1,121 @@
+// Copyright Epic Games, Inc. All Rights Reserved.
+
+using System;
+using System.Collections.Generic;
+using System.Linq;
+using System.Threading;
+using System.Threading.Tasks;
+using Amazon;
+using Amazon.EC2;
+using Amazon.EC2.Model;
+using Horde.Build.Agents.Pools;
+using Horde.Build.Auditing;
+using Microsoft.Extensions.Logging;
+using OpenTracing;
+using OpenTracing.Util;
+
+namespace Horde.Build.Agents.Fleet.Providers
+{
+	/// <summary>
+	/// Fleet manager for handling AWS EC2 instances
+	/// Will start already existing but stopped instances to reuse existing EBS disks.
+	/// See <see cref="AwsFleetManager" /> for creating and terminating instances from scratch.
+	/// </summary>
+	public sealed class AwsReuseFleetManager : IFleetManager, IDisposable
+	{
+		readonly AmazonEC2Client _client;
+		readonly IAgentCollection _agentCollection;
+		readonly ILogger _logger;
+
+		/// <summary>
+		/// Constructor
+		/// </summary>
+		public AwsReuseFleetManager(IAgentCollection agentCollection, ILogger<AwsReuseFleetManager> logger)
+		{
+			_agentCollection = agentCollection;
+			_logger = logger;
+
+			AmazonEC2Config config = new AmazonEC2Config();
+			config.RegionEndpoint = RegionEndpoint.USEast1;
+
+			logger.LogInformation("Initializing AWS fleet manager for region {Region}", config.RegionEndpoint);
+
+			_client = new AmazonEC2Client(config);
+		}
+
+		/// <inheritdoc/>
+		public void Dispose()
+		{
+			_client.Dispose();
+		}
+
+		/// <inheritdoc/>
+		public async Task ExpandPoolAsync(IPool pool, IReadOnlyList<IAgent> agents, int count, CancellationToken cancellationToken)
+		{
+			using IScope scope = GlobalTracer.Instance.BuildSpan("ExpandPool").StartActive();
+			scope.Span.SetTag("poolName", pool.Name);
+			scope.Span.SetTag("numAgents", agents.Count);
+			scope.Span.SetTag("count", count);
+
+			DescribeInstancesResponse describeResponse;
+			using (IScope describeScope = GlobalTracer.Instance.BuildSpan("DescribeInstances").StartActive())
+			{
+				// Find stopped instances in the correct pool
+				DescribeInstancesRequest describeRequest = new DescribeInstancesRequest();
+				describeRequest.Filters = new List<Filter>();
+				describeRequest.Filters.Add(new Filter("instance-state-name", new List<string> { InstanceStateName.Stopped.Value }));
+				describeRequest.Filters.Add(new Filter("tag:" + AwsFleetManager.PoolTagName, new List<string> { pool.Name }));
+				describeResponse = await _client.DescribeInstancesAsync(describeRequest);
+				describeScope.Span.SetTag("res.statusCode", (int)describeResponse.HttpStatusCode);
+				describeScope.Span.SetTag("res.numReservations", describeResponse.Reservations.Count);
+			}
+
+			using (IScope startScope = GlobalTracer.Instance.BuildSpan("StartInstances").StartActive())
+			{
+				// Try to start the given instances
+				StartInstancesRequest startRequest = new StartInstancesRequest();
+				startRequest.InstanceIds.AddRange(describeResponse.Reservations.SelectMany(x => x.Instances).Select(x => x.InstanceId).Take(count));
+				
+				startScope.Span.SetTag("req.instanceIds", String.Join(",", startRequest.InstanceIds));
+				if (startRequest.InstanceIds.Count > 0)
+				{
+					StartInstancesResponse startResponse = await _client.StartInstancesAsync(startRequest);
+					startScope.Span.SetTag("res.statusCode", (int)startResponse.HttpStatusCode);
+					startScope.Span.SetTag("res.numInstances", startResponse.StartingInstances.Count);
+					if ((int)startResponse.HttpStatusCode >= 200 && (int)startResponse.HttpStatusCode <= 299)
+					{
+						foreach (InstanceStateChange instanceChange in startResponse.StartingInstances)
+						{
+							_logger.LogInformation("Starting instance {InstanceId} for pool {PoolId} (prev state {PrevState}, current state {CurrentState}", instanceChange.InstanceId, pool.Id, instanceChange.PreviousState, instanceChange.CurrentState);
+						}
+					}
+				}
+
+				if (startRequest.InstanceIds.Count < count)
+				{
+					_logger.LogInformation("Unable to expand pool {PoolName} with the requested number of instances. " +
+					                      "Num requested instances to add {RequestedCount}. Actual instances started {InstancesStarted}", pool.Name, count, startRequest.InstanceIds.Count);
+				}
+			}
+		}
+
+		/// <inheritdoc/>
+		public Task ShrinkPoolAsync(IPool pool, IReadOnlyList<IAgent> agents, int count, CancellationToken cancellationToken)
+		{
+			return AwsFleetManager.ShrinkPoolViaAgentShutdownRequestAsync(_agentCollection, pool, agents, count, cancellationToken);
+		} 
+
+		/// <inheritdoc/>
+		public async Task<int> GetNumStoppedInstancesAsync(IPool pool, CancellationToken cancellationToken)
+		{
+			// Find all instances in the pool
+			DescribeInstancesRequest describeRequest = new DescribeInstancesRequest();
+			describeRequest.Filters = new List<Filter>();
+			describeRequest.Filters.Add(new Filter("instance-state-name", new List<string> { InstanceStateName.Stopped.Value }));
+			describeRequest.Filters.Add(new Filter("tag:" + AwsFleetManager.PoolTagName, new List<string> { pool.Name }));
+
+			DescribeInstancesResponse describeResponse = await _client.DescribeInstancesAsync(describeRequest);
+			return describeResponse.Reservations.SelectMany(x => x.Instances).Select(x => x.InstanceId).Distinct().Count();
+		}
+	}
+}
diff --git a/Engine/Source/Programs/Horde/Horde.Build/Agents/Fleet/Providers/AwsTerminatingFleetManager.cs b/Engine/Source/Programs/Horde/Horde.Build/Agents/Fleet/Providers/AwsTerminatingFleetManager.cs
new file mode 100644
--- /dev/null
+++ b/Engine/Source/Programs/Horde/Horde.Build/Agents/Fleet/Providers/AwsTerminatingFleetManager.cs
@@ -0,0 +1,252 @@
+// Copyright Epic Games, Inc. All Rights Reserved.
+
+using System;
+using System.Collections.Generic;
+using System.Linq;
+using System.Threading;
+using System.Threading.Tasks;
+using Amazon;
+using Amazon.EC2;
+using Amazon.EC2.Model;
+using Horde.Build.Agents.Pools;
+using Horde.Build.Auditing;
+using Microsoft.Extensions.Logging;
+using OpenTracing;
+using OpenTracing.Util;
+
+namespace Horde.Build.Agents.Fleet.Providers
+{
+	/// <summary>
+	/// Settings for launching an EC2 instance
+	/// </summary>
+	public class AwsInstanceLaunchSettings
+	{
+		/// <summary>
+		/// AWS region (e.g us-east-1)
+		/// </summary>
+		public string Region { get; }
+		
+		/// <summary>
+		/// Amazon Machine Image (AMI) to use (e.g ami-0c2e245e4936a2ab2)
+		/// </summary>
+		public string ImageId { get; }
+		
+		/// <summary>
+		/// An EC2 instance type (see https://docs.aws.amazon.com/AWSEC2/latest/UserGuide/instance-types.html for full list)
+		/// </summary>
+		public string InstanceType { get; }
+		
+		/// <summary>
+		/// IAM	instance profile 
+		/// </summary>
+		public string? IamInstanceProfile { get; }
+		
+		/// <summary>
+		/// Key/value tags for the instance
+		/// </summary>
+		public IReadOnlyDictionary<string, string> Tags { get; }
+		
+		/// <summary>
+		/// List of subnet IDs this instance can run in. Will implicitly pick the VPC.
+		/// A random subnet ID will be picked during instance creation.
+		/// </summary>
+		public List<string> SubnetIds { get; }
+		
+		/// <summary>
+		/// Security group IDs to apply. Must exist in the same VPC as the subnets.
+		/// </summary>
+		public List<string> SecurityGroupIds { get; }
+		
+		/// <summary>
+		/// User data to pass to the instance. Commonly used for cloud-init scripts that configures the instance during boot up. 
+		/// </summary>
+		public string? UserData { get; }
+		
+		/// <summary>
+		/// Key name for the instance (optional).
+		/// If skipped, alternative ways to administer the instance must be enabled, such as the AWS SSM Agent. 
+		/// </summary>
+		public string? KeyName { get; }
+
+		/// <summary>
+		/// Constructor
+		/// </summary>
+		/// <param name="region"></param>
+		/// <param name="imageId"></param>
+		/// <param name="instanceType"></param>
+		/// <param name="iamInstanceProfile"></param>
+		/// <param name="tags"></param>
+		/// <param name="subnetIds"></param>
+		/// <param name="securityGroupIds"></param>
+		/// <param name="userData"></param>
+		/// <param name="keyName"></param>
+		public AwsInstanceLaunchSettings(string region, string imageId, string instanceType, string? iamInstanceProfile, IReadOnlyDictionary<string, string> tags,
+			List<string> subnetIds, List<string> securityGroupIds, string? userData, string? keyName)
+		{
+			Region = region;
+			ImageId = imageId;
+			InstanceType = instanceType;
+			IamInstanceProfile = iamInstanceProfile;
+			Tags = tags;
+			SubnetIds = subnetIds;
+			SecurityGroupIds = securityGroupIds;
+			UserData = userData;
+			KeyName = keyName;
+		}
+
+		/// <summary>
+		/// Get a random subnet ID. Used for 
+		/// </summary>
+		/// <returns>A random subnet ID from the list of subnets</returns>
+		public string GetRandomSubnetId()
+		{
+			int index = new Random().Next(SubnetIds.Count);
+			return SubnetIds[index];
+		}
+	}
+	
+	/// <summary>
+	/// Fleet manager for handling AWS EC2 instances
+	/// Will create and/or terminate instances from scratch.
+	/// </summary>
+	public sealed class AwsFleetManager : IFleetManager, IDisposable
+	{
+		/// <summary>
+		/// Prefix for signaling that an EC2 instance belong to an agent pool
+		/// </summary>
+		public const string PoolTagName = "Horde_Autoscale_Pool";
+		private const string AwsTagPropertyName = "aws-tag";
+
+		private readonly AmazonEC2Client _client;
+		private readonly IAgentCollection _agentCollection;
+		private readonly AwsInstanceLaunchSettings _launchSettings;
+		private readonly ILogger _logger;
+
+		/// <summary>
+		/// Constructor
+		/// </summary>
+		public AwsFleetManager(IAgentCollection agentCollection, AwsInstanceLaunchSettings launchSettings, ILogger<AwsFleetManager> logger)
+		{
+			_agentCollection = agentCollection;
+			_launchSettings = launchSettings;
+			_logger = logger;
+
+			AmazonEC2Config config = new () { RegionEndpoint = RegionEndpoint.GetBySystemName(_launchSettings.Region) };
+			logger.LogInformation("Initializing AWS fleet manager for region {Region}", config.RegionEndpoint);
+			_client = new AmazonEC2Client(config);
+		}
+
+		/// <inheritdoc/>
+		public void Dispose()
+		{
+			_client.Dispose();
+		}
+
+		/// <inheritdoc/>
+		public async Task ExpandPoolAsync(IPool pool, IReadOnlyList<IAgent> agents, int count, CancellationToken cancellationToken)
+		{
+			using IScope scope = GlobalTracer.Instance.BuildSpan("ExpandPool").StartActive();
+			scope.Span.SetTag("poolName", pool.Name);
+			scope.Span.SetTag("numAgents", agents.Count);
+			scope.Span.SetTag("count", count);
+
+			List<Tag> tags = _launchSettings.Tags.Select(x => new Tag(x.Key, x.Value)).ToList();
+			
+			RunInstancesRequest request = new ()
+			{
+				ImageId = _launchSettings.ImageId,
+				InstanceType = InstanceType.FindValue(_launchSettings.InstanceType),
+				KeyName = _launchSettings.KeyName,
+				SecurityGroupIds = _launchSettings.SecurityGroupIds,
+				SubnetId = _launchSettings.GetRandomSubnetId(), // Pick randomly to spread evenly among available subnets 
+				MinCount = count,
+				MaxCount = count,
+				UserData = _launchSettings.UserData
+			};
+
+			if (_launchSettings.IamInstanceProfile != null)
+			{
+				request.IamInstanceProfile = new IamInstanceProfileSpecification { Name = _launchSettings.IamInstanceProfile };
+			}
+			
+			if (tags.Count > 0)
+			{
+				request.TagSpecifications = new() { new() { ResourceType = ResourceType.Instance, Tags = tags } };
+			}
+
+			RunInstancesResponse response = await _client.RunInstancesAsync(request, cancellationToken);
+			int numStartedInstances = response.Reservation.Instances.Count;
+			scope.Span.SetTag("res.statusCode", (int)response.HttpStatusCode);
+			scope.Span.SetTag("res.numInstances", numStartedInstances);
+
+			foreach (Instance i in response.Reservation.Instances)
+			{
+				_logger.LogInformation("Created instance {InstanceId} for pool {PoolId}", i.InstanceId, pool.Id);
+			}
+
+			if (numStartedInstances != count)
+			{
+				_logger.LogWarning("Unable to create all the requested instances for pool {PoolId}. RequestedCount={RequestedCount} ActualCount={ActualCount}", pool.Id, count, numStartedInstances);
+			}
+		}
+
+		/// <inheritdoc/>
+		public Task ShrinkPoolAsync(IPool pool, IReadOnlyList<IAgent> agents, int count, CancellationToken cancellationToken)
+		{
+			return ShrinkPoolViaAgentShutdownRequestAsync(_agentCollection, pool, agents, count, cancellationToken);
+		}
+		
+		/// <summary>
+		/// Shrink pool by gracefully requesting a shutdown for the agent
+		/// Allowing the agent to terminate at next best possible moment (usually when all the leases have finished).
+		/// </summary>
+		/// <param name="agentCollection">Instance of the AgentCollection</param>
+		/// <param name="pool">Pool to resize</param>
+		/// <param name="agents">Current list of agents in the pool</param>
+		/// <param name="count">Number of agents to remove</param>		/// <param name="cancellationToken">Cancellation token for the call</param>
+		public static async Task ShrinkPoolViaAgentShutdownRequestAsync(IAgentCollection agentCollection, IPool pool, IReadOnlyList<IAgent> agents, int count, CancellationToken cancellationToken)
+		{
+			using IScope scope = GlobalTracer.Instance.BuildSpan("ShrinkPool").StartActive();
+			scope.Span.SetTag("poolName", pool.Name);
+			scope.Span.SetTag("count", count);
+			
+			string awsTagProperty = $"{AwsTagPropertyName}={PoolTagName}:{pool.Name}";
+			
+			// Sort the agents by number of active leases. It's better to shutdown agents currently doing nothing.
+			List<IAgent> filteredAgents = agents.OrderBy(x => x.Leases.Count).ToList();
+			List<IAgent> agentsWithAwsTags = filteredAgents.Where(x => x.HasProperty(awsTagProperty)).ToList(); 
+			List<IAgent> agentsLimitedByCount = agentsWithAwsTags.Take(count).ToList();
+			
+			scope.Span.SetTag("agents.num", agents.Count);
+			scope.Span.SetTag("agents.filtered.num", filteredAgents.Count);
+			scope.Span.SetTag("agents.withAwsTags.num", agentsWithAwsTags.Count);
+			scope.Span.SetTag("agents.limitedByCount.num", agentsLimitedByCount.Count);
+
+			foreach (IAgent agent in agentsLimitedByCount)
+			{
+				IAuditLogChannel<AgentId> agentLogger = agentCollection.GetLogger(agent.Id);
+				if (await agentCollection.TryUpdateSettingsAsync(agent, bRequestShutdown: true, shutdownReason: "Autoscaler") != null)
+				{
+					agentLogger.LogInformation("Marked {AgentId} in pool {PoolName} for shutdown due to autoscaling (currently {NumLeases} leases outstanding)", agent.Id, pool.Name, agent.Leases.Count);
+				}
+				else
+				{
+					agentLogger.LogError("Unable to mark agent {AgentId} in pool {PoolName} for shutdown due to autoscaling", agent.Id, pool.Name);
+				}
+			}
+		}
+
+		/// <inheritdoc/>
+		public async Task<int> GetNumStoppedInstancesAsync(IPool pool, CancellationToken cancellationToken)
+		{
+			// Find all instances in the pool
+			DescribeInstancesRequest describeRequest = new ();
+			describeRequest.Filters = new List<Filter>();
+			describeRequest.Filters.Add(new Filter("instance-state-name", new List<string> { InstanceStateName.Stopped.Value }));
+			describeRequest.Filters.Add(new Filter("tag:" + PoolTagName, new List<string> { pool.Name }));
+
+			DescribeInstancesResponse describeResponse = await _client.DescribeInstancesAsync(describeRequest);
+			return describeResponse.Reservations.SelectMany(x => x.Instances).Select(x => x.InstanceId).Distinct().Count();
+		}
+	}
+}
diff --git a/Engine/Source/Programs/Horde/Horde.Build/Agents/Fleet/Providers/DefaultFleetManager.cs b/Engine/Source/Programs/Horde/Horde.Build/Agents/Fleet/Providers/DefaultFleetManager.cs
--- a/Engine/Source/Programs/Horde/Horde.Build/Agents/Fleet/Providers/DefaultFleetManager.cs
+++ b/Engine/Source/Programs/Horde/Horde.Build/Agents/Fleet/Providers/DefaultFleetManager.cs
@@ -1,6 +1,7 @@
 // Copyright Epic Games, Inc. All Rights Reserved.
 
 using System.Collections.Generic;
+using System.Threading;
 using System.Threading.Tasks;
 using Horde.Build.Agents.Pools;
 using Microsoft.Extensions.Logging;
@@ -24,20 +25,20 @@ namespace Horde.Build.Agents.Fleet.Providers
 		}
 
 		/// <inheritdoc/>
-		public Task ExpandPoolAsync(IPool pool, IReadOnlyList<IAgent> agents, int count)
+		public Task ExpandPoolAsync(IPool pool, IReadOnlyList<IAgent> agents, int count, CancellationToken cancellationToken)
 		{
 			_logger.LogInformation("Expand pool {PoolId} by {Count} agents", pool.Id, count);
 			return Task.CompletedTask;
 		}
 
 		/// <inheritdoc/>
-		public Task ShrinkPoolAsync(IPool pool, IReadOnlyList<IAgent> agents, int count)
+		public Task ShrinkPoolAsync(IPool pool, IReadOnlyList<IAgent> agents, int count, CancellationToken cancellationToken)
 		{
 			_logger.LogInformation("Shrink pool {PoolId} by {Count} agents", pool.Id, count);
 			return Task.CompletedTask;
 		}
 
 		/// <inheritdoc/>
-		public Task<int> GetNumStoppedInstancesAsync(IPool pool) => Task.FromResult(0);
+		public Task<int> GetNumStoppedInstancesAsync(IPool pool, CancellationToken cancellationToken) => Task.FromResult(0);
 	}
 }
diff --git a/Engine/Source/Programs/Horde/Horde.Build/Startup.cs b/Engine/Source/Programs/Horde/Horde.Build/Startup.cs
--- a/Engine/Source/Programs/Horde/Horde.Build/Startup.cs
+++ b/Engine/Source/Programs/Horde/Horde.Build/Startup.cs
@@ -363,7 +363,7 @@ namespace Horde.Build
 			switch (settings.FleetManager)
 			{
 				case FleetManagerType.Aws:
-					services.AddSingleton<IFleetManager, AwsFleetManager>();
+					services.AddSingleton<IFleetManager, AwsReuseFleetManager>();
 					break;
 				default:
 					services.AddSingleton<IFleetManager, DefaultFleetManager>();
