Horde: Terminate non-reusable EC2 instances on shutdown

diff --git a/Engine/Source/Programs/Horde/Horde.Build/Agents/Fleet/Providers/AwsTerminatingFleetManager.cs b/Engine/Source/Programs/Horde/Horde.Build/Agents/Fleet/Providers/AwsTerminatingFleetManager.cs
--- a/Engine/Source/Programs/Horde/Horde.Build/Agents/Fleet/Providers/AwsTerminatingFleetManager.cs
+++ b/Engine/Source/Programs/Horde/Horde.Build/Agents/Fleet/Providers/AwsTerminatingFleetManager.cs
@@ -161,7 +161,8 @@ namespace Horde.Build.Agents.Fleet.Providers
 				SubnetId = _launchSettings.GetRandomSubnetId(), // Pick randomly to spread evenly among available subnets 
 				MinCount = count,
 				MaxCount = count,
-				UserData = _launchSettings.UserData
+				UserData = _launchSettings.UserData,
+				InstanceInitiatedShutdownBehavior = ShutdownBehavior.Terminate
 			};
 
 			if (_launchSettings.IamInstanceProfile != null)
