Horde: Include any additional agent properties from settings

diff --git a/Engine/Source/Programs/Horde/Horde.Agent/AgentSettings.cs b/Engine/Source/Programs/Horde/Horde.Agent/AgentSettings.cs
--- a/Engine/Source/Programs/Horde/Horde.Agent/AgentSettings.cs
+++ b/Engine/Source/Programs/Horde/Horde.Agent/AgentSettings.cs
@@ -207,6 +207,11 @@ namespace Horde.Agent
 		/// </summary>
 		public bool WriteStepOutputToLogger { get; set; }
 
+		/// <summary>
+		/// Key/value properties in addition to those set internally by the agent
+		/// </summary>
+		public Dictionary<string, string> Properties { get; } = new();
+
 		/// <summary>
 		/// Gets the current server settings
 		/// </summary>
diff --git a/Engine/Source/Programs/Horde/Horde.Agent/Services/WorkerService.cs b/Engine/Source/Programs/Horde/Horde.Agent/Services/WorkerService.cs
--- a/Engine/Source/Programs/Horde/Horde.Agent/Services/WorkerService.cs
+++ b/Engine/Source/Programs/Horde/Horde.Agent/Services/WorkerService.cs
@@ -473,7 +473,7 @@ namespace Horde.Agent.Services
 			TerminateProcesses(_logger, stoppingToken);
 
 			// Show the worker capabilities
-			AgentCapabilities capabilities = await GetAgentCapabilities(_workingDir, _logger);
+			AgentCapabilities capabilities = await GetAgentCapabilities(_workingDir, _logger, _settings.Properties);
 			if (capabilities.Properties.Count > 0)
 			{
 				_logger.LogInformation("Global:");
@@ -588,7 +588,7 @@ namespace Horde.Agent.Services
 					{
 						try
 						{
-							updateSessionRequest.Capabilities = await GetAgentCapabilities(_workingDir, _logger);
+							updateSessionRequest.Capabilities = await GetAgentCapabilities(_workingDir, _logger, _settings.Properties);
 						}
 						catch (Exception ex)
 						{
@@ -1540,7 +1540,7 @@ namespace Horde.Agent.Services
 		/// Gets the hardware capabilities of this worker
 		/// </summary>
 		/// <returns>Worker object for advertising to the server</returns>
-		public static async Task<AgentCapabilities> GetAgentCapabilities(DirectoryReference workingDir, ILogger logger)
+		public static async Task<AgentCapabilities> GetAgentCapabilities(DirectoryReference workingDir, ILogger logger, Dictionary<string, string>? extraProperties = null)
 		{
 			// Create the primary device
 			DeviceCapabilities primaryDevice = new DeviceCapabilities();
@@ -1876,9 +1876,15 @@ namespace Horde.Agent.Services
 			}
 
 			// Create the worker
-			AgentCapabilities agent = new AgentCapabilities();
+			AgentCapabilities agent = new ();
 			agent.Devices.Add(primaryDevice);
 			agent.Devices.AddRange(otherDevices);
+			
+			if (extraProperties != null)
+			{
+				agent.Properties.AddRange(extraProperties.Select(kvp => $"{kvp.Key}={kvp.Value}"));				
+			}
+
 			return agent;
 		}
 
