Horde: Fix missing remote IP log property

This somehow stopped working and I was unable to get it working with the inline Use(...).

diff --git a/Engine/Source/Programs/Horde/Horde.Build/Startup.cs b/Engine/Source/Programs/Horde/Horde.Build/Startup.cs
--- a/Engine/Source/Programs/Horde/Horde.Build/Startup.cs
+++ b/Engine/Source/Programs/Horde/Horde.Build/Startup.cs
@@ -905,14 +905,13 @@ namespace Horde.Build
 			app.UseSwagger();
 
 			// Enable serilog request logging
-			app.UseSerilogRequestLogging(options => options.GetLevel = GetRequestLoggingLevel);
-
-			// Include the source IP address with requests
-			app.Use(async (context, next) => {
-				using (Serilog.Context.LogContext.PushProperty("RemoteIP", context.Connection.RemoteIpAddress))
+			app.UseSerilogRequestLogging(options =>
+			{
+				options.GetLevel = GetRequestLoggingLevel;
+				options.EnrichDiagnosticContext = (diagnosticContext, httpContext) =>
 				{
-					await next();
-				}
+					diagnosticContext.Set("RemoteIP", httpContext.Connection.RemoteIpAddress);
+				};
 			});
 
 			// Enable middleware to serve swagger-ui (HTML, JS, CSS, etc.),
