Horde: Add more tracing scopes to popular streams endpoint

diff --git a/Engine/Source/Programs/Horde/Horde.Build/Streams/StreamsController.cs b/Engine/Source/Programs/Horde/Horde.Build/Streams/StreamsController.cs
--- a/Engine/Source/Programs/Horde/Horde.Build/Streams/StreamsController.cs
+++ b/Engine/Source/Programs/Horde/Horde.Build/Streams/StreamsController.cs
@@ -13,6 +13,8 @@ using Horde.Build.Users;
 using Horde.Build.Utilities;
 using Microsoft.AspNetCore.Authorization;
 using Microsoft.AspNetCore.Mvc;
+using OpenTracing;
+using OpenTracing.Util;
 
 namespace Horde.Build.Streams
 {
@@ -109,11 +111,17 @@ namespace Horde.Build.Streams
 		/// <returns>Response object</returns>
 		async Task<GetStreamResponse> CreateGetStreamResponse(IStream stream, ProjectPermissionsCache cache)
 		{
+			using IScope scope = GlobalTracer.Instance.BuildSpan("CreateGetStreamResponse").StartActive();
+			scope.Span.SetTag("streamId", stream.Id);
+			
 			bool bIncludeAcl = stream.Acl != null && await _streamService.AuthorizeAsync(stream, AclAction.ViewPermissions, User, cache);
 
 			List<GetTemplateRefResponse> apiTemplateRefs = new List<GetTemplateRefResponse>();
 			foreach (KeyValuePair<TemplateRefId, TemplateRef> pair in stream.Templates)
 			{
+				using IScope templateScope = GlobalTracer.Instance.BuildSpan("CreateGetStreamResponse.Template").StartActive();
+				templateScope.Span.SetTag("templateName", pair.Value.Name);
+				
 				if (await _streamService.AuthorizeAsync(stream, pair.Value, AclAction.ViewTemplate, User, cache))
 				{
 					ITemplate? template = await _templateCollection.GetAsync(pair.Value.Hash);
