Horde: Auto-join agent to pools requested in properties

diff --git a/Engine/Source/Programs/Horde/Horde.Build.Tests/RpcServiceTest.cs b/Engine/Source/Programs/Horde/Horde.Build.Tests/RpcServiceTest.cs
--- a/Engine/Source/Programs/Horde/Horde.Build.Tests/RpcServiceTest.cs
+++ b/Engine/Source/Programs/Horde/Horde.Build.Tests/RpcServiceTest.cs
@@ -3,6 +3,7 @@
 using System;
 using System.Collections.Generic;
 using System.IO;
+using System.Linq;
 using System.Reflection;
 using System.Security.Claims;
 using System.Threading;
@@ -306,6 +307,17 @@ namespace Horde.Build.Tests
 			Assert.AreEqual("MYNAME", res.AgentId);
 			// TODO: Check Token, ExpiryTime, SessionId 
 		}
+		
+		[TestMethod]
+		public async Task AgentJoinsPoolThroughProperties()
+		{
+			CreateSessionRequest req = new () { Name = "bogusAgentName", Capabilities = new AgentCapabilities() };
+			req.Capabilities.Properties.Add($"{KnownPropertyNames.RequestedPools}=fooPool,barPool");
+			CreateSessionResponse res = await RpcService.CreateSession(req, _adminContext);
+
+			IAgent agent = (await AgentService.GetAgentAsync(new AgentId(res.AgentId)))!;
+			CollectionAssert.AreEquivalent(new List<PoolId> { new ("fooPool"), new ("barPool") }, agent.GetPools().ToList());
+		}
 
 		[TestMethod]
 		public async Task UpdateSessionTest()
diff --git a/Engine/Source/Programs/Horde/Horde.Build/Agents/IAgent.cs b/Engine/Source/Programs/Horde/Horde.Build/Agents/IAgent.cs
--- a/Engine/Source/Programs/Horde/Horde.Build/Agents/IAgent.cs
+++ b/Engine/Source/Programs/Horde/Horde.Build/Agents/IAgent.cs
@@ -333,6 +333,11 @@ namespace Horde.Build.Agents
 		/// </summary>
 		public const string Pool = "Pool";
 
+		/// <summary>
+		/// Pools requested by the agent to join when registering with server
+		/// </summary>
+		public const string RequestedPools = "RequestedPools";
+		
 		/// <summary>
 		/// Number of logical cores
 		/// </summary>
diff --git a/Engine/Source/Programs/Horde/Horde.Build/Server/RpcService.cs b/Engine/Source/Programs/Horde/Horde.Build/Server/RpcService.cs
--- a/Engine/Source/Programs/Horde/Horde.Build/Server/RpcService.cs
+++ b/Engine/Source/Programs/Horde/Horde.Build/Server/RpcService.cs
@@ -23,7 +23,6 @@ using Horde.Build.Jobs.Graphs;
 using Horde.Build.Jobs.Templates;
 using Horde.Build.Jobs.TestData;
 using Horde.Build.Logs;
-using Horde.Build.Server;
 using Horde.Build.Streams;
 using Horde.Build.Tasks;
 using Horde.Build.Utilities;
@@ -284,8 +283,7 @@ namespace Horde.Build.Server
 					throw new StructuredRpcException(StatusCode.PermissionDenied, "User is not authenticated to create new agents");
 				}
 
-				const bool BEnabled = true;
-				agent = await _agentService.CreateAgentAsync(request.Name, BEnabled, null, null);
+				agent = await _agentService.CreateAgentAsync(request.Name, true, null, GetRequestedPoolsFromCapabilities(request.Capabilities));
 			}
 
 			// Make sure we're allowed to create sessions on this agent
@@ -313,6 +311,20 @@ namespace Horde.Build.Server
 			return response;
 		}
 
+		private static List<StringId<IPool>>? GetRequestedPoolsFromCapabilities(RpcAgentCapabilities capabilities)
+		{
+			foreach (string property in capabilities.Properties)
+			{
+				const string Key = "RequestedPools=";
+				if (property.StartsWith(Key, StringComparison.InvariantCulture))
+				{
+					return property[Key.Length..].Split(",").Select(x => new StringId<IPool>(x)).ToList();
+				}
+			}
+
+			return null;
+		}
+
 		/// <summary>
 		/// Updates an agent session
 		/// </summary>
