Horde Storage - Disable forwarding of AWS SDK logs into serilog. The AWS log is quite useless producing multiple errors for operations that are not errors, and thus also spammy. It can still be useful to debug certain complicated AWS SDK issues so leaving the code in place should we need it.

diff --git a/Engine/Source/Programs/Horde/HordeStorage/Jupiter.Common/BaseProgram.cs b/Engine/Source/Programs/Horde/HordeStorage/Jupiter.Common/BaseProgram.cs
--- a/Engine/Source/Programs/Horde/HordeStorage/Jupiter.Common/BaseProgram.cs
+++ b/Engine/Source/Programs/Horde/HordeStorage/Jupiter.Common/BaseProgram.cs
@@ -75,11 +75,13 @@ namespace Jupiter
                     webBuilder.UseConfiguration(Configuration);
                     webBuilder.UseSerilog();
                     // configure microsoft.extensions.logging to configure log4net to allow us to set it in our appsettings
-                    webBuilder.ConfigureLogging((hostingContext, logging) =>
+                    // Disabled forwarding of log4net logs into serilog, as the AWS sdk is very spammy with its output producing multiple errors for a 404 (which isn't even an error in the first place)
+                    // This can be enabled if you need to investigate some more complicated AWS sdk issue
+                    /*webBuilder.ConfigureLogging((hostingContext, logging) =>
                     {
                         // configure log4net (used by aws sdk) to write to serilog so we get the logs in the system we want it in
                         Log4net.Appender.Serilog.Configuration.Configure();
-                    });
+                    });*/
                     // remove the server header from kestrel
                     webBuilder.ConfigureKestrel(options =>
                     {
