Horde: Upgrade Datadog tracing library

diff --git a/Engine/Source/Programs/Horde/Horde.Build/Horde.Build.csproj b/Engine/Source/Programs/Horde/Horde.Build/Horde.Build.csproj
--- a/Engine/Source/Programs/Horde/Horde.Build/Horde.Build.csproj
+++ b/Engine/Source/Programs/Horde/Horde.Build/Horde.Build.csproj
@@ -23,7 +23,7 @@
     <PackageReference Include="AWSSDK.Extensions.NETCore.Setup" Version="3.7.0.1" />
     <PackageReference Include="AWSSDK.S3" Version="3.7.0.17" />
     <PackageReference Include="AWSSDK.SecurityToken" Version="3.7.1.69" />
-    <PackageReference Include="Datadog.Trace.OpenTracing" Version="1.18.0" />
+    <PackageReference Include="Datadog.Trace.OpenTracing" Version="2.11.0" />
     <PackageReference Include="DogStatsD-CSharp-Client" Version="6.0.0" />
     <PackageReference Include="Grpc.AspNetCore" Version="2.45.0" />
     <PackageReference Include="Grpc.AspNetCore.Server.Reflection" Version="2.45.0" />
