Fix issue with missing test report in job details. Bump up the limit of test report loaded in jobdetails

diff --git a/Engine/Source/Programs/Horde/HordeDashboard/src/backend/Backend.ts b/Engine/Source/Programs/Horde/HordeDashboard/src/backend/Backend.ts
--- a/Engine/Source/Programs/Horde/HordeDashboard/src/backend/Backend.ts
+++ b/Engine/Source/Programs/Horde/HordeDashboard/src/backend/Backend.ts
@@ -1136,7 +1136,7 @@ export class Backend {
             const params: any = {
                 jobId: jobId,
                 filter: 'id,key,change,jobId,stepId,streamId',
-                count: 200,
+                count: 500,
             };
 
             if (stepId) {
