Horde Storage - Ignore case when looking for claim values that do not contain a = (we already do for those containing = )

This resolves issues with the auth for legacy transaction log endpoints using legacy claims

diff --git a/Engine/Source/Programs/Horde/HordeStorage/Jupiter.Common/Authentication/Acl.cs b/Engine/Source/Programs/Horde/HordeStorage/Jupiter.Common/Authentication/Acl.cs
--- a/Engine/Source/Programs/Horde/HordeStorage/Jupiter.Common/Authentication/Acl.cs
+++ b/Engine/Source/Programs/Horde/HordeStorage/Jupiter.Common/Authentication/Acl.cs
@@ -85,7 +85,7 @@ public class AclEntry
                     claimFound = true;
                 }
             }
-            else if (user.HasClaim(claim => claim.Type == expectedClaim))
+            else if (user.HasClaim(claim => string.Equals(claim.Type, expectedClaim, StringComparison.OrdinalIgnoreCase)))
             {
                 claimFound = true;
             }
