Horde Storage - Reverted enabling of PII everywhere, this should never be set in prod.

diff --git a/Engine/Source/Programs/Horde/HordeStorage/Jupiter.Common/BaseStartup.cs b/Engine/Source/Programs/Horde/HordeStorage/Jupiter.Common/BaseStartup.cs
--- a/Engine/Source/Programs/Horde/HordeStorage/Jupiter.Common/BaseStartup.cs
+++ b/Engine/Source/Programs/Horde/HordeStorage/Jupiter.Common/BaseStartup.cs
@@ -599,7 +599,7 @@ namespace Jupiter
         public long MemoryBufferSize { get; set; } = int.MaxValue;
 
         // enable to unhide potentially personal information, see https://aka.ms/IdentityModel/PII
-        public bool ShowPII { get; set; } = true;
+        public bool ShowPII { get; set; } = false;
         public bool DisableHealthChecks { get; set; } = false;
         public bool HostSwaggerDocumentation { get; set; } = true;
 
