Horde Storage - Fixed incorrect claim mapping for legacy transaction log access, also made more of the claim verification case insensitive. Lastly I updated the scylla keyspace mapping to what its actually set to in the db (no functional difference)

diff --git a/Engine/Source/Programs/Horde/HordeStorage/Jupiter.Common/Authentication/Acl.cs b/Engine/Source/Programs/Horde/HordeStorage/Jupiter.Common/Authentication/Acl.cs
--- a/Engine/Source/Programs/Horde/HordeStorage/Jupiter.Common/Authentication/Acl.cs
+++ b/Engine/Source/Programs/Horde/HordeStorage/Jupiter.Common/Authentication/Acl.cs
@@ -80,7 +80,7 @@ public class AclEntry
                 int separatorIndex = expectedClaim.IndexOf('=', StringComparison.InvariantCultureIgnoreCase);
                 string claimName = expectedClaim.Substring(0, separatorIndex);
                 string claimValue = expectedClaim.Substring(separatorIndex + 1);
-                if (user.HasClaim(claim => claim.Type == claimName && claim.Value == claimValue))
+                if (user.HasClaim(claim => string.Equals(claim.Type, claimName, StringComparison.OrdinalIgnoreCase) && string.Equals(claim.Value, claimValue, StringComparison.OrdinalIgnoreCase)))
                 {
                     claimFound = true;
                 }
