Horde: Don't update sentinels for quarantined issues
#rnx

diff --git a/Engine/Source/Programs/Horde/Horde.Build.Tests/IssueServiceTests.cs b/Engine/Source/Programs/Horde/Horde.Build.Tests/IssueServiceTests.cs
--- a/Engine/Source/Programs/Horde/Horde.Build.Tests/IssueServiceTests.cs
+++ b/Engine/Source/Programs/Horde/Horde.Build.Tests/IssueServiceTests.cs
@@ -1704,6 +1704,88 @@ namespace Horde.Build.Tests
 			}
 		}
 
+		[TestMethod]
+		public async Task QuarantineTest()
+		{
+			int issueId;
+			DateTime lastSeenAt;
+
+			// #1
+			// Scenario: Warning in first step
+			// Expected: Default issue is created
+			{
+				IJob job = CreateJob(_mainStreamId, 105, "Test Build", _graph);
+				await AddEvent(job, 0, 0, new { level = nameof(LogLevel.Warning) }, EventSeverity.Warning);
+				await UpdateCompleteStep(job, 0, 0, JobStepOutcome.Warnings);
+
+				List<IIssue> issues = await IssueCollection.FindIssuesAsync();
+				Assert.AreEqual(1, issues.Count);
+				Assert.AreEqual(IssueSeverity.Warning, issues[0].Severity);
+
+				Assert.AreEqual("Warnings in Update Version Files", issues[0].Summary);
+
+				issueId = issues[0].Id;
+				lastSeenAt = issues[0].LastSeenAt;
+			}
+
+			// Mark issue as quarantined
+			await IssueService.UpdateIssueAsync(issueId, quarantinedById: _jerryId);
+
+			// #2
+			// Scenario: Job succeeds
+			// Expected: Issue is not marked resolved
+			{
+				IJob job = CreateJob(_mainStreamId, 115, "Test Build", _graph);
+				await UpdateCompleteStep(job, 0, 0, JobStepOutcome.Success);
+
+				IIssue? issue = await IssueCollection.GetIssueAsync(issueId);
+				Assert.IsNull(issue!.ResolvedAt);
+
+				List<IIssueSpan> spans = await IssueCollection.FindSpansAsync(issue.Id);
+				Assert.AreEqual(spans.Count, 1);
+
+			}
+
+			// #3
+			// Scenario: Job fails
+			// Expected: Existing issue is updated
+			{
+				IJob job = CreateJob(_mainStreamId, 125, "Test Build", _graph);
+				await AddEvent(job, 0, 0, new { level = nameof(LogLevel.Warning) }, EventSeverity.Warning);
+				await UpdateCompleteStep(job, 0, 0, JobStepOutcome.Warnings);
+
+				List<IIssue> issues = await IssueCollection.FindIssuesAsync();
+				Assert.AreEqual(1, issues.Count);
+				Assert.AreEqual(IssueSeverity.Warning, issues[0].Severity);
+
+				Assert.AreEqual("Warnings in Update Version Files", issues[0].Summary);
+
+				Assert.AreEqual(issueId, issues[0].Id);
+				Assert.AreNotEqual(lastSeenAt, issues[0].LastSeenAt);
+			}
+
+			// Mark issue as not quarantined
+			await IssueService.UpdateIssueAsync(issueId, quarantinedById: UserId.Empty);
+
+			// #4
+			// Scenario: Job succeeds
+			// Expected: Issue is marked resolved and closed
+			{
+				IJob job = CreateJob(_mainStreamId, 135, "Test Build", _graph);
+				await UpdateCompleteStep(job, 0, 0, JobStepOutcome.Success);
+
+				IIssue? issue = await IssueCollection.GetIssueAsync(issueId);
+				Assert.IsNotNull(issue!.ResolvedAt);
+
+				List<IIssue> issues = await IssueCollection.FindIssuesAsync();
+				Assert.AreEqual(0, issues.Count);
+												
+			}
+
+
+		}
+
+
 		private async Task ParseAsync(LogId logId, string[] lines)
 		{
 			await using (TestJsonLogger logger = new TestJsonLogger(LogFileService, logId))
diff --git a/Engine/Source/Programs/Horde/Horde.Build/Issues/IssueService.cs b/Engine/Source/Programs/Horde/Horde.Build/Issues/IssueService.cs
--- a/Engine/Source/Programs/Horde/Horde.Build/Issues/IssueService.cs
+++ b/Engine/Source/Programs/Horde/Horde.Build/Issues/IssueService.cs
@@ -560,8 +560,13 @@ namespace Horde.Build.Issues
 						}
 					}
 
+					// Only update sentinels for issues which aren't quarantined
+					List<IIssue> quarantined = await _issueCollection.FindIssuesAsync(openSpans.Select(x => x.IssueId).Distinct());
+					quarantined = quarantined.Where(x => x.QuarantinedByUserId != null).ToList();
+					List<IIssueSpan> sentinels = openSpans.Where(x => !quarantined.Any(y => x.IssueId == y.Id)).ToList();
+
 					// Try to update the sentinels for any other open steps
-					if (!await TryUpdateSentinelsAsync(openSpans, stream, job, batch, step))
+					if (!await TryUpdateSentinelsAsync(sentinels, stream, job, batch, step))
 					{
 						continue;
 					}
