Horde: Fix legacy serialization paths for IMemoryReader/IMemoryWriter causing test failures.

diff --git a/Engine/Source/Programs/Horde/Horde.Build/Logs/Data/LogIndexData.cs b/Engine/Source/Programs/Horde/Horde.Build/Logs/Data/LogIndexData.cs
--- a/Engine/Source/Programs/Horde/Horde.Build/Logs/Data/LogIndexData.cs
+++ b/Engine/Source/Programs/Horde/Horde.Build/Logs/Data/LogIndexData.cs
@@ -547,7 +547,7 @@ namespace Horde.Build.Logs.Data
 			{
 				ReadOnlyTrie index = reader.ReadTrie();
 				int numBlockBits = reader.ReadInt32();
-				LogIndexBlock[] blocks = reader.ReadVariableLengthArray(() => reader.ReadLogIndexBlock());
+				LogIndexBlock[] blocks = reader.ReadVariableLengthArrayWithInt32Length(() => reader.ReadLogIndexBlock());
 				for(int idx = 1; idx < blocks.Length; idx++)
 				{
 					if(blocks[idx].LineIndex == 0)
@@ -573,7 +573,7 @@ namespace Horde.Build.Logs.Data
 			writer.WriteInt32(3);
 			writer.WriteTrie(BuildTrie());
 			writer.WriteInt32(_numBlockBits);
-			writer.WriteVariableLengthArray(_blocks, x => writer.WriteLogIndexBlock(x));
+			writer.WriteVariableLengthArrayWithInt32Length(_blocks, x => writer.WriteLogIndexBlock(x));
 		}
 
 		/// <summary>
diff --git a/Engine/Source/Programs/Shared/EpicGames.Core/MemoryReader.cs b/Engine/Source/Programs/Shared/EpicGames.Core/MemoryReader.cs
--- a/Engine/Source/Programs/Shared/EpicGames.Core/MemoryReader.cs
+++ b/Engine/Source/Programs/Shared/EpicGames.Core/MemoryReader.cs
@@ -234,6 +234,17 @@ namespace EpicGames.Core
 			return ReadFixedLengthArray(reader, length, readItem);
 		}
 
+		/// <summary>
+		/// Reads a variable length array
+		/// </summary>
+		/// <param name="reader">Reader to deserialize from</param>
+		/// <param name="readItem">Delegate to write an individual item</param>
+		public static T[] ReadVariableLengthArrayWithInt32Length<T>(this IMemoryReader reader, Func<T> readItem)
+		{
+			int length = ReadInt32(reader);
+			return ReadFixedLengthArray(reader, length, readItem);
+		}
+
 		/// <summary>
 		/// Writes a fixed length array
 		/// </summary>
diff --git a/Engine/Source/Programs/Shared/EpicGames.Core/MemoryWriter.cs b/Engine/Source/Programs/Shared/EpicGames.Core/MemoryWriter.cs
--- a/Engine/Source/Programs/Shared/EpicGames.Core/MemoryWriter.cs
+++ b/Engine/Source/Programs/Shared/EpicGames.Core/MemoryWriter.cs
@@ -237,6 +237,18 @@ namespace EpicGames.Core
 			WriteFixedLengthArray(writer, list, writeItem);
 		}
 
+		/// <summary>
+		/// Writes a variable length array
+		/// </summary>
+		/// <param name="writer">Writer to serialize to</param>
+		/// <param name="array">The array to write</param>
+		/// <param name="writeItem">Delegate to write an individual item</param>
+		public static void WriteVariableLengthArrayWithInt32Length<T>(this IMemoryWriter writer, T[] array, Action<T> writeItem)
+		{
+			WriteInt32(writer, array.Length);
+			WriteFixedLengthArray(writer, array, writeItem);
+		}
+
 		/// <summary>
 		/// Writes a fixed length array
 		/// </summary>
