Horde Storage - Use new oidc token binaries path

diff --git a/Engine/Source/Programs/Horde/HordeStorage/updateAuthHeaderUser.ps1 b/Engine/Source/Programs/Horde/HordeStorage/updateAuthHeaderUser.ps1
--- a/Engine/Source/Programs/Horde/HordeStorage/updateAuthHeaderUser.ps1
+++ b/Engine/Source/Programs/Horde/HordeStorage/updateAuthHeaderUser.ps1
@@ -1,4 +1,4 @@
 $TempFile = New-TemporaryFile
 
-& "dotnet" "..\..\..\..\..\Engine\Restricted\NotForLicensees\Binaries\DotNET\OidcToken\portable\OidcToken.dll" "--Service=EpicGames-Okta-Jupiter" "--OutFile" $TempFile.FullName 
+& "dotnet" "..\..\..\..\..\Engine\Binaries\DotNET\OidcToken\portable\OidcToken.dll" "--Service=EpicGames-Okta-Jupiter" "--OutFile" $TempFile.FullName 
 $env:curl_auth_header = Get-Content -Raw -Path $TempFile.FullName | ConvertFrom-Json | ForEach-Object { "Authorization: Bearer $($_.Token)" }
\ No newline at end of file
