Horde: Use properties from agent too, in addition to device properties

diff --git a/Engine/Source/Programs/Horde/Horde.Build.Tests/RpcServiceTest.cs b/Engine/Source/Programs/Horde/Horde.Build.Tests/RpcServiceTest.cs
--- a/Engine/Source/Programs/Horde/Horde.Build.Tests/RpcServiceTest.cs
+++ b/Engine/Source/Programs/Horde/Horde.Build.Tests/RpcServiceTest.cs
@@ -26,6 +26,7 @@ using AgentCapabilities = HordeCommon.Rpc.Messages.AgentCapabilities;
 using ISession = Microsoft.AspNetCore.Http.ISession;
 using Horde.Build.Jobs.Artifacts;
 using Horde.Build.Agents;
+using HordeCommon.Rpc.Messages;
 
 namespace Horde.Build.Tests
 {
@@ -318,6 +319,36 @@ namespace Horde.Build.Tests
 			IAgent agent = (await AgentService.GetAgentAsync(new AgentId(res.AgentId)))!;
 			CollectionAssert.AreEquivalent(new List<PoolId> { new ("fooPool"), new ("barPool") }, agent.GetPools().ToList());
 		}
+		
+		[TestMethod]
+		public async Task PropertiesFromAgentCapabilities()
+		{
+			CreateSessionRequest req = new () { Name = "bogusAgentName", Capabilities = new AgentCapabilities() };
+			req.Capabilities.Properties.Add("fooKey=barValue");
+			CreateSessionResponse res = await RpcService.CreateSession(req, _adminContext);
+			IAgent agent = (await AgentService.GetAgentAsync(new AgentId(res.AgentId)))!;
+			Assert.IsTrue(agent.Properties.Contains("fooKey=barValue"));
+		}
+		
+		[TestMethod]
+		public async Task PropertiesFromDeviceCapabilities()
+		{
+			CreateSessionRequest req = new () { Name = "bogusAgentName", Capabilities = new AgentCapabilities() };
+			req.Capabilities.Devices.Add(new DeviceCapabilities { Handle = "someHandle", Properties = { "foo=bar" } });
+			CreateSessionResponse res = await RpcService.CreateSession(req, _adminContext);
+			IAgent agent = (await AgentService.GetAgentAsync(new AgentId(res.AgentId)))!;
+			Assert.IsTrue(agent.Properties.Contains("foo=bar"));
+		}
+		
+		[TestMethod]
+		public async Task KnownPropertiesAreSetAsResources()
+		{
+			CreateSessionRequest req = new () { Name = "bogusAgentName", Capabilities = new AgentCapabilities() };
+			req.Capabilities.Devices.Add(new DeviceCapabilities { Handle = "someHandle", Properties = { $"{KnownPropertyNames.LogicalCores}=10" }});
+			CreateSessionResponse res = await RpcService.CreateSession(req, _adminContext);
+			IAgent agent = (await AgentService.GetAgentAsync(new AgentId(res.AgentId)))!;
+			Assert.AreEqual(10, agent.Resources[KnownPropertyNames.LogicalCores]);
+		}
 
 		[TestMethod]
 		public async Task UpdateSessionTest()
diff --git a/Engine/Source/Programs/Horde/Horde.Build/Server/RpcService.cs b/Engine/Source/Programs/Horde/Horde.Build/Server/RpcService.cs
--- a/Engine/Source/Programs/Horde/Horde.Build/Server/RpcService.cs
+++ b/Engine/Source/Programs/Horde/Horde.Build/Server/RpcService.cs
@@ -246,16 +246,16 @@ namespace Horde.Build.Server
 			properties = new List<string>();
 			resources = new Dictionary<string, int>();
 
-			if (capabilities != null && capabilities.Devices.Count > 0)
-			{
-				RpcDeviceCapabilities device = capabilities.Devices[0];
-				if (device.Properties != null)
-				{
-					properties = new List<string>(device.Properties);
-					CopyPropertyToResource(KnownPropertyNames.LogicalCores, properties, resources);
-					CopyPropertyToResource(KnownPropertyNames.Ram, properties, resources);
-				}
-			}
+			if (capabilities == null) return;
+			properties.AddRange(capabilities.Properties);
+
+			if (capabilities.Devices.Count <= 0) return;
+			RpcDeviceCapabilities device = capabilities.Devices[0];
+			if (device.Properties == null) return;
+
+			properties.AddRange(device.Properties);
+			CopyPropertyToResource(KnownPropertyNames.LogicalCores, properties, resources);
+			CopyPropertyToResource(KnownPropertyNames.Ram, properties, resources);
 		}
 
 		/// <summary>
