Horde Storage - Added some more logging for when we are replicating to know which blobs are actually being replicated. Will use this to try and determine why we are missing blobs in remote regions (could be they are not replicated, or they are replicated but then deleted).

Also tweaked some of the reporting around each ref being replicated (making it easier to find which ref is being replicated)

diff --git a/Engine/Source/Programs/Horde/HordeStorage/Horde.Storage/Implementation/Replication/RefsReplicator.cs b/Engine/Source/Programs/Horde/HordeStorage/Horde.Storage/Implementation/Replication/RefsReplicator.cs
--- a/Engine/Source/Programs/Horde/HordeStorage/Horde.Storage/Implementation/Replication/RefsReplicator.cs
+++ b/Engine/Source/Programs/Horde/HordeStorage/Horde.Storage/Implementation/Replication/RefsReplicator.cs
@@ -414,9 +414,9 @@ namespace Horde.Storage.Implementation
                     return;
                 }
                 using IScope scope = Tracer.Instance.StartActive("replicator.replicate_op_incremental");
-                scope.Span.ResourceName = $"[REDACTED_EMAIL]}[REDACTED_EMAIL]}";
+                scope.Span.ResourceName = $"[REDACTED_EMAIL]}[REDACTED_EMAIL]}";
 
-                _logger.Information("{Name} New transaction to replicate found. New op: {@Op} . Count of running replications: {CurrentReplications}", _name, @event, replicationTasks.Count);
+                _logger.Information("{Name} New transaction to replicate found. Ref: {Namespace} {Bucket} {Key} in {TimeBucket} with id {EventId}. Count of running replications: {CurrentReplications}", _name, @event.Namespace, @event.Bucket, @event.Key, @event.TimeBucket, @event.EventId, replicationTasks.Count);
                 
                 Info.CountOfRunningReplications = replicationTasks.Count;
                 LogReplicationHeartbeat(replicationTasks.Count);
@@ -484,23 +484,25 @@ namespace Horde.Storage.Implementation
             return countOfReplicationsDone;
         }
 
-        private async Task ReplicateOp(NamespaceId ns, BlobIdentifier blob, CancellationToken cancellationToken)
+        private async Task ReplicateOp(NamespaceId ns, BlobIdentifier objectToReplicate, CancellationToken cancellationToken)
         {
             using IScope scope = Tracer.Instance.StartActive("replicator.replicate_op");
-            scope.Span.ResourceName = $"{ns}.{blob}";
+            scope.Span.ResourceName = $"{ns}.{objectToReplicate}";
+
+            _logger.Information("Attempting to replicate object {Blob} in {Namespace}.", objectToReplicate, ns);
 
             // We could potentially do this, but that could be dangerous if missing child references
             // check if this blob exists locally before replicating, if it does we assume we have all of its references already
             //if (await _blobService.Exists(ns, blob))
             //    return currentOffset;
 
-            using HttpRequestMessage referencesRequest = BuildHttpRequest(HttpMethod.Get, new Uri($"api/v1/objects/{ns}/{blob}/references", UriKind.Relative));
+            using HttpRequestMessage referencesRequest = BuildHttpRequest(HttpMethod.Get, new Uri($"api/v1/objects/{ns}/{objectToReplicate}/references", UriKind.Relative));
             HttpResponseMessage referencesResponse = await _httpClient.SendAsync(referencesRequest, cancellationToken);
             string body = await referencesResponse.Content.ReadAsStringAsync(cancellationToken);
 
             if (referencesResponse.StatusCode == HttpStatusCode.BadRequest)
             {
-                _logger.Warning("Failed to resolve references for object {Blob} in {Namespace}. Skipping replication", blob, ns);
+                _logger.Warning("Failed to resolve references for object {Blob} in {Namespace}. Skipping replication", objectToReplicate, ns);
                 return;
             }
             referencesResponse.EnsureSuccessStatusCode();
@@ -508,12 +510,12 @@ namespace Horde.Storage.Implementation
             ResolvedReferencesResult? refs = JsonConvert.DeserializeObject<ResolvedReferencesResult>(body);
             if (refs == null)
             {
-                throw new Exception($"Unable to resolve references for object {blob} in namespace {ns}");
+                throw new Exception($"Unable to resolve references for object {objectToReplicate} in namespace {ns}");
             }
 
             BlobIdentifier[] potentialBlobs = new BlobIdentifier[refs.References.Length + 1];
             Array.Copy(refs.References, potentialBlobs, refs.References.Length);
-            potentialBlobs[^1] = blob;
+            potentialBlobs[^1] = objectToReplicate;
 
             BlobIdentifier[] missingBlobs = await _blobService.FilterOutKnownBlobs(ns, potentialBlobs);
             Task[] blobReplicationTasks = new Task[missingBlobs.Length];
@@ -525,6 +527,8 @@ namespace Horde.Storage.Implementation
                     HttpRequestMessage blobRequest = BuildHttpRequest(HttpMethod.Get, new Uri($"api/v1/blobs/{ns}/{blobToReplicate}", UriKind.Relative));
                     HttpResponseMessage blobResponse = await _httpClient.SendAsync(blobRequest, HttpCompletionOption.ResponseHeadersRead, cancellationToken);
 
+                    _logger.Information("Attempting to replicate blob {Blob} in {Namespace}.", blobToReplicate, ns);
+
                     if (blobResponse.StatusCode == HttpStatusCode.NotFound)
                     {
                         _logger.Warning("Failed to replicate {Blob} in {Namespace} due to it not existing.", blobToReplicate, ns);
