Horde Storage - Improved the orphan blob GC to run more in parallel making each objects state determined faster (we are still only GCing a single item at a time for now as we need to balance how quickly the GC runs with load on the DB).

diff --git a/Engine/Source/Programs/Horde/HordeStorage/Horde.Storage/Implementation/GC/OrphanBlobCleanupRefs.cs b/Engine/Source/Programs/Horde/HordeStorage/Horde.Storage/Implementation/GC/OrphanBlobCleanupRefs.cs
--- a/Engine/Source/Programs/Horde/HordeStorage/Horde.Storage/Implementation/GC/OrphanBlobCleanupRefs.cs
+++ b/Engine/Source/Programs/Horde/HordeStorage/Horde.Storage/Implementation/GC/OrphanBlobCleanupRefs.cs
@@ -66,6 +66,9 @@ namespace Horde.Storage.Implementation
                     break;
                 }
 
+                NamespacePolicy policy = _namespacePolicyResolver.GetPoliciesForNs(@namespace);
+                List<NamespaceId> namespacesThatSharePool = namespaces.Where(ns => _namespacePolicyResolver.GetPoliciesForNs(ns).StoragePool == policy.StoragePool).ToList();
+
                 // only consider blobs that have been around for 60 minutes
                 // this due to cases were blobs are uploaded first
                 DateTime cutoff = DateTime.Now.AddMinutes(-60);
@@ -85,53 +88,47 @@ namespace Horde.Storage.Implementation
                     removeBlobScope.Span.ResourceName = $"{@namespace}.{blob}";
 
                     bool found = false;
-                    NamespacePolicy policy = _namespacePolicyResolver.GetPoliciesForNs(@namespace);
 
-                    List<Task> referenceChecks = new List<Task>();
-                    List<NamespaceId> namespacesThatSharePool = namespaces.Where(ns => _namespacePolicyResolver.GetPoliciesForNs(ns).StoragePool == policy.StoragePool).ToList();
                     // check all other namespaces that share the same storage pool for presence of the blob
-                    foreach (NamespaceId blobNamespace in namespacesThatSharePool)
+                    await Parallel.ForEachAsync(namespacesThatSharePool, cancellationToken, async (blobNamespace, token) =>
                     {
-                        if (cancellationToken.IsCancellationRequested)
+                        if (token.IsCancellationRequested)
                         {
-                            break;
+                            return;
                         }
 
                         BlobInfo? blobIndex = await _blobIndex.GetBlobInfo(blobNamespace, blob);
 
                         if (blobIndex == null)
                         {
-                            continue;
+                            return;
                         }
 
-                        Task referenceCheckTask = Parallel.ForEachAsync(blobIndex.References, cancellationToken,
-                            async (tuple, token) =>
+                        await Parallel.ForEachAsync(blobIndex.References, cancellationToken, async (tuple, token) =>
+                        {
+                            (BucketId bucket, IoHashKey key) = tuple;
+                            if (token.IsCancellationRequested)
                             {
-                                (BucketId bucket, IoHashKey key) = tuple;
-                                if (token.IsCancellationRequested)
-                                {
-                                    return;
-                                }
-
-                                try
-                                {
-                                    (ObjectRecord, BlobContents?) _ = await _objectService.Get(blobNamespace, bucket, key, new string[] {"name"});
-                                    found = true;
-                                }
-                                catch (ObjectNotFoundException)
-                                {
-                                    // this is not a valid reference so we should delete
-                                }
-                                catch (MissingBlobsException)
-                                {
-                                    // we do not care if there are missing blobs, as long as the record is valid we keep this blob around
-                                    found = true;
-                                }
-                            });
-                        referenceChecks.Add(referenceCheckTask);
-                    }
+                                return;
+                            }
 
-                    await Task.WhenAll(referenceChecks);
+                            try
+                            {
+                                (ObjectRecord, BlobContents?) _ = await _objectService.Get(blobNamespace,
+                                    bucket, key, new string[] { "name" });
+                                found = true;
+                            }
+                            catch (ObjectNotFoundException)
+                            {
+                                // this is not a valid reference so we should delete
+                            }
+                            catch (MissingBlobsException)
+                            {
+                                // we do not care if there are missing blobs, as long as the record is valid we keep this blob around
+                                found = true;
+                            }
+                        });
+                    });
 
                     if (cancellationToken.IsCancellationRequested)
                     {
