Horde Storage - Fixed build issues with ambigous IBlobStore type

diff --git a/Engine/Source/Programs/Horde/HordeStorage/Horde.Storage/Tests/Functional/StorageTests.cs b/Engine/Source/Programs/Horde/HordeStorage/Horde.Storage/Tests/Functional/StorageTests.cs
--- a/Engine/Source/Programs/Horde/HordeStorage/Horde.Storage/Tests/Functional/StorageTests.cs
+++ b/Engine/Source/Programs/Horde/HordeStorage/Horde.Storage/Tests/Functional/StorageTests.cs
@@ -35,6 +35,7 @@ using Newtonsoft.Json.Linq;
 using Serilog;
 using Serilog.Core;
 using ContentHash = Jupiter.Implementation.ContentHash;
+using IBlobStore = Horde.Storage.Implementation.IBlobStore;
 
 namespace Horde.Storage.FunctionalTests.Storage
 {
diff --git a/Engine/Source/Programs/Horde/HordeStorage/Horde.Storage/Tests/UnitTests/HierarchicalBlobStorageTest.cs b/Engine/Source/Programs/Horde/HordeStorage/Horde.Storage/Tests/UnitTests/HierarchicalBlobStorageTest.cs
--- a/Engine/Source/Programs/Horde/HordeStorage/Horde.Storage/Tests/UnitTests/HierarchicalBlobStorageTest.cs
+++ b/Engine/Source/Programs/Horde/HordeStorage/Horde.Storage/Tests/UnitTests/HierarchicalBlobStorageTest.cs
@@ -21,6 +21,7 @@ using Moq;
 namespace Horde.Storage.UnitTests
 {
     using BlobNotFoundException = Horde.Storage.Implementation.BlobNotFoundException;
+    using IBlobStore = Implementation.IBlobStore;
 
     [TestClass]
     public class HierarchicalBlobStorageTest
