Horde: Let agent join requested pools on each new session

Previously the pools were only joined when the agent was first created.

diff --git a/Engine/Source/Programs/Horde/Horde.Build.Tests/RpcServiceTest.cs b/Engine/Source/Programs/Horde/Horde.Build.Tests/RpcServiceTest.cs
--- a/Engine/Source/Programs/Horde/Horde.Build.Tests/RpcServiceTest.cs
+++ b/Engine/Source/Programs/Horde/Horde.Build.Tests/RpcServiceTest.cs
@@ -318,6 +318,14 @@ namespace Horde.Build.Tests
 
 			IAgent agent = (await AgentService.GetAgentAsync(new AgentId(res.AgentId)))!;
 			CollectionAssert.AreEquivalent(new List<PoolId> { new ("fooPool"), new ("barPool") }, agent.GetPools().ToList());
+			
+			// Connect a second time, when the agent has already been created
+			req = new () { Name = "bogusAgentName", Capabilities = new AgentCapabilities() };
+			req.Capabilities.Properties.Add($"{KnownPropertyNames.RequestedPools}=bazPool");
+			res = await RpcService.CreateSession(req, _adminContext);
+
+			agent = (await AgentService.GetAgentAsync(new AgentId(res.AgentId)))!;
+			CollectionAssert.AreEquivalent(new List<PoolId> { new ("fooPool"), new ("barPool") , new ("bazPool") }, agent.GetPools().ToList());
 		}
 		
 		[TestMethod]
diff --git a/Engine/Source/Programs/Horde/Horde.Build/Agents/AgentCollection.cs b/Engine/Source/Programs/Horde/Horde.Build/Agents/AgentCollection.cs
--- a/Engine/Source/Programs/Horde/Horde.Build/Agents/AgentCollection.cs
+++ b/Engine/Source/Programs/Horde/Horde.Build/Agents/AgentCollection.cs
@@ -482,12 +482,13 @@ namespace Horde.Build.Agents
 		}
 
 		/// <inheritdoc/>
-		public async Task<IAgent?> TryStartSessionAsync(IAgent agentInterface, SessionId sessionId, DateTime sessionExpiresAt, AgentStatus status, IReadOnlyList<string> properties, IReadOnlyDictionary<string, int> resources, IReadOnlyList<PoolId> dynamicPools, string? version)
+		public async Task<IAgent?> TryStartSessionAsync(IAgent agentInterface, SessionId sessionId, DateTime sessionExpiresAt, AgentStatus status, IReadOnlyList<string> properties, IReadOnlyDictionary<string, int> resources, IReadOnlyList<PoolId> pools, IReadOnlyList<PoolId> dynamicPools, string? version)
 		{
 			AgentDocument agent = (AgentDocument)agentInterface;
 			List<string> newProperties = properties.OrderBy(x => x, StringComparer.OrdinalIgnoreCase).ToList();
-			Dictionary<string, int> newResources = new Dictionary<string, int>(resources);
-			List<PoolId> newDynamicPools = new List<PoolId>(dynamicPools);
+			Dictionary<string, int> newResources = new (resources);
+			List<PoolId> newPools = new (pools);
+			List<PoolId> newDynamicPools = new (dynamicPools);
 
 			// Reset the agent to use the new session
 			UpdateDefinitionBuilder<AgentDocument> updateBuilder = Builders<AgentDocument>.Update;
@@ -500,6 +501,7 @@ namespace Horde.Build.Agents
 			updates.Add(updateBuilder.Unset(x => x.Deleted));
 			updates.Add(updateBuilder.Set(x => x.Properties, newProperties));
 			updates.Add(updateBuilder.Set(x => x.Resources, newResources));
+			updates.Add(updateBuilder.Set(x => x.Pools, newPools));
 			updates.Add(updateBuilder.Set(x => x.DynamicPools, newDynamicPools));
 			updates.Add(updateBuilder.Set(x => x.Version, version));
 			updates.Add(updateBuilder.Unset(x => x.RequestRestart));
diff --git a/Engine/Source/Programs/Horde/Horde.Build/Agents/AgentService.cs b/Engine/Source/Programs/Horde/Horde.Build/Agents/AgentService.cs
--- a/Engine/Source/Programs/Horde/Horde.Build/Agents/AgentService.cs
+++ b/Engine/Source/Programs/Horde/Horde.Build/Agents/AgentService.cs
@@ -258,6 +258,28 @@ namespace Horde.Build.Agents
 
 			return newDynamicPools;
 		}
+		
+		private static List<StringId<IPool>> GetRequestedPoolsFromProperties(IReadOnlyList<string> properties)
+		{
+			List<StringId<IPool>> poolIds = new();
+			foreach (string property in properties)
+			{
+				const string Key = KnownPropertyNames.RequestedPools + "=";
+				if (property.StartsWith(Key, StringComparison.InvariantCulture))
+				{
+					poolIds.AddRange(property[Key.Length..].Split(",").Select(x => new StringId<IPool>(x)));
+				}
+			}
+
+			return poolIds;
+		}
+		
+		private static List<StringId<IPool>> CombineCurrentAndRequestedPools(IReadOnlyList<StringId<IPool>> pools, IReadOnlyList<string> properties)
+		{
+			HashSet<StringId<IPool>> uniquePools = new(pools);
+			uniquePools.UnionWith(GetRequestedPoolsFromProperties(properties));
+			return new List<PoolId>(uniquePools);
+		}
 
 		/// <summary>
 		/// Callback for an agents 
@@ -311,11 +333,12 @@ namespace Horde.Build.Agents
 					ISession newSession = await _sessions.AddAsync(SessionId.GenerateNewId(), agent.Id, _clock.UtcNow, properties, resources, version);
 					DateTime sessionExpiresAt = utcNow + SessionExpiryTime;
 
-					// Get the new dynamic pools for the agent
+					// Get the new pools for the agent
 					List<PoolId> dynamicPools = await GetDynamicPoolsAsync(agent);
+					List<PoolId> pools = CombineCurrentAndRequestedPools(agent.ExplicitPools, properties);
 
 					// Reset the agent to use the new session
-					newAgent = await Agents.TryStartSessionAsync(agent, newSession.Id, sessionExpiresAt, status, properties, resources, dynamicPools, version);
+					newAgent = await Agents.TryStartSessionAsync(agent, newSession.Id, sessionExpiresAt, status, properties, resources, pools, dynamicPools, version);
 					if(newAgent != null)
 					{
 						agent = newAgent;
diff --git a/Engine/Source/Programs/Horde/Horde.Build/Agents/IAgentCollection.cs b/Engine/Source/Programs/Horde/Horde.Build/Agents/IAgentCollection.cs
--- a/Engine/Source/Programs/Horde/Horde.Build/Agents/IAgentCollection.cs
+++ b/Engine/Source/Programs/Horde/Horde.Build/Agents/IAgentCollection.cs
@@ -107,10 +107,11 @@ namespace Horde.Build.Agents
 		/// <param name="status">Status of the agent</param>
 		/// <param name="properties">Properties for the current session</param>
 		/// <param name="resources">Resources for the agent</param>
+		/// <param name="pools">New list of pools for the agent</param>
 		/// <param name="dynamicPools">New list of dynamic pools for the agent</param>
 		/// <param name="version">Current version of the agent software</param>
 		/// <returns>New agent state</returns>
-		Task<IAgent?> TryStartSessionAsync(IAgent agent, SessionId sessionId, DateTime sessionExpiresAt, AgentStatus status, IReadOnlyList<string> properties, IReadOnlyDictionary<string, int> resources, IReadOnlyList<PoolId> dynamicPools, string? version);
+		Task<IAgent?> TryStartSessionAsync(IAgent agent, SessionId sessionId, DateTime sessionExpiresAt, AgentStatus status, IReadOnlyList<string> properties, IReadOnlyDictionary<string, int> resources, IReadOnlyList<PoolId> pools, IReadOnlyList<PoolId> dynamicPools, string? version);
 
 		/// <summary>
 		/// Attempt to update the agent state
diff --git a/Engine/Source/Programs/Horde/Horde.Build/Server/RpcService.cs b/Engine/Source/Programs/Horde/Horde.Build/Server/RpcService.cs
--- a/Engine/Source/Programs/Horde/Horde.Build/Server/RpcService.cs
+++ b/Engine/Source/Programs/Horde/Horde.Build/Server/RpcService.cs
@@ -283,7 +283,7 @@ namespace Horde.Build.Server
 					throw new StructuredRpcException(StatusCode.PermissionDenied, "User is not authenticated to create new agents");
 				}
 
-				agent = await _agentService.CreateAgentAsync(request.Name, true, null, GetRequestedPoolsFromCapabilities(request.Capabilities));
+				agent = await _agentService.CreateAgentAsync(request.Name, true, null, null);
 			}
 
 			// Make sure we're allowed to create sessions on this agent
@@ -311,20 +311,6 @@ namespace Horde.Build.Server
 			return response;
 		}
 
-		private static List<StringId<IPool>>? GetRequestedPoolsFromCapabilities(RpcAgentCapabilities capabilities)
-		{
-			foreach (string property in capabilities.Properties)
-			{
-				const string Key = "RequestedPools=";
-				if (property.StartsWith(Key, StringComparison.InvariantCulture))
-				{
-					return property[Key.Length..].Split(",").Select(x => new StringId<IPool>(x)).ToList();
-				}
-			}
-
-			return null;
-		}
-
 		/// <summary>
 		/// Updates an agent session
 		/// </summary>
