Horde Storage - Added more information on when each replicated event actually was from in the APM tracking for debugging purposes.

diff --git a/Engine/Source/Programs/Horde/HordeStorage/Horde.Storage/Implementation/Replication/RefsReplicator.cs b/Engine/Source/Programs/Horde/HordeStorage/Horde.Storage/Implementation/Replication/RefsReplicator.cs
--- a/Engine/Source/Programs/Horde/HordeStorage/Horde.Storage/Implementation/Replication/RefsReplicator.cs
+++ b/Engine/Source/Programs/Horde/HordeStorage/Horde.Storage/Implementation/Replication/RefsReplicator.cs
@@ -2,6 +2,7 @@
 
 using System;
 using System.Collections.Generic;
+using System.Globalization;
 using System.IO;
 using System.Linq;
 using System.Net;
@@ -415,8 +416,9 @@ namespace Horde.Storage.Implementation
                 }
                 using IScope scope = Tracer.Instance.StartActive("replicator.replicate_op_incremental");
                 scope.Span.ResourceName = $"[REDACTED_EMAIL]}[REDACTED_EMAIL]}";
+                scope.Span.SetTag("time-bucket", @event.Timestamp.ToString(CultureInfo.InvariantCulture));
 
-                _logger.Information("{Name} New transaction to replicate found. Ref: {Namespace} {Bucket} {Key} in {TimeBucket} with id {EventId}. Count of running replications: {CurrentReplications}", _name, @event.Namespace, @event.Bucket, @event.Key, @event.TimeBucket, @event.EventId, replicationTasks.Count);
+                _logger.Information("{Name} New transaction to replicate found. Ref: {Namespace} {Bucket} {Key} in {TimeBucket} ({TimeDate}) with id {EventId}. Count of running replications: {CurrentReplications}", _name, @event.Namespace, @event.Bucket, @event.Key, @event.TimeBucket, @event.Timestamp, @event.EventId, replicationTasks.Count);
                 
                 Info.CountOfRunningReplications = replicationTasks.Count;
                 LogReplicationHeartbeat(replicationTasks.Count);
