Horde Storage - Removed MaxReplicationsPerRun setting as this was added for tweaking the apm logging and this isnt used anymore and might be causing issues with the replication

diff --git a/Engine/Source/Programs/Horde/HordeStorage/Horde.Storage/Implementation/Replication/RefsReplicator.cs b/Engine/Source/Programs/Horde/HordeStorage/Horde.Storage/Implementation/Replication/RefsReplicator.cs
--- a/Engine/Source/Programs/Horde/HordeStorage/Horde.Storage/Implementation/Replication/RefsReplicator.cs
+++ b/Engine/Source/Programs/Horde/HordeStorage/Horde.Storage/Implementation/Replication/RefsReplicator.cs
@@ -408,12 +408,6 @@ namespace Horde.Storage.Implementation
 
             await GetRefEvents(ns, lastBucket, lastEvent, replicationToken).ParallelForEachAsync(async (ReplicationLogEvent @event) =>
             {
-                // if we have done all the replication events we should do in a single run we abort
-                if (countOfReplicationsDone > _replicatorSettings.MaxReplicationsPerRun)
-                {
-                    cancellationTokenSource.Cancel();
-                    return;
-                }
                 using IScope scope = Tracer.Instance.StartActive("replicator.replicate_op_incremental");
                 scope.Span.ResourceName = $"[REDACTED_EMAIL]}[REDACTED_EMAIL]}";
                 scope.Span.SetTag("time-bucket", @event.Timestamp.ToString(CultureInfo.InvariantCulture));
diff --git a/Engine/Source/Programs/Horde/HordeStorage/Horde.Storage/Implementation/Replication/ReplicationSettings.cs b/Engine/Source/Programs/Horde/HordeStorage/Horde.Storage/Implementation/Replication/ReplicationSettings.cs
--- a/Engine/Source/Programs/Horde/HordeStorage/Horde.Storage/Implementation/Replication/ReplicationSettings.cs
+++ b/Engine/Source/Programs/Horde/HordeStorage/Horde.Storage/Implementation/Replication/ReplicationSettings.cs
@@ -74,11 +74,6 @@ namespace Horde.Storage
         /// </summary>
         public int MaxParallelReplications { get; set; } = 64;
 
-        /// <summary>
-        /// The replication will stop its run after this many events have been processed, yielding its time until the next time a replication is triggered.
-        /// </summary>
-        public int MaxReplicationsPerRun { get; set; } = 5000;
-
         /// <summary>
         /// Max number of offsets to skip when finding broken events in the transaction logs, higher value means we scan further increasing the likelihood of finding a valid entry, but also means risk that we skip records you would want.
         /// </summary>
