Horde Storage - Added better errors messages when the replicator is not configured with a namespace. Also added more debug logging to better understand why some replicators stop running occasionally.

diff --git a/Engine/Source/Programs/Horde/HordeStorage/Horde.Storage/Implementation/Replication/RefsReplicator.cs b/Engine/Source/Programs/Horde/HordeStorage/Horde.Storage/Implementation/Replication/RefsReplicator.cs
--- a/Engine/Source/Programs/Horde/HordeStorage/Horde.Storage/Implementation/Replication/RefsReplicator.cs
+++ b/Engine/Source/Programs/Horde/HordeStorage/Horde.Storage/Implementation/Replication/RefsReplicator.cs
@@ -40,6 +40,7 @@ namespace Horde.Storage.Implementation
         private readonly IReplicationLog _replicationLog;
         private readonly IServiceCredentials _serviceCredentials;
         private readonly HttpClient _httpClient;
+        private readonly NamespaceId _namespace;
         private RefsState _refsState;
         private bool _replicationRunning;
         private bool _disposed = false;
@@ -47,6 +48,7 @@ namespace Horde.Storage.Implementation
         public RefsReplicator(ReplicatorSettings replicatorSettings, IOptionsMonitor<ReplicationSettings> replicationSettings, IBlobService blobService, IHttpClientFactory httpClientFactory, IReplicationLog replicationLog, IServiceCredentials serviceCredentials)
         {
             _name = replicatorSettings.ReplicatorName;
+            _namespace = new NamespaceId(replicatorSettings.NamespaceToReplicate);
             _replicatorSettings = replicatorSettings;
             _blobService = blobService;
             _replicationLog = replicationLog;
@@ -59,7 +61,7 @@ namespace Horde.Storage.Implementation
             DirectoryInfo stateRoot = new DirectoryInfo(replicationSettings.CurrentValue.StateRoot);
             _stateFile = new FileInfo(Path.Combine(stateRoot.FullName, stateFileName));
 
-            ReplicatorState? replicatorState = _replicationLog.GetReplicatorState(_replicatorSettings.NamespaceToReplicate, _name).Result;
+            ReplicatorState? replicatorState = _replicationLog.GetReplicatorState(_namespace, _name).Result;
             if (replicatorState == null)
             {
                 if (_stateFile.Exists)
@@ -80,7 +82,7 @@ namespace Horde.Storage.Implementation
                 };
             }
 
-            Info = new ReplicatorInfo(replicatorSettings.ReplicatorName, replicatorSettings.NamespaceToReplicate, _refsState);
+            Info = new ReplicatorInfo(replicatorSettings.ReplicatorName, _namespace, _refsState);
         }
 
         public void Dispose()
@@ -140,11 +142,12 @@ namespace Horde.Storage.Implementation
         {
             if (_replicationRunning)
             {
+                _logger.Debug("Skipping replication of replicator: {Name} as it was already running.", _name);
                 return false;
             }
 
             // read the state again to allow it to be modified by the admin controller / other instances of horde-storage connected to the same filesystem
-            ReplicatorState? replicatorState = await _replicationLog.GetReplicatorState(_replicatorSettings.NamespaceToReplicate, _name);
+            ReplicatorState? replicatorState = await _replicationLog.GetReplicatorState(_namespace, _name);
             if (replicatorState == null)
             {
                 _stateFile.Refresh();
@@ -166,23 +169,27 @@ namespace Horde.Storage.Implementation
                 };
             }
 
+            _logger.Debug("Read Replication state for replicator: {Name}. {LastBucket} {LastEvent}.", _name, _refsState.LastBucket, _refsState.LastEvent);
+
             LogReplicationHeartbeat(0);
 
             bool hasRun;
+            int countOfReplicationsDone = 0;
 
             try
             {
                 //using IScope scope = Tracer.Instance.StartActive("replicator.run");
                 //scope.Span.ResourceName =_name;
 
+                _logger.Debug("Replicator: {Name} is starting a run.", _name);
+
                 _replicationTokenSource.TryReset();
                 _replicationRunning = true;
                 _replicationFinishedEvent.Reset();
                 CancellationToken replicationToken = _replicationTokenSource.Token;
 
-                NamespaceId ns = _replicatorSettings.NamespaceToReplicate;
+                NamespaceId ns = _namespace;
 
-                int countOfReplicationsDone = 0;
                 Info.CountOfRunningReplications = 0;
 
                 string? lastBucket = null;
@@ -285,6 +292,8 @@ namespace Horde.Storage.Implementation
                 _replicationFinishedEvent.Set();
             }
 
+            _logger.Debug("Replicator: {Name} finished its replication run. Replications completed: {ReplicationsDone} .", _name, countOfReplicationsDone);
+
             return hasRun;
         }
 
diff --git a/Engine/Source/Programs/Horde/HordeStorage/Horde.Storage/Implementation/Replication/ReplicationService.cs b/Engine/Source/Programs/Horde/HordeStorage/Horde.Storage/Implementation/Replication/ReplicationService.cs
--- a/Engine/Source/Programs/Horde/HordeStorage/Horde.Storage/Implementation/Replication/ReplicationService.cs
+++ b/Engine/Source/Programs/Horde/HordeStorage/Horde.Storage/Implementation/Replication/ReplicationService.cs
@@ -120,12 +120,14 @@ namespace Horde.Storage.Implementation
                     }
                     else
                     {
+                        _logger.Debug("Replication of replicator: {Name} is still running", replicator.Info.ReplicatorName);
                         // if the replication is still running let it continue to run
                         continue;
                     }
                 }
 
                 // start a new run of the replication
+                _logger.Debug("Triggering new replication of replicator: {Name}", replicator.Info.ReplicatorName);
                 Task<bool> newReplication = replicator.TriggerNewReplications();
                 _currentReplications[replicator.Info.ReplicatorName] = newReplication;
             }
diff --git a/Engine/Source/Programs/Horde/HordeStorage/Horde.Storage/Implementation/Replication/ReplicationSettings.cs b/Engine/Source/Programs/Horde/HordeStorage/Horde.Storage/Implementation/Replication/ReplicationSettings.cs
--- a/Engine/Source/Programs/Horde/HordeStorage/Horde.Storage/Implementation/Replication/ReplicationSettings.cs
+++ b/Engine/Source/Programs/Horde/HordeStorage/Horde.Storage/Implementation/Replication/ReplicationSettings.cs
@@ -3,7 +3,6 @@
 using System;
 using System.Collections.Generic;
 using System.ComponentModel.DataAnnotations;
-using EpicGames.Horde.Storage;
 using Horde.Storage.Implementation;
 
 namespace Horde.Storage
@@ -52,7 +51,7 @@ namespace Horde.Storage
         /// The namespace which this replicator is replicating, will be used both on the source and destination side 
         /// </summary>
         [Required]
-        public NamespaceId NamespaceToReplicate { get; set; }
+        public string NamespaceToReplicate { get; set; } = "";
 
         /// <summary>
         /// Name of this replicator instance, should be unique within the cluster
diff --git a/Engine/Source/Programs/Horde/HordeStorage/Horde.Storage/Implementation/Replication/Replicator.cs b/Engine/Source/Programs/Horde/HordeStorage/Horde.Storage/Implementation/Replication/Replicator.cs
--- a/Engine/Source/Programs/Horde/HordeStorage/Horde.Storage/Implementation/Replication/Replicator.cs
+++ b/Engine/Source/Programs/Horde/HordeStorage/Horde.Storage/Implementation/Replication/Replicator.cs
@@ -40,7 +40,7 @@ namespace Horde.Storage.Implementation
         protected override IAsyncEnumerable<TransactionEvent> GetCallistoOp(long stateReplicatorOffset, Guid? stateReplicatingGeneration, string currentSite,
             OpsEnumerationState enumerationState, CancellationToken replicationToken)
         {
-            CallistoReader remoteCallistoReader = new CallistoReader(Client, ReplicatorSettings.NamespaceToReplicate);
+            CallistoReader remoteCallistoReader = new CallistoReader(Client, Namespace);
 
             return remoteCallistoReader.GetOps(stateReplicatorOffset, stateReplicatingGeneration, currentSite, enumerationState: enumerationState, cancellationToken: replicationToken, maxOffsetsAttempted: ReplicatorSettings.MaxOffsetsAttempted);
         }
@@ -48,7 +48,7 @@ namespace Horde.Storage.Implementation
         protected override async Task ReplicateOp(IRestClient remoteClient, TransactionEvent op, CancellationToken replicationToken)
         {
             using IScope scope = Tracer.Instance.StartActive("replicator.replicate_blobs");
-            NamespaceId ns = ReplicatorSettings.NamespaceToReplicate;
+            NamespaceId ns = Namespace;
 
             // now we replicate the blobs
             switch (op)
@@ -71,7 +71,7 @@ namespace Horde.Storage.Implementation
             using IScope scope = Tracer.Instance.StartActive("replicator.replicate_inline");
             //scope.Span.ResourceName = op.transactionId;
 
-            NamespaceId ns = ReplicatorSettings.NamespaceToReplicate;
+            NamespaceId ns = Namespace;
 
             // Put a copy of the event in the local callisto store
             // Make sure this is put before the io put request as we want to pin the reference before uploading the content to avoid the gc removing it right away
@@ -100,6 +100,7 @@ namespace Horde.Storage.Implementation
         private readonly HttpClient _httpClient;
 
         protected string Name { get; set; }
+        public NamespaceId Namespace { get; set; }
         protected ReplicatorSettings ReplicatorSettings { get; init; }
 
         public ReplicatorState State { get; private set; }
@@ -117,6 +118,7 @@ namespace Horde.Storage.Implementation
         protected Replicator(ReplicatorSettings replicatorSettings, IOptionsMonitor<ReplicationSettings> replicationSettings, IOptionsMonitor<JupiterSettings> jupiterSettings, IBlobService blobService, IRestClient remoteClient, IServiceCredentials serviceCredentials, IHttpClientFactory httpClientFactory)
         {
             Name = replicatorSettings.ReplicatorName;
+            Namespace = new NamespaceId(replicatorSettings.NamespaceToReplicate);
             _blobService = blobService;
             _currentSite = jupiterSettings.CurrentValue.CurrentSite;
             ReplicatorSettings = replicatorSettings;
@@ -134,7 +136,7 @@ namespace Horde.Storage.Implementation
                 State = new ReplicatorState {ReplicatorOffset = null, ReplicatingGeneration = null};
             }
 
-            Info = new ReplicatorInfo(ReplicatorSettings.ReplicatorName, ReplicatorSettings.NamespaceToReplicate, State);
+            Info = new ReplicatorInfo(ReplicatorSettings.ReplicatorName, Namespace, State);
             Client = remoteClient;
 
             _serviceCredentials = serviceCredentials;
diff --git a/Engine/Source/Programs/Horde/HordeStorage/Horde.Storage/Tests/Functional/Objects/RefReplicatorTests.cs b/Engine/Source/Programs/Horde/HordeStorage/Horde.Storage/Tests/Functional/Objects/RefReplicatorTests.cs
--- a/Engine/Source/Programs/Horde/HordeStorage/Horde.Storage/Tests/Functional/Objects/RefReplicatorTests.cs
+++ b/Engine/Source/Programs/Horde/HordeStorage/Horde.Storage/Tests/Functional/Objects/RefReplicatorTests.cs
@@ -103,7 +103,7 @@ namespace Horde.Storage.FunctionalTests.Replication
             {
                 ConnectionString = "http://localhost",
                 MaxParallelReplications = 16,
-                NamespaceToReplicate = TestNamespace,
+                NamespaceToReplicate = TestNamespace.ToString(),
                 ReplicatorName = "test-replicator",
                 Version = ReplicatorVersion.Refs
             };
@@ -160,7 +160,7 @@ namespace Horde.Storage.FunctionalTests.Replication
             {
                 ConnectionString = "http://localhost",
                 MaxParallelReplications = 16,
-                NamespaceToReplicate = TestNamespace,
+                NamespaceToReplicate = TestNamespace.ToString(),
                 ReplicatorName = "test-replicator",
                 Version = ReplicatorVersion.Refs
             };
@@ -235,7 +235,7 @@ namespace Horde.Storage.FunctionalTests.Replication
             {
                 ConnectionString = "http://localhost",
                 MaxParallelReplications = 16,
-                NamespaceToReplicate = TestNamespace,
+                NamespaceToReplicate = TestNamespace.ToString(),
                 ReplicatorName = "test-replicator",
                 Version = ReplicatorVersion.Refs
             };
@@ -322,7 +322,7 @@ namespace Horde.Storage.FunctionalTests.Replication
             {
                 ConnectionString = "http://localhost",
                 MaxParallelReplications = 16,
-                NamespaceToReplicate = TestNamespace,
+                NamespaceToReplicate = TestNamespace.ToString(),
                 ReplicatorName = "test-replicator",
                 Version = ReplicatorVersion.Refs
             };
diff --git a/Engine/Source/Programs/Horde/HordeStorage/Horde.Storage/Tests/UnitTests/ReplicatorTests.cs b/Engine/Source/Programs/Horde/HordeStorage/Horde.Storage/Tests/UnitTests/ReplicatorTests.cs
--- a/Engine/Source/Programs/Horde/HordeStorage/Horde.Storage/Tests/UnitTests/ReplicatorTests.cs
+++ b/Engine/Source/Programs/Horde/HordeStorage/Horde.Storage/Tests/UnitTests/ReplicatorTests.cs
@@ -48,7 +48,7 @@ namespace Horde.Storage.UnitTests
         {
             ReplicatorSettings replicatorSettings = new ReplicatorSettings
             {
-                NamespaceToReplicate = NamespaceV1, ReplicatorName = ReplicatorNameV1, Version = ReplicatorVersion.V1, ConnectionString = "http://localhost"
+                NamespaceToReplicate = NamespaceV1.ToString(), ReplicatorName = ReplicatorNameV1, Version = ReplicatorVersion.V1, ConnectionString = "http://localhost"
             };
 
             byte[] contents = Encoding.ASCII.GetBytes("test string");
@@ -108,7 +108,7 @@ namespace Horde.Storage.UnitTests
             bool ran = await replicator.TriggerNewReplications();
             Assert.IsTrue(ran);
 
-            NamespaceId ns = replicatorSettings.NamespaceToReplicate;
+            NamespaceId ns = new NamespaceId(replicatorSettings.NamespaceToReplicate);
 
             // we should have checked the local io once for if we had the blob
             blobStoreMock.Verify(blobStore => blobStore.Exists(ns, blobToReplication), Times.Once);
@@ -130,7 +130,7 @@ namespace Horde.Storage.UnitTests
         {
             ReplicatorSettings replicatorSettings = new ReplicatorSettings
             {
-                NamespaceToReplicate = NamespaceV1, ReplicatorName = ReplicatorNameV1, Version = ReplicatorVersion.V1, ConnectionString = "http://localhost"
+                NamespaceToReplicate = NamespaceV1.ToString(), ReplicatorName = ReplicatorNameV1, Version = ReplicatorVersion.V1, ConnectionString = "http://localhost"
             };
 
             Mock<IRestClient> remoteClientMock = new Mock<IRestClient> { DefaultValue = DefaultValue.Empty };
@@ -183,7 +183,7 @@ namespace Horde.Storage.UnitTests
         {
             ReplicatorSettings replicatorSettings = new ReplicatorSettings
             {
-                NamespaceToReplicate = NamespaceV1, ReplicatorName = ReplicatorNameV1, Version = ReplicatorVersion.V1, ConnectionString = "http://localhost"
+                NamespaceToReplicate = NamespaceV1.ToString(), ReplicatorName = ReplicatorNameV1, Version = ReplicatorVersion.V1, ConnectionString = "http://localhost"
             };
 
             byte[] contents = Encoding.ASCII.GetBytes("test string");
@@ -245,7 +245,7 @@ namespace Horde.Storage.UnitTests
             bool ran = await replicator.TriggerNewReplications();
             Assert.IsTrue(ran);
 
-            NamespaceId ns = replicatorSettings.NamespaceToReplicate;
+            NamespaceId ns = new NamespaceId(replicatorSettings.NamespaceToReplicate);
 
             // we should have checked the local io once for if we had the blob
             blobStoreMock.Verify(blobStore => blobStore.Exists(ns, blobToReplication), Times.Once);
@@ -268,10 +268,10 @@ namespace Horde.Storage.UnitTests
         {
             ReplicatorSettings replicatorSettings = new ReplicatorSettings
             {
-                NamespaceToReplicate = NamespaceV1, ReplicatorName = ReplicatorNameV1, Version = ReplicatorVersion.V1, ConnectionString = "http://localhost"
+                NamespaceToReplicate = NamespaceV1.ToString(), ReplicatorName = ReplicatorNameV1, Version = ReplicatorVersion.V1, ConnectionString = "http://localhost"
             };
 
-            NamespaceId ns = replicatorSettings.NamespaceToReplicate;
+            NamespaceId ns = new NamespaceId(replicatorSettings.NamespaceToReplicate);
 
             Mock<IBlobService> blobStoreMock = new Mock<IBlobService>();
             IServiceCredentials serviceCredentials = Mock.Of<IServiceCredentials>();
