Horde Storage - Reduce logging from system loggers (to warning from information) as this generates a lot of verbose output from http client. This is useful to have when running locally but it becomes a bit much in production deployments.

diff --git a/Engine/Source/Programs/Horde/HordeStorage/Horde.Storage/appsettings.Production.json b/Engine/Source/Programs/Horde/HordeStorage/Horde.Storage/appsettings.Production.json
--- a/Engine/Source/Programs/Horde/HordeStorage/Horde.Storage/appsettings.Production.json
+++ b/Engine/Source/Programs/Horde/HordeStorage/Horde.Storage/appsettings.Production.json
@@ -5,7 +5,7 @@
       "Default": "Information",
       "Override": {
         "Microsoft": "Warning",
-        "System": "Information"
+        "System": "Warning"
       }
     },
     "WriteTo": [
