Horde: Add a pool setting to configure how frequently to run conforms ("ConformInterval", specified in hours).

Setting the value to zero will disable conforms, setting the value to -1 will use the default. If an agent is part of several pools, the highest frequency schedule will be used.

diff --git a/Engine/Source/Programs/Horde/Horde.Build/Agents/Pools/IPool.cs b/Engine/Source/Programs/Horde/Horde.Build/Agents/Pools/IPool.cs
--- a/Engine/Source/Programs/Horde/Horde.Build/Agents/Pools/IPool.cs
+++ b/Engine/Source/Programs/Horde/Horde.Build/Agents/Pools/IPool.cs
@@ -60,6 +60,11 @@ namespace Horde.Build.Agents.Pools
 		/// </summary>
 		public int? NumReserveAgents { get; }
 
+		/// <summary>
+		/// Interval between conforms. If zero, the pool will not conform on a schedule.
+		/// </summary>
+		public TimeSpan? ConformInterval { get; }
+
 		/// <summary>
 		/// Last time the pool was (auto) scaled up
 		/// </summary>
diff --git a/Engine/Source/Programs/Horde/Horde.Build/Agents/Pools/IPoolCollection.cs b/Engine/Source/Programs/Horde/Horde.Build/Agents/Pools/IPoolCollection.cs
--- a/Engine/Source/Programs/Horde/Horde.Build/Agents/Pools/IPoolCollection.cs
+++ b/Engine/Source/Programs/Horde/Horde.Build/Agents/Pools/IPoolCollection.cs
@@ -25,6 +25,7 @@ namespace Horde.Build.Agents.Pools
 		/// <param name="enableAutoscaling">Whether to enable autoscaling for this pool</param>
 		/// <param name="minAgents">Minimum number of agents in the pool</param>
 		/// <param name="numReserveAgents">Minimum number of idle agents to maintain</param>
+		/// <param name="conformInterval">Interval between conforms. Set to zero to disable.</param>
 		/// <param name="scaleOutCooldown">Cooldown time between scale-out events</param>
 		/// <param name="scaleInCooldown">Cooldown time between scale-in events</param>
 		/// <param name="sizeStrategy">Pool sizing strategy</param>
@@ -39,6 +40,7 @@ namespace Horde.Build.Agents.Pools
 			bool? enableAutoscaling = null,
 			int? minAgents = null,
 			int? numReserveAgents = null,
+			TimeSpan? conformInterval = null,
 			TimeSpan? scaleOutCooldown = null,
 			TimeSpan? scaleInCooldown = null,
 			PoolSizeStrategy? sizeStrategy = null,
@@ -84,6 +86,7 @@ namespace Horde.Build.Agents.Pools
 		/// <param name="newWorkspaces">New workspaces for the pool</param>
 		/// <param name="newUseAutoSdk">New setting for whether to use autosdk</param>
 		/// <param name="newProperties">New properties for the pool</param>
+		/// <param name="conformInterval">Interval between conforms. Set to zero to disable.</param>
 		/// <param name="lastScaleUpTime">New time for last (auto) scale up</param>
 		/// <param name="lastScaleDownTime">New time for last (auto) scale down</param>
 		/// <param name="scaleOutCooldown">Cooldown time between scale-out events</param>
@@ -103,6 +106,7 @@ namespace Horde.Build.Agents.Pools
 			List<AgentWorkspace>? newWorkspaces = null,
 			bool? newUseAutoSdk = null,
 			Dictionary<string, string?>? newProperties = null,
+			TimeSpan? conformInterval = null,
 			DateTime? lastScaleUpTime = null,
 			DateTime? lastScaleDownTime = null,
 			TimeSpan? scaleOutCooldown = null,
diff --git a/Engine/Source/Programs/Horde/Horde.Build/Agents/Pools/PoolCollection.cs b/Engine/Source/Programs/Horde/Horde.Build/Agents/Pools/PoolCollection.cs
--- a/Engine/Source/Programs/Horde/Horde.Build/Agents/Pools/PoolCollection.cs
+++ b/Engine/Source/Programs/Horde/Horde.Build/Agents/Pools/PoolCollection.cs
@@ -37,6 +37,7 @@ namespace Horde.Build.Agents.Pools
 			public bool EnableAutoscaling { get; set; } = true;
 			public int? MinAgents { get; set; }
 			public int? NumReserveAgents { get; set; }
+			public TimeSpan? ConformInterval { get; set; }
 			public DateTime? LastScaleUpTime { get; set; }
 			public DateTime? LastScaleDownTime { get; set; }
 			public TimeSpan? ScaleOutCooldown { get; set; }
@@ -66,6 +67,7 @@ namespace Horde.Build.Agents.Pools
 				EnableAutoscaling = other.EnableAutoscaling;
 				MinAgents = other.MinAgents;
 				NumReserveAgents = other.NumReserveAgents;
+				ConformInterval = other.ConformInterval;
 				LastScaleUpTime = other.LastScaleUpTime;
 				LastScaleDownTime = other.LastScaleDownTime;
 				ScaleOutCooldown = other.ScaleOutCooldown;
@@ -103,6 +105,7 @@ namespace Horde.Build.Agents.Pools
 			bool? enableAutoscaling,
 			int? minAgents,
 			int? numReserveAgents,
+			TimeSpan? conformInterval,
 			TimeSpan? scaleOutCooldown,
 			TimeSpan? scaleInCooldown,
 			PoolSizeStrategy? sizeStrategy,
@@ -125,6 +128,7 @@ namespace Horde.Build.Agents.Pools
 				pool.Properties = new Dictionary<string, string>(properties);
 			}
 
+			pool.ConformInterval = conformInterval;
 			pool.ScaleOutCooldown = scaleOutCooldown;
 			pool.ScaleInCooldown = scaleInCooldown;
 			pool.SizeStrategy = sizeStrategy;
@@ -214,6 +218,7 @@ namespace Horde.Build.Agents.Pools
 			List<AgentWorkspace>? newWorkspaces,
 			bool? newUseAutoSdk,
 			Dictionary<string, string?>? newProperties,
+			TimeSpan? conformInterval,
 			DateTime? lastScaleUpTime,
 			DateTime? lastScaleDownTime,
 			TimeSpan? scaleOutCooldown,
@@ -284,6 +289,17 @@ namespace Horde.Build.Agents.Pools
 			{
 				transaction.UpdateDictionary(x => x.Properties, newProperties);
 			}
+			if (conformInterval != null)
+			{
+				if (conformInterval.Value < TimeSpan.Zero)
+				{
+					transaction.Unset(x => x.ConformInterval!);
+				}
+				else
+				{
+					transaction.Set(x => x.ConformInterval, conformInterval);
+				}
+			}
 			if (lastScaleUpTime != null)
 			{
 				transaction.Set(x => x.LastScaleUpTime, lastScaleUpTime);
diff --git a/Engine/Source/Programs/Horde/Horde.Build/Agents/Pools/PoolMessages.cs b/Engine/Source/Programs/Horde/Horde.Build/Agents/Pools/PoolMessages.cs
--- a/Engine/Source/Programs/Horde/Horde.Build/Agents/Pools/PoolMessages.cs
+++ b/Engine/Source/Programs/Horde/Horde.Build/Agents/Pools/PoolMessages.cs
@@ -95,6 +95,11 @@ namespace Horde.Build.Agents.Pools
 		/// Whether to enable autoscaling for this pool
 		/// </summary>
 		public bool? EnableAutoscaling { get; set; }
+
+		/// <summary>
+		/// Interval between conforms in hours. Set to zero to disable.
+		/// </summary>
+		public int? ConformInterval { get; set; }
 		
 		/// <summary>
 		/// Cooldown time between scale-out events in seconds
@@ -176,6 +181,11 @@ namespace Horde.Build.Agents.Pools
 		/// Whether to enable autoscaling for this pool
 		/// </summary>
 		public bool? EnableAutoscaling { get; set; }
+
+		/// <summary>
+		/// Interval between conforms in hours. Set to -1 to reset to the default, or 0 to disable.
+		/// </summary>
+		public int? ConformInterval { get; set; }
 		
 		/// <summary>
 		/// Cooldown time between scale-out events in seconds
@@ -270,6 +280,11 @@ namespace Horde.Build.Agents.Pools
 		/// </summary>
 		public bool EnableAutoscaling { get; set; }
 		
+		/// <summary>
+		/// Frequency to run conforms, in hours.
+		/// </summary>
+		public int? ConformInterval { get; set; }
+
 		/// <summary>
 		/// Cooldown time between scale-out events in seconds
 		/// </summary>
@@ -325,6 +340,7 @@ namespace Horde.Build.Agents.Pools
 			Name = pool.Name;
 			Condition = pool.Condition;
 			EnableAutoscaling = pool.EnableAutoscaling;
+			ConformInterval = pool.ConformInterval == null ? null : (int)pool.ConformInterval.Value.TotalHours;
 			ScaleOutCooldown = pool.ScaleOutCooldown == null ? null : (int)pool.ScaleOutCooldown.Value.TotalSeconds;
 			ScaleInCooldown = pool.ScaleInCooldown == null ? null : (int)pool.ScaleInCooldown.Value.TotalSeconds;
 			SizeStrategy = pool.SizeStrategy;
diff --git a/Engine/Source/Programs/Horde/Horde.Build/Agents/Pools/PoolService.cs b/Engine/Source/Programs/Horde/Horde.Build/Agents/Pools/PoolService.cs
--- a/Engine/Source/Programs/Horde/Horde.Build/Agents/Pools/PoolService.cs
+++ b/Engine/Source/Programs/Horde/Horde.Build/Agents/Pools/PoolService.cs
@@ -61,6 +61,7 @@ namespace Horde.Build.Agents.Pools
 		/// <param name="enableAutoscaling">Whether to enable autoscaling for this pool</param>
 		/// <param name="minAgents">Minimum number of agents in the pool</param>
 		/// <param name="numReserveAgents">Minimum number of idle agents to maintain</param>
+		/// <param name="conformInterval">Interval between conforms. Set to zero to disable.</param>
 		/// <param name="scaleOutCooldown">Cooldown time between scale-out events</param>
 		/// <param name="scaleInCooldown">Cooldown time between scale-in events</param>
 		/// <param name="sizeStrategy">Pool sizing strategy</param>
@@ -74,6 +75,7 @@ namespace Horde.Build.Agents.Pools
 			bool? enableAutoscaling = null,
 			int? minAgents = null,
 			int? numReserveAgents = null,
+			TimeSpan? conformInterval = null,
 			TimeSpan? scaleOutCooldown = null,
 			TimeSpan? scaleInCooldown = null,
 			PoolSizeStrategy? sizeStrategy = null,
@@ -88,6 +90,7 @@ namespace Horde.Build.Agents.Pools
 				enableAutoscaling,
 				minAgents,
 				numReserveAgents,
+				conformInterval,
 				scaleOutCooldown,
 				scaleInCooldown,
 				sizeStrategy,
@@ -116,6 +119,7 @@ namespace Horde.Build.Agents.Pools
 		/// <param name="newMinAgents">Minimum number of agents in the pool</param>
 		/// <param name="newNumReserveAgents">Minimum number of idle agents to maintain</param>
 		/// <param name="newProperties">Properties on the pool to update. Any properties with a value of null will be removed.</param>
+		/// <param name="conformInterval">Interval between conforms. Set to zero to disable.</param>
 		/// <param name="scaleOutCooldown">Cooldown time between scale-out events</param>
 		/// <param name="scaleInCooldown">Cooldown time between scale-in events</param>
 		/// <param name="sizeStrategy">New pool sizing strategy for the pool</param>
@@ -131,6 +135,7 @@ namespace Horde.Build.Agents.Pools
 			int? newMinAgents = null,
 			int? newNumReserveAgents = null,
 			Dictionary<string, string?>? newProperties = null,
+			TimeSpan? conformInterval = null,
 			TimeSpan? scaleOutCooldown = null,
 			TimeSpan? scaleInCooldown = null,
 			PoolSizeStrategy? sizeStrategy = null,
@@ -148,6 +153,7 @@ namespace Horde.Build.Agents.Pools
 					newMinAgents,
 					newNumReserveAgents,
 					newProperties: newProperties,
+					conformInterval: conformInterval,
 					scaleOutCooldown: scaleOutCooldown,
 					scaleInCooldown: scaleInCooldown,
 					sizeStrategy: sizeStrategy,
@@ -182,6 +188,41 @@ namespace Horde.Build.Agents.Pools
 			return _pools.GetAsync(poolId);
 		}
 
+		/// <summary>
+		/// Gets a cached pool definition
+		/// </summary>
+		/// <param name="poolId"></param>
+		/// <param name="validAtTime"></param>
+		/// <returns></returns>
+		public async Task<IPool?> GetCachedPoolAsync(PoolId poolId, DateTime validAtTime)
+		{
+			Dictionary<PoolId, IPool> poolMapping = await GetPoolLookupAsync(validAtTime);
+			poolMapping.TryGetValue(poolId, out IPool? pool);
+			return pool;
+		}
+
+		/// <summary>
+		/// Gets a cached pool definition
+		/// </summary>
+		/// <param name="agent"></param>
+		/// <param name="validAtTime"></param>
+		/// <returns></returns>
+		public async Task<List<IPool>> GetCachedPoolsAsync(IAgent agent, DateTime validAtTime)
+		{
+			Dictionary<PoolId, IPool> poolMapping = await GetPoolLookupAsync(validAtTime);
+
+			List<IPool> pools = new List<IPool>();
+			foreach(PoolId poolId in agent.GetPools())
+			{
+				if(poolMapping.TryGetValue(poolId, out IPool? pool))
+				{
+					pools.Add(pool);
+				}
+			}
+
+			return pools;
+		}
+
 		/// <summary>
 		/// Get a list of workspaces for the given agent
 		/// </summary>
diff --git a/Engine/Source/Programs/Horde/Horde.Build/Agents/Pools/PoolsController.cs b/Engine/Source/Programs/Horde/Horde.Build/Agents/Pools/PoolsController.cs
--- a/Engine/Source/Programs/Horde/Horde.Build/Agents/Pools/PoolsController.cs
+++ b/Engine/Source/Programs/Horde/Horde.Build/Agents/Pools/PoolsController.cs
@@ -58,12 +58,13 @@ namespace Horde.Build.Agents.Pools
 
 			LeaseUtilizationSettings? luSettings = create.LeaseUtilizationSettings?.Convert();
 			JobQueueSettings? jqSettings = create.JobQueueSettings?.Convert();
-			
+
+			TimeSpan? conformInterval = create.ConformInterval == null ? null : TimeSpan.FromHours(create.ConformInterval.Value);
 			TimeSpan? scaleOutCooldown = create.ScaleOutCooldown == null ? null : TimeSpan.FromSeconds(create.ScaleOutCooldown.Value);
 			TimeSpan? scaleInCooldown = create.ScaleInCooldown == null ? null : TimeSpan.FromSeconds(create.ScaleInCooldown.Value);
 
 			IPool newPool = await _poolService.CreatePoolAsync(
-				create.Name, create.Condition, create.EnableAutoscaling, create.MinAgents, create.NumReserveAgents,
+				create.Name, create.Condition, create.EnableAutoscaling, create.MinAgents, create.NumReserveAgents, conformInterval,
 				scaleOutCooldown, scaleInCooldown, create.SizeStrategy, luSettings, jqSettings, create.Properties);
 			return new CreatePoolResponse(newPool.Id.ToString());
 		}
@@ -142,12 +143,13 @@ namespace Horde.Build.Agents.Pools
 			{
 				return NotFound(poolIdValue);
 			}
-			
+
+			TimeSpan? conformInterval = update.ConformInterval == null ? null : TimeSpan.FromHours(update.ConformInterval.Value);
 			TimeSpan? scaleOutCooldown = update.ScaleOutCooldown == null ? null : TimeSpan.FromSeconds(update.ScaleOutCooldown.Value);
 			TimeSpan? scaleInCooldown = update.ScaleInCooldown == null ? null : TimeSpan.FromSeconds(update.ScaleInCooldown.Value);
 
 			await _poolService.UpdatePoolAsync(pool, update.Name, update.Condition, update.EnableAutoscaling,
-				update.MinAgents, update.NumReserveAgents, update.Properties, scaleOutCooldown, scaleInCooldown, update.SizeStrategy,
+				update.MinAgents, update.NumReserveAgents, update.Properties, conformInterval, scaleOutCooldown, scaleInCooldown, update.SizeStrategy,
 				update.LeaseUtilizationSettings?.Convert(), update.JobQueueSettings?.Convert(), update.UseDefaultStrategy);
 			return new OkResult();
 		}
diff --git a/Engine/Source/Programs/Horde/Horde.Build/Tasks/ConformTaskSource.cs b/Engine/Source/Programs/Horde/Horde.Build/Tasks/ConformTaskSource.cs
--- a/Engine/Source/Programs/Horde/Horde.Build/Tasks/ConformTaskSource.cs
+++ b/Engine/Source/Programs/Horde/Horde.Build/Tasks/ConformTaskSource.cs
@@ -330,8 +330,27 @@ namespace Horde.Build.Tasks
 					return !IsConformCoolDownPeriod(agent, utcNow);
 				}
 
+				// Get the current pools for the agent
+				List<IPool> pools = await _poolService.GetCachedPoolsAsync(agent, DateTime.UtcNow - TimeSpan.FromMinutes(2.0));
+
+				TimeSpan? conformInterval = null;
+				foreach(IPool pool in pools)
+				{
+					TimeSpan interval = pool.ConformInterval ?? TimeSpan.FromDays(1.0);
+					if (interval > TimeSpan.Zero && (conformInterval == null || interval < conformInterval.Value))
+					{
+						conformInterval = interval;
+					}
+				}
+
+				// If there is no conform interval, early out
+				if(conformInterval == null)
+				{
+					return false;
+				}
+
 				// Always run a conform every 24h
-				if (utcNow > agent.LastConformTime + TimeSpan.FromDays(1.0))
+				if (utcNow > agent.LastConformTime + conformInterval.Value)
 				{
 					return true;
 				}
