Horde: Add conform interval to pool editor
#rnx
#skipci

diff --git a/Engine/Source/Programs/Horde/HordeDashboard/src/backend/Api.ts b/Engine/Source/Programs/Horde/HordeDashboard/src/backend/Api.ts
--- a/Engine/Source/Programs/Horde/HordeDashboard/src/backend/Api.ts
+++ b/Engine/Source/Programs/Horde/HordeDashboard/src/backend/Api.ts
@@ -501,6 +501,9 @@ export type UpdatePoolRequest = {
 
 	/** Whether to enable autoscaling for this pool */
 	enableAutoscaling?: boolean;
+
+	/// Frequency to run conforms, in hours, set 0 to disable
+	conformInterval?: number;
 		
 	/** Pool sizing strategy */
 	sizeStrategy?: PoolSizeStrategy;
@@ -787,6 +790,9 @@ export type GetPoolResponse = {
 	/// Whether to enable autoscaling for this pool
 	enableAutoscaling: boolean;
 
+	/// Frequency to run conforms, in hours, set 0 to disable
+	conformInterval?: number;
+		
 	/// Cooldown time between scale-out events in seconds
 	scaleOutCooldown?: number;
 
