Horde: Rewrite bundle create/extract commands to work with new bundle structures.

diff --git a/Engine/Source/Programs/Horde/Horde.Agent/Commands/Bundles/CreateCommand.cs b/Engine/Source/Programs/Horde/Horde.Agent/Commands/Bundles/CreateCommand.cs
--- a/Engine/Source/Programs/Horde/Horde.Agent/Commands/Bundles/CreateCommand.cs
+++ b/Engine/Source/Programs/Horde/Horde.Agent/Commands/Bundles/CreateCommand.cs
@@ -1,78 +1,150 @@
 // Copyright Epic Games, Inc. All Rights Reserved.
 
 using System;
+using System.Buffers;
+using System.Collections.Generic;
+using System.Threading;
 using System.Threading.Tasks;
 using EpicGames.Core;
-using EpicGames.Horde.Bundles;
-using EpicGames.Horde.Bundles.Nodes;
 using EpicGames.Horde.Storage;
-using EpicGames.Horde.Storage.Impl;
-using EpicGames.Serialization;
+using EpicGames.Horde.Storage.Bundles;
+using EpicGames.Horde.Storage.Nodes;
 using Microsoft.Extensions.Caching.Memory;
-using Microsoft.Extensions.Configuration;
-using Microsoft.Extensions.DependencyInjection;
 using Microsoft.Extensions.Logging;
 
 namespace Horde.Agent.Commands.Bundles
 {
 	abstract class BundleCommandBase : Command
 	{
-		[CommandLine("-Server=")]
-		public string? Server { get; set; } = null;
+		class FileBlob : IBlob
+		{
+			readonly IReadOnlyList<BlobId> _references;
+			readonly ReadOnlyMemory<byte> _data;
 
-		[CommandLine("-StorageDir=", Description = "Overrides the default storage server with a local directory")]
-		public DirectoryReference? StorageDir { get; set; } = null;
+			public FileBlob(IReadOnlyList<BlobId> references, ReadOnlyMemory<byte> data)
+			{
+				_references = references;
+				_data = data;
+			}
+
+			public static async Task<FileBlob> ReadAsync(FileReference file, CancellationToken cancellationToken)
+			{
+				byte[] bytes = await FileReference.ReadAllBytesAsync(file, cancellationToken);
+				MemoryReader reader = new MemoryReader(bytes);
+				IReadOnlyList<BlobId> references = reader.ReadVariableLengthArray(() => new BlobId(reader.ReadUtf8String()));
+				ReadOnlyMemory<byte> data = reader.ReadVariableLengthBytes();
+				return new FileBlob(references, data);
+			}
+
+			public async Task WriteAsync(FileReference file, CancellationToken cancellationToken)
+			{
+				ByteArrayBuilder writer = new ByteArrayBuilder();
+				writer.WriteVariableLengthArray(_references, x => writer.WriteUtf8String(x.Inner));
+				writer.WriteVariableLengthBytes(_data.Span);
+				await FileReference.WriteAllBytesAsync(file, writer.ToByteArray(), cancellationToken);
+			}
+
+			/// <inheritdoc/>
+			public ValueTask<ReadOnlyMemory<byte>> GetDataAsync() => new ValueTask<ReadOnlyMemory<byte>>(_data);
 
-		protected IStorageClient CreateStorageClient(ILogger logger)
+			/// <inheritdoc/>
+			public ValueTask<IReadOnlyList<BlobId>> GetReferencesAsync() => new ValueTask<IReadOnlyList<BlobId>>(_references);
+		}
+
+		class FileBlobStore : IBlobStore
 		{
-			if (StorageDir != null)
+			readonly DirectoryReference _rootDir;
+			readonly ILogger _logger;
+
+			public FileBlobStore(DirectoryReference rootDir, ILogger logger)
+			{
+				_rootDir = rootDir;
+				_logger = logger;
+
+				DirectoryReference.CreateDirectory(_rootDir);
+			}
+
+			FileReference GetRefFile(RefId id) => FileReference.Combine(_rootDir, id.Hash.ToString() + ".ref");
+			FileReference GetBlobFile(BlobId id) => FileReference.Combine(_rootDir, id.Inner.ToString() + ".blob");
+
+			public async Task<IBlob> ReadBlobAsync(BlobId id, CancellationToken cancellationToken = default)
+			{
+				FileReference file = GetBlobFile(id);
+				_logger.LogInformation("Reading {File}", file);
+				return await FileBlob.ReadAsync(file, cancellationToken);
+			}
+
+			public Task<bool> DeleteRefAsync(RefId id, CancellationToken cancellationToken = default)
 			{
-				return new FileStorageClient(StorageDir, logger);
+				throw new NotImplementedException();
 			}
-			else
+
+			public Task<bool> HasRefAsync(RefId id, CancellationToken cancellationToken = default)
 			{
-				IConfiguration config = new ConfigurationBuilder()
-					.AddJsonFile("appsettings.json", optional: false)
-					.AddJsonFile($"appsettings.{Environment.GetEnvironmentVariable("DOTNET_ENVIRONMENT")}.json", optional: true) // environment variable overrides, also used in k8s setups with Helm
-					.AddJsonFile("appsettings.User.json", optional: true)
-					.AddEnvironmentVariables()
-					.Build();
+				throw new NotImplementedException();
+			}
 
-				IConfigurationSection configSection = config.GetSection(AgentSettings.SectionName);
+			public async Task<IBlob> ReadRefAsync(RefId id, CancellationToken cancellationToken = default)
+			{
+				FileReference file = GetRefFile(id);
+				_logger.LogInformation("Reading {File}", file);
+				return await FileBlob.ReadAsync(file, cancellationToken);
+			}
 
-				IServiceCollection services = new ServiceCollection();
-				services.AddHordeStorage(settings => configSection.GetCurrentServerProfile().GetSection(nameof(ServerProfile.Storage)).Bind(settings));
+			public async Task<BlobId> WriteBlobAsync(ReadOnlySequence<byte> data, IReadOnlyList<BlobId> references, CancellationToken cancellationToken = default)
+			{
+				BlobId id = new BlobId(Guid.NewGuid().ToString());
+				FileReference file = GetBlobFile(id);
+				_logger.LogInformation("Writing {File}", file);
+				await new FileBlob(references, data.AsSingleSegment()).WriteAsync(file, cancellationToken);
+				return id;
+			}
 
-				IServiceProvider serviceProvider = services.BuildServiceProvider();
-				return serviceProvider.GetRequiredService<IStorageClient>();
+			public async Task WriteRefAsync(RefId id, ReadOnlySequence<byte> data, IReadOnlyList<BlobId> references, CancellationToken cancellationToken = default)
+			{
+				FileReference file = GetRefFile(id);
+				_logger.LogInformation("Writing {File}", file);
+				await new FileBlob(references, data.AsSingleSegment()).WriteAsync(file, cancellationToken);
 			}
 		}
+
+		public static RefId DefaultRefId { get; } = new RefId("default-ref");
+
+		[CommandLine("-StorageDir=", Description = "Overrides the default storage server with a local directory")]
+		public DirectoryReference StorageDir { get; set; } = DirectoryReference.Combine(Program.AppDir, "bundles");
+
+		IBlobStore? _blobStore;
+
+		protected IBlobStore CreateBlobStore(ILogger logger)
+		{
+			_blobStore ??= new FileBlobStore(StorageDir, logger);
+			return _blobStore;
+		}
+
+		protected BundleStore<T> CreateTreeStore<T>(ILogger logger, IMemoryCache cache) where T : TreeNode
+		{
+			IBlobStore blobStore = CreateBlobStore(logger);
+			return new BundleStore<T>(blobStore, new BundleOptions(), cache);
+		}
 	}
 
 	[Command("bundle", "create", "Creates a bundle from a folder on the local hard drive")]
 	class CreateCommand : BundleCommandBase
 	{
-		[CommandLine("-Namespace=")]
-		public NamespaceId NamespaceId { get; set; } = new NamespaceId("default-ns");
-
-		[CommandLine("-Bucket=")]
-		public BucketId BucketId { get; set; } = new BucketId("default-bkt");
-
 		[CommandLine("-Ref=")]
-		public RefId RefId { get; set; } = new RefId("default-ref");
+		public RefId RefId { get; set; } = DefaultRefId;
 
 		[CommandLine("-InputDir=", Required = true)]
 		public DirectoryReference InputDir { get; set; } = null!;
 
 		public override async Task<int> ExecuteAsync(ILogger logger)
 		{
-			IStorageClient storageClient = base.CreateStorageClient(logger);
-			using (MemoryCache cache = new MemoryCache(new MemoryCacheOptions { SizeLimit = 50 * 1024 * 1024 }))
-			{
-				using Bundle<DirectoryNode> bundle = Bundle.Create<DirectoryNode>(storageClient, NamespaceId, new BundleOptions(), cache);
-				await bundle.Root.CopyFromDirectoryAsync(InputDir.ToDirectoryInfo(), new ChunkingOptions(), logger);
-				await bundle.WriteAsync(BucketId, RefId, CbObject.Empty, false);
-			}
+			using IMemoryCache cache = new MemoryCache(new MemoryCacheOptions());
+			using ITreeStore<DirectoryNode> store = CreateTreeStore<DirectoryNode>(logger, cache);
+
+			DirectoryNode node = new DirectoryNode();
+			await node.CopyFromDirectoryAsync(InputDir.ToDirectoryInfo(), new ChunkingOptions(), logger, CancellationToken.None);
+			await store.WriteTreeAsync(RefId, node);
 
 			return 0;
 		}
diff --git a/Engine/Source/Programs/Horde/Horde.Agent/Commands/Bundles/DeltaCommand.cs b/Engine/Source/Programs/Horde/Horde.Agent/Commands/Bundles/DeltaCommand.cs
deleted file mode 100644
diff --git a/Engine/Source/Programs/Horde/Horde.Agent/Commands/Bundles/DumpCommand.cs b/Engine/Source/Programs/Horde/Horde.Agent/Commands/Bundles/DumpCommand.cs
--- a/Engine/Source/Programs/Horde/Horde.Agent/Commands/Bundles/DumpCommand.cs
+++ b/Engine/Source/Programs/Horde/Horde.Agent/Commands/Bundles/DumpCommand.cs
@@ -1,34 +1,96 @@
 // Copyright Epic Games, Inc. All Rights Reserved.
 
+using System;
+using System.Collections.Generic;
+using System.Linq;
+using System.Threading;
 using System.Threading.Tasks;
 using EpicGames.Core;
-using EpicGames.Horde.Bundles;
 using EpicGames.Horde.Storage;
+using EpicGames.Horde.Storage.Bundles;
 using Microsoft.Extensions.Logging;
 
 namespace Horde.Agent.Commands.Bundles
 {
-	[Command("bundle", "dump", "Indexes a tree and creates a table of contents for each blob describing the contents")]
+	[Command("bundle", "dump", "Dumps the contents of a bundle")]
 	internal class DumpCommand : BundleCommandBase
 	{
-		[CommandLine("-Namespace=")]
-		public NamespaceId NamespaceId { get; set; } = new NamespaceId("default-ns");
-
-		[CommandLine("-Bucket=")]
-		public BucketId BucketId { get; set; } = new BucketId("default-bkt");
-
 		[CommandLine("-Ref=")]
-		public RefId RefId { get; set; } = new RefId("default-ref");
+		public RefId RefId { get; set; } = DefaultRefId;
 
-		[CommandLine("-OutputDir=")]
-		public DirectoryReference? OutputDir { get; set; }
+		[CommandLine("-Blob=")]
+		public BlobId? BlobId { get; set; }
 
 		public override async Task<int> ExecuteAsync(ILogger logger)
 		{
-			OutputDir = DirectoryReference.Combine(Program.DataDir, "Bundles");
+			IBlobStore blobStore = CreateBlobStore(logger);
+
+			IBlob blob;
+			if (BlobId != null)
+			{
+				blob = await blobStore.ReadBlobAsync(BlobId.Value, CancellationToken.None);
+				logger.LogInformation("Summary for blob {BlobId}", BlobId.Value);
+			}
+			else
+			{
+				blob = await blobStore.ReadRefAsync(RefId, CancellationToken.None);
+				logger.LogInformation("Summary for ref {RefId}", RefId);
+			}
+
+			IReadOnlyList<BlobId> references = await blob.GetReferencesAsync();
+			logger.LogInformation("");
+			logger.LogInformation("BlobRefs: {NumRefs}", references.Count);
+			foreach (BlobId reference in references)
+			{
+				logger.LogInformation("  {BlobId}", reference);
+			}
+
+			ReadOnlyMemory<byte> data = await blob.GetDataAsync();
+			MemoryReader reader = new MemoryReader(data);
+			BundleHeader header = new BundleHeader(reader);
+			int packetStart = data.Length - reader.Memory.Length;
+
+			logger.LogInformation("");
+			logger.LogInformation("Imports: {NumImports}", header.Imports.Count);
+
+			int refIdx = 0;
+			foreach (BundleImport import in header.Imports)
+			{
+				logger.LogInformation("  From blob {BlobId} ({NumExports}/{TotalExports} nodes)", import.BlobId, import.Exports.Count, import.ExportCount);
+				foreach ((int exportIdx, IoHash exportHash) in import.Exports)
+				{
+					logger.LogInformation("    [{Index}] IMP {BlobId}:{ExportIdx} = {ExportHash}", refIdx, import.BlobId, exportIdx, exportHash);
+					refIdx++;
+				}
+			}
+
+			logger.LogInformation("");
+			logger.LogInformation("Exports: {NumExports}", header.Exports.Count);
+
+			int packetIdx = 0;
+			int packetOffset = 0;
+			foreach (BundleExport export in header.Exports)
+			{
+				string refs = (export.References.Count == 0) ? "[ ]" : $"[ {String.Join(", ", export.References.Select(x => x.ToString()))} ]";
+				logger.LogInformation("  [{Index}] EXP {ExportHash} (length: {NumBytes:n0}, packet: {PacketIdx}, refs: {Refs})", refIdx, export.Hash, export.Length, packetIdx, refs);
+				refIdx++;
+
+				packetOffset += export.Length;
+				if(packetOffset >= header.Packets[packetIdx].DecodedLength)
+				{
+					packetIdx++;
+					packetOffset = 0;
+				}
+			}
 
-			IStorageClient storageClient = CreateStorageClient(logger);
-			await BundleTools.WriteSummaryAsync(storageClient, NamespaceId, BucketId, RefId, OutputDir);
+			logger.LogInformation("");
+			logger.LogInformation("Packets: {NumPackets}", header.Packets.Count);
+			for (int idx = 0; idx < header.Packets.Count; idx++)
+			{
+				BundlePacket packet = header.Packets[idx];
+				logger.LogInformation("  PKT {Idx} (file offset: {Offset:n0}, encoded: {EncodedLength:n0}, decoded: {DecodedLength:n0}, ratio: {Ratio}%)", idx, packetStart, packet.EncodedLength, packet.DecodedLength, (int)(packet.EncodedLength * 100) / packet.DecodedLength);
+				packetStart += packet.EncodedLength;
+			}
 
 			return 0;
 		}
diff --git a/Engine/Source/Programs/Horde/Horde.Agent/Commands/Bundles/ExtractCommand.cs b/Engine/Source/Programs/Horde/Horde.Agent/Commands/Bundles/ExtractCommand.cs
--- a/Engine/Source/Programs/Horde/Horde.Agent/Commands/Bundles/ExtractCommand.cs
+++ b/Engine/Source/Programs/Horde/Horde.Agent/Commands/Bundles/ExtractCommand.cs
@@ -1,10 +1,10 @@
 // Copyright Epic Games, Inc. All Rights Reserved.
 
+using System.Threading;
 using System.Threading.Tasks;
 using EpicGames.Core;
-using EpicGames.Horde.Bundles;
-using EpicGames.Horde.Bundles.Nodes;
 using EpicGames.Horde.Storage;
+using EpicGames.Horde.Storage.Nodes;
 using Microsoft.Extensions.Caching.Memory;
 using Microsoft.Extensions.Logging;
 
@@ -13,26 +13,20 @@ namespace Horde.Agent.Commands.Bundles
 	[Command("bundle", "extract", "Extracts data from a bundle to the local hard drive")]
 	internal class ExtractCommand : BundleCommandBase
 	{
-		[CommandLine("-Namespace=")]
-		public NamespaceId NamespaceId { get; set; } = new NamespaceId("default-ns");
-
-		[CommandLine("-Bucket=")]
-		public BucketId BucketId { get; set; } = new BucketId("default-bkt");
-
 		[CommandLine("-Ref=")]
-		public RefId RefId { get; set; } = new RefId("default-ref");
+		public RefId RefId { get; set; } = DefaultRefId;
 
 		[CommandLine("-OutputDir=", Required = true)]
 		public DirectoryReference OutputDir { get; set; } = null!;
 
 		public override async Task<int> ExecuteAsync(ILogger logger)
 		{
-			IStorageClient storageClient = CreateStorageClient(logger);
-			using (IMemoryCache cache = new MemoryCache(new MemoryCacheOptions { SizeLimit = 1024 * 1024 * 1000 }))
-			{
-				using Bundle<DirectoryNode> bundle = await Bundle.ReadAsync<DirectoryNode>(storageClient, NamespaceId, BucketId, RefId, new BundleOptions(), cache);
-				await bundle.Root.CopyToDirectoryAsync(bundle, OutputDir.ToDirectoryInfo(), logger);
-			}
+			using IMemoryCache cache = new MemoryCache(new MemoryCacheOptions());
+			using ITreeStore<DirectoryNode> store = CreateTreeStore<DirectoryNode>(logger, cache);
+
+			DirectoryNode node = await store.ReadTreeAsync(RefId);
+			await node.CopyToDirectoryAsync(OutputDir.ToDirectoryInfo(), logger, CancellationToken.None);
+
 			return 0;
 		}
 	}
diff --git a/Engine/Source/Programs/Horde/Horde.Build/Commands/Bundles/CreateCommand.cs b/Engine/Source/Programs/Horde/Horde.Build/Commands/Bundles/CreateCommand.cs
deleted file mode 100644
diff --git a/Engine/Source/Programs/Horde/Horde.Build/Commands/Bundles/DeltaCommand.cs b/Engine/Source/Programs/Horde/Horde.Build/Commands/Bundles/DeltaCommand.cs
deleted file mode 100644
diff --git a/Engine/Source/Programs/Horde/Horde.Build/Commands/Bundles/DumpCommand.cs b/Engine/Source/Programs/Horde/Horde.Build/Commands/Bundles/DumpCommand.cs
deleted file mode 100644
diff --git a/Engine/Source/Programs/Horde/Horde.Build/Commands/Bundles/ExtractCommand.cs b/Engine/Source/Programs/Horde/Horde.Build/Commands/Bundles/ExtractCommand.cs
deleted file mode 100644
diff --git a/Engine/Source/Programs/Shared/EpicGames.Horde/Storage/ITreeStore.cs b/Engine/Source/Programs/Shared/EpicGames.Horde/Storage/ITreeStore.cs
--- a/Engine/Source/Programs/Shared/EpicGames.Horde/Storage/ITreeStore.cs
+++ b/Engine/Source/Programs/Shared/EpicGames.Horde/Storage/ITreeStore.cs
@@ -47,7 +47,7 @@ namespace EpicGames.Horde.Storage
 	/// <summary>
 	/// Base interface for a tree store
 	/// </summary>
-	public interface ITreeStore
+	public interface ITreeStore : IDisposable
 	{
 		/// <summary>
 		/// Tests whether a tree exists with the given name
diff --git a/Engine/Source/Programs/Shared/EpicGames.Horde/Storage/Nodes/DirectoryNode.cs b/Engine/Source/Programs/Shared/EpicGames.Horde/Storage/Nodes/DirectoryNode.cs
--- a/Engine/Source/Programs/Shared/EpicGames.Horde/Storage/Nodes/DirectoryNode.cs
+++ b/Engine/Source/Programs/Shared/EpicGames.Horde/Storage/Nodes/DirectoryNode.cs
@@ -616,12 +616,14 @@ namespace EpicGames.Horde.Storage.Nodes
 			{
 				FileInfo fileInfo = new FileInfo(Path.Combine(directoryInfo.FullName, fileEntry.Name.ToString()));
 				FileNode fileNode = await fileEntry.ExpandAsync(cancellationToken);
+				logger.LogInformation("Writing {File}", fileInfo.FullName);
 				tasks.Add(Task.Run(() => fileNode.CopyToFileAsync(fileInfo, cancellationToken), cancellationToken));
 			}
 			foreach (DirectoryEntry directoryEntry in Directories)
 			{
 				DirectoryInfo subDirectoryInfo = directoryInfo.CreateSubdirectory(directoryEntry.Name.ToString());
 				DirectoryNode subDirectoryNode = await directoryEntry.ExpandAsync(cancellationToken);
+				logger.LogInformation("Writing {Dir}", subDirectoryInfo.FullName);
 				tasks.Add(Task.Run(() => subDirectoryNode.CopyToDirectoryAsync(subDirectoryInfo, logger, cancellationToken), cancellationToken));
 			}
 
