Horde: Include a link to the Slack discussion thread when DM'ing a user due to an issue assignment.

diff --git a/Engine/Source/Programs/Horde/Horde.Build/Notifications/Sinks/SlackNotificationSink.cs b/Engine/Source/Programs/Horde/Horde.Build/Notifications/Sinks/SlackNotificationSink.cs
--- a/Engine/Source/Programs/Horde/Horde.Build/Notifications/Sinks/SlackNotificationSink.cs
+++ b/Engine/Source/Programs/Horde/Horde.Build/Notifications/Sinks/SlackNotificationSink.cs
@@ -983,7 +983,23 @@ namespace Horde.Build.Notifications.Sinks
 			attachment.Blocks.Add(new SectionBlock($"*<{issueUrl}|Issue #{issue.Id}: {issue.Summary}>*"));
 
 			string streamList = StringUtils.FormatList(details.Spans.Select(x => $"*{x.StreamName}*").Distinct(StringComparer.OrdinalIgnoreCase).OrderBy(x => x, StringComparer.OrdinalIgnoreCase));
-			attachment.Blocks.Add(new SectionBlock($"Occurring in {streamList}"));
+			StringBuilder summaryBuilder = new StringBuilder($"Occurring in {streamList}.");
+
+			IIssueSpan span = details.Spans[0];
+			WorkflowId? workflowId = span.LastFailure.Annotations.WorkflowId;
+			if (workflowId != null)
+			{
+				IStream? stream = await _streamService.GetStreamAsync(span.StreamId);
+				if (stream != null && stream.Config.TryGetWorkflow(workflowId.Value, out WorkflowConfig? workflow) && workflow.TriageChannel != null)
+				{
+					MessageStateDocument? state = await GetMessageStateAsync(workflow.TriageChannel, GetTriageThreadEventId(issue.Id));
+					if (state != null)
+					{
+						summaryBuilder.Append($" See *<{state.Permalink}|discussion thread>*.");
+					}
+				}
+			}
+			attachment.Blocks.Add(new SectionBlock(summaryBuilder.ToString()));
 
 			IIssueSpan? lastSpan = details.Spans.OrderByDescending(x => x.LastFailure.StepTime).FirstOrDefault();
 			if (lastSpan != null && lastSpan.LastFailure.LogId != null)
