Horde: Update last issue summary with state of each issue.

diff --git a/Engine/Source/Programs/Horde/Horde.Build/Issues/IssueReportService.cs b/Engine/Source/Programs/Horde/Horde.Build/Issues/IssueReportService.cs
--- a/Engine/Source/Programs/Horde/Horde.Build/Issues/IssueReportService.cs
+++ b/Engine/Source/Programs/Horde/Horde.Build/Issues/IssueReportService.cs
@@ -32,6 +32,7 @@ namespace Horde.Build.Issues
 	public class IssueReport
 	{
 		public string Channel { get; }
+		public WorkflowId WorkflowId { get; }
 		public DateTimeOffset Time { get; }
 		public IStream Stream { get; }
 		public WorkflowStats WorkflowStats { get; }
@@ -40,9 +41,10 @@ namespace Horde.Build.Issues
 		public List<IIssueSpan> IssueSpans { get; } = new List<IIssueSpan>();
 		public bool GroupByTemplate { get; }
 
-		public IssueReport(string channel, DateTimeOffset time, IStream stream, WorkflowStats workflowStats, string? triageChannel, bool groupByTemplate)
+		public IssueReport(string channel, WorkflowId workflowId, DateTimeOffset time, IStream stream, WorkflowStats workflowStats, string? triageChannel, bool groupByTemplate)
 		{
 			Channel = channel;
+			WorkflowId = workflowId;
 			Time = time;
 			Stream = stream;
 			WorkflowStats = workflowStats;
@@ -155,7 +157,7 @@ namespace Horde.Build.Issues
 							workflowStats = new WorkflowStats();
 						}
 
-						IssueReport report = new IssueReport(workflow.ReportChannel, currentTime, stream, workflowStats, workflow.TriageChannel, workflow.GroupIssuesByTemplate);
+						IssueReport report = new IssueReport(workflow.ReportChannel, workflow.Id, currentTime, stream, workflowStats, workflow.TriageChannel, workflow.GroupIssuesByTemplate);
 						foreach (IIssueSpan span in spans)
 						{
 							if (span.LastSuccess != null && span.LastFailure.Annotations.WorkflowId == workflow.Id)
diff --git a/Engine/Source/Programs/Horde/Horde.Build/Notifications/Sinks/SlackNotificationSink.cs b/Engine/Source/Programs/Horde/Horde.Build/Notifications/Sinks/SlackNotificationSink.cs
--- a/Engine/Source/Programs/Horde/Horde.Build/Notifications/Sinks/SlackNotificationSink.cs
+++ b/Engine/Source/Programs/Horde/Horde.Build/Notifications/Sinks/SlackNotificationSink.cs
@@ -235,6 +235,7 @@ namespace Horde.Build.Notifications.Sinks
 		readonly IMongoCollection<SlackUser> _slackUsers;
 		readonly HashSet<string>? _allowUsers;
 		readonly IExternalIssueService _externalIssueService;
+		readonly JsonSerializerOptions _jsonSerializerOptions;
 		readonly ILogger _logger;
 
 		/// <summary>
@@ -259,6 +260,9 @@ namespace Horde.Build.Notifications.Sinks
 			_slackUsers = mongoService.Database.GetCollection<SlackUser>("Slack.UsersV2");
 			_logger = logger;
 
+			_jsonSerializerOptions = new JsonSerializerOptions();
+			Startup.ConfigureJsonSerializer(_jsonSerializerOptions);
+
 			if (!String.IsNullOrEmpty(settings.Value.SlackUsers))
 			{
 				_allowUsers = new HashSet<string>(settings.Value.SlackUsers.Split(','), StringComparer.OrdinalIgnoreCase);
@@ -287,13 +291,18 @@ namespace Horde.Build.Notifications.Sinks
 
 		#region Message state 
 
-		async Task<(MessageStateDocument, bool)> AddOrUpdateMessageStateAsync(string recipient, string eventId, UserId? userId, string digest)
+		async Task<(MessageStateDocument, bool)> AddOrUpdateMessageStateAsync(string recipient, string eventId, UserId? userId, string digest, string? ts = null)
 		{
 			ObjectId newId = ObjectId.GenerateNewId();
 
 			FilterDefinition<MessageStateDocument> filter = Builders<MessageStateDocument>.Filter.Eq(x => x.Recipient, recipient) & Builders<MessageStateDocument>.Filter.Eq(x => x.EventId, eventId);
 			UpdateDefinition<MessageStateDocument> update = Builders<MessageStateDocument>.Update.SetOnInsert(x => x.Id, newId).Set(x => x.UserId, userId).Set(x => x.Digest, digest);
 
+			if (ts != null)
+			{
+				update = update.Set(x => x.Ts, ts);
+			}
+
 			MessageStateDocument state = await _messageStates.FindOneAndUpdateAsync(filter, update, new FindOneAndUpdateOptions<MessageStateDocument> { IsUpsert = true, ReturnDocument = ReturnDocument.After });
 			return (state, state.Id == newId);
 		}
@@ -723,6 +732,8 @@ namespace Horde.Build.Notifications.Sinks
 					await SendIssueMessageAsync(channel, issue, details, null, defaultAllowMentions);
 				}
 			}
+
+			await UpdateReportsAsync(issue, details.Spans);
 		}
 
 		async Task AddReactionAsync(string channel, string ts, string name)
@@ -1257,27 +1268,37 @@ namespace Horde.Build.Notifications.Sinks
 			return $"<!date^{time.ToUnixTimeSeconds()}^{{time}}|{time}>";
 		}
 
+		class ReportBlock
+		{
+			[JsonPropertyName("t")]
+			public TemplateRefId TemplateId { get; set; }
+
+			[JsonPropertyName("i")]
+			public List<int> IssueIds { get; set; } = new List<int>();
+		}
+
+		class ReportState
+		{
+			[JsonPropertyName("tm")]
+			public DateTime Time { get; set; }
+
+			[JsonPropertyName("msg")]
+			public List<ReportBlock> Blocks { get; set; } = new List<ReportBlock>();
+		}
+
+		static string GetReportEventId(StreamId streamId, WorkflowId workflowId) => $"issue-report:{streamId}:{workflowId}";
+
+		static string GetReportBlockEventId(string ts, int idx) => $"issue-report-block:{ts}:{idx}";
+
 		/// <inheritdoc/>
 		public async Task SendIssueReportAsync(IssueReport report)
 		{
 			if (report.Channel != null)
 			{
-				List<BlockBase> blocks = new List<BlockBase>();
-				blocks.Add(new HeaderBlock(AddEnvironmentAnnotation($"{report.Stream.Name}: {report.Time:d}")));
-				await SendMessageAsync(report.Channel, blocks: blocks.ToArray(), withEnvironment: false);
-
-				string summary;
-				if (report.Issues.Count == 0)
-				{
-					summary = ":tick: No issues open.";
-				}
-				else
-				{
-					TimeSpan averageAge = TimeSpan.FromHours(report.Issues.Select(x => (report.Time - x.CreatedAt).TotalHours).Average());
-					summary = $"*{report.Issues.Count} unique issues* currently open (average age {FormatReadableTimeSpan(averageAge)}).";
-				}
-				await SendMessageAsync(report.Channel, text: summary, withEnvironment: false);
+				ReportState state = new ReportState();
+				state.Time = report.Time.UtcDateTime;
 
+				List<List<(IIssue, IIssueSpan?)>> issuesByBlock = new List<List<(IIssue, IIssueSpan?)>>();
 				if (report.GroupByTemplate)
 				{
 					Dictionary<int, IIssue> issueIdToInfo = new Dictionary<int, IIssue>();
@@ -1305,50 +1326,179 @@ namespace Horde.Build.Notifications.Sinks
 
 						if (pairs.Count > 0)
 						{
-							string template = templateConfig.Name;
+							List<(IIssue Issue, IIssueSpan? Span)> sortedPairs = pairs.OrderBy(x => x.Key.Id).Select(x => (x.Key, (IIssueSpan?)x.Value)).ToList();
 
-							CreateJobsTabRequest? tab = report.Stream.Config.Tabs.OfType<CreateJobsTabRequest>().FirstOrDefault(x => x.Templates != null && x.Templates.Contains(templateConfig.Id));
-							if (tab != null)
-							{
-								Uri templateUrl = new Uri(_settings.DashboardUrl, $"stream/{report.Stream.Id}?tab={tab.Title}&template={templateConfig.Id}");
-								template = $"<{templateUrl}|{template}>";
-							}
+							ReportBlock block = new ReportBlock();
+							block.TemplateId = group.Key;
+							block.IssueIds.AddRange(sortedPairs.Select(x => x.Issue.Id));
+							state.Blocks.Add(block);
 
-							StringBuilder body = new StringBuilder($"*{template}*:");
-							foreach ((IIssue issue, IIssueSpan span) in pairs.OrderBy(x => x.Key.Id))
-							{
-								body.Append('\n');
-								body.Append(await FormatIssueAsync(issue, span, report));
-							}
-							await SendMessageAsync(report.Channel, text: body.ToString(), withEnvironment: false);
+							issuesByBlock.Add(sortedPairs);
 						}
 					}
 				}
 				else
 				{
-					StringBuilder body = new StringBuilder();
+					List<(IIssue Issue, IIssueSpan? Span)> sortedPairs = new List<(IIssue, IIssueSpan?)>();
 					foreach (IIssue issue in report.Issues.OrderByDescending(x => x.Id))
 					{
 						IIssueSpan? span = report.IssueSpans.FirstOrDefault(x => x.IssueId == issue.Id);
-						if (body.Length > 0)
+						sortedPairs.Add((issue, span));
+					}
+
+					ReportBlock block = new ReportBlock();
+					block.IssueIds.AddRange(sortedPairs.Select(x => x.Issue.Id));
+					state.Blocks.Add(block);
+				}
+
+				List<BlockBase> blocks = new List<BlockBase>();
+				blocks.Add(new HeaderBlock(AddEnvironmentAnnotation($"{report.Stream.Name}: {report.Time:d}")));
+
+				PostMessageResponse? response = await SendMessageAsync(report.Channel, blocks: blocks.ToArray(), withEnvironment: false);
+				if (response != null && response.Ts != null)
+				{
+					string reportEventId = GetReportEventId(report.Stream.Id, report.WorkflowId);
+					string json = JsonSerializer.Serialize(state, _jsonSerializerOptions);
+					await AddOrUpdateMessageStateAsync(report.Channel, reportEventId, null, json, response.Ts);
+
+					string summary;
+					if (report.Issues.Count == 0)
+					{
+						summary = ":tick: No issues open.";
+					}
+					else
+					{
+						TimeSpan averageAge = TimeSpan.FromHours(report.Issues.Select(x => (report.Time - x.CreatedAt).TotalHours).Average());
+						summary = $"*{report.Issues.Count} unique issues* currently open (average age {FormatReadableTimeSpan(averageAge)}).";
+					}
+					await SendMessageAsync(report.Channel, text: summary, withEnvironment: false);
+
+					for (int idx = 0; idx < state.Blocks.Count; idx++)
+					{
+						string blockEventId = GetReportBlockEventId(response.Ts, idx);
+						await UpdateReportBlockAsync(report.Channel, blockEventId, report.Time.UtcDateTime, report.Stream, state.Blocks[idx].TemplateId, issuesByBlock[idx]);
+					}
+
+					if (report.WorkflowStats.NumSteps > 0)
+					{
+						double totalPct = (report.WorkflowStats.NumPassingSteps * 100.0) / report.WorkflowStats.NumSteps;
+						string footer = $"{report.WorkflowStats.NumPassingSteps:n0} of {report.WorkflowStats.NumSteps:n0} (*{totalPct:0.0}%*) build steps succeeded since last status update.";
+						await SendMessageAsync(report.Channel, text: footer, withEnvironment: false);
+					}
+				}
+			}
+		}
+
+		async Task UpdateReportsAsync(IIssue issue, IReadOnlyList<IIssueSpan> spans)
+		{
+			HashSet<string> reportEventIds = new HashSet<string>(StringComparer.Ordinal);
+			foreach (IIssueSpan span in spans)
+			{
+				WorkflowId? workflowId = span.LastFailure.Annotations.WorkflowId;
+				if (workflowId == null)
+				{
+					continue;
+				}
+
+				string reportEventId = GetReportEventId(span.StreamId, workflowId.Value);
+				if (!reportEventIds.Add(reportEventId))
+				{
+					continue;
+				}
+
+				IStream? stream = await _streamService.GetStreamAsync(span.StreamId);
+				if (stream == null)
+				{
+					continue;
+				}
+
+				WorkflowConfig? workflow;
+				if (!stream.Config.TryGetWorkflow(workflowId.Value, out workflow) || workflow.ReportChannel == null)
+				{
+					continue;
+				}
+
+				MessageStateDocument? messageState = await GetMessageStateAsync(workflow.ReportChannel, reportEventId);
+				if (messageState == null)
+				{
+					continue;
+				}
+
+				ReportState? state;
+				try
+				{
+					state = JsonSerializer.Deserialize<ReportState>(messageState.Digest, _jsonSerializerOptions);
+				}
+				catch
+				{
+					continue;
+				}
+
+				if (state == null)
+				{
+					continue;
+				}
+
+				for (int idx = 0; idx < state.Blocks.Count; idx++)
+				{
+					ReportBlock block = state.Blocks[idx];
+					if (block.IssueIds.Contains(issue.Id))
+					{
+						string blockEventId = GetReportBlockEventId(messageState.Ts, idx);
+
+						List<(IIssue, IIssueSpan?)> pairs = new List<(IIssue, IIssueSpan?)>();
+						foreach (int issueId in block.IssueIds)
 						{
-							body.Append('\n');
+							IIssueDetails? details = await _issueService.GetIssueDetailsAsync(issueId);
+							if (details != null)
+							{
+								IIssueSpan? otherSpan = details.Spans.FirstOrDefault(x => x.TemplateRefId == block.TemplateId);
+								if (otherSpan == null && details.Spans.Count > 0)
+								{
+									otherSpan = details.Spans[0];
+								}
+								pairs.Add((details.Issue, otherSpan));
+							}
 						}
-						body.Append(await FormatIssueAsync(issue, span, report));
+
+						await UpdateReportBlockAsync(workflow.ReportChannel, blockEventId, state.Time, stream, block.TemplateId, pairs);
+					}
+				}
+			}
+		}
+
+		async Task UpdateReportBlockAsync(string channel, string eventId, DateTime reportTime, IStream stream, TemplateRefId? templateId, List<(IIssue, IIssueSpan?)> pairs)
+		{
+			StringBuilder body = new StringBuilder();
+			body.Append(DateTime.UtcNow.ToString() + "\n");
+
+			if (templateId != null)
+			{
+				TemplateRefConfig? templateConfig;
+				if (stream.Config.TryGetTemplate(templateId.Value, out templateConfig))
+				{
+					CreateJobsTabRequest? tab = stream.Config.Tabs.OfType<CreateJobsTabRequest>().FirstOrDefault(x => x.Templates != null && x.Templates.Contains(templateId.Value));
+					if (tab != null)
+					{
+						Uri templateUrl = new Uri(_settings.DashboardUrl, $"stream/{stream.Id}?tab={tab.Title}&template={templateId.Value}");
+						body.Append($"*<{templateUrl}|{templateConfig.Name}>*:");
 					}
-					await SendMessageAsync(report.Channel, text: body.ToString(), withEnvironment: false);
 				}
+			}
 
-				if (report.WorkflowStats.NumSteps > 0)
+			foreach ((IIssue issue, IIssueSpan? span) in pairs)
+			{
+				if (body.Length > 0)
 				{
-					double totalPct = (report.WorkflowStats.NumPassingSteps * 100.0) / report.WorkflowStats.NumSteps;
-					string footer = $"{report.WorkflowStats.NumPassingSteps:n0} of {report.WorkflowStats.NumSteps:n0} (*{totalPct:0.0}%*) build steps succeeded since last status update.";
-					await SendMessageAsync(report.Channel, text: footer, withEnvironment: false);
+					body.Append('\n');
 				}
+				body.Append(await FormatIssueAsync(issue, span, channel, reportTime));
 			}
+
+			await SendOrUpdateMessageAsync(channel, eventId, null, body.ToString(), withEnvironment: false);
 		}
 
-		async ValueTask<string> FormatIssueAsync(IIssue issue, IIssueSpan? span, IssueReport report)
+		async ValueTask<string> FormatIssueAsync(IIssue issue, IIssueSpan? span, string? triageChannel, DateTime reportTime)
 		{
 			Uri issueUrl = _settings.DashboardUrl;
 			if (span != null)
@@ -1382,7 +1532,6 @@ namespace Horde.Build.Notifications.Sinks
 
 			StringBuilder body = new StringBuilder($"\u2022 *Issue <{issueUrl}|{issue.Id}>");
 
-			string? triageChannel = report.TriageChannel;
 			if (triageChannel != null)
 			{
 				MessageStateDocument? state = await GetMessageStateAsync(triageChannel, GetTriageThreadEventId(issue.Id));
@@ -1392,7 +1541,7 @@ namespace Horde.Build.Notifications.Sinks
 				}
 			}
 
-			body.Append($"*: {issue.Summary} [{FormatReadableTimeSpan(report.Time - issue.CreatedAt)}]");
+			body.Append($"*: {issue.Summary} [{FormatReadableTimeSpan(reportTime - issue.CreatedAt)}]");
 
 			if (!String.IsNullOrEmpty(issue.ExternalIssueKey))
 			{
@@ -1779,11 +1928,11 @@ namespace Horde.Build.Notifications.Sinks
 			}
 		}
 
-		private async Task SendMessageAsync(string recipient, string? text = null, BlockBase[]? blocks = null, BlockKitAttachment[]? attachments = null, bool withEnvironment = true)
+		private async Task<PostMessageResponse?> SendMessageAsync(string recipient, string? text = null, BlockBase[]? blocks = null, BlockKitAttachment[]? attachments = null, bool withEnvironment = true)
 		{
 			if (!IsRecipientAllowed(recipient, text))
 			{
-				return;
+				return null;
 			}
 
 			BlockKitMessage message = new BlockKitMessage();
@@ -1801,7 +1950,7 @@ namespace Horde.Build.Notifications.Sinks
 				message.Attachments.AddRange(attachments);
 			}
 
-			await SendRequestAsync<PostMessageResponse>(PostMessageUrl, message);
+			return await SendRequestAsync<PostMessageResponse>(PostMessageUrl, message);
 		}
 
 		bool IsRecipientAllowed(string recipient, string? description)
@@ -1814,11 +1963,13 @@ namespace Horde.Build.Notifications.Sinks
 			return true;
 		}
 
-		private async Task<(MessageStateDocument, bool)> SendOrUpdateMessageAsync(string recipient, string eventId, UserId? userId, string? text = null, BlockBase[]? blocks = null, BlockKitAttachment[]? attachments = null)
+		private async Task<(MessageStateDocument, bool)> SendOrUpdateMessageAsync(string recipient, string eventId, UserId? userId, string? text = null, BlockBase[]? blocks = null, BlockKitAttachment[]? attachments = null, bool withEnvironment = true)
 		{
 			BlockKitMessage message = new BlockKitMessage();
-			message.Text = AddEnvironmentAnnotation(text);
-
+			if (text != null)
+			{
+				message.Text = withEnvironment ? AddEnvironmentAnnotation(text) : text;
+			}
 			if (blocks != null)
 			{
 				message.Blocks.AddRange(blocks);
