Horde: Remove debug timestamps, add additional logging for report updates.

diff --git a/Engine/Source/Programs/Horde/Horde.Build/Notifications/Sinks/SlackNotificationSink.cs b/Engine/Source/Programs/Horde/Horde.Build/Notifications/Sinks/SlackNotificationSink.cs
--- a/Engine/Source/Programs/Horde/Horde.Build/Notifications/Sinks/SlackNotificationSink.cs
+++ b/Engine/Source/Programs/Horde/Horde.Build/Notifications/Sinks/SlackNotificationSink.cs
@@ -1391,6 +1391,8 @@ namespace Horde.Build.Notifications.Sinks
 
 		async Task UpdateReportsAsync(IIssue issue, IReadOnlyList<IIssueSpan> spans)
 		{
+			_logger.LogInformation("Checking for report updates to issue {IssueId}", issue.Id);
+
 			HashSet<string> reportEventIds = new HashSet<string>(StringComparer.Ordinal);
 			foreach (IIssueSpan span in spans)
 			{
@@ -1439,12 +1441,15 @@ namespace Horde.Build.Notifications.Sinks
 					continue;
 				}
 
+				_logger.LogInformation("Updating report {EventId}", reportEventId);
+
 				for (int idx = 0; idx < state.Blocks.Count; idx++)
 				{
 					ReportBlock block = state.Blocks[idx];
 					if (block.IssueIds.Contains(issue.Id))
 					{
 						string blockEventId = GetReportBlockEventId(messageState.Ts, idx);
+						_logger.LogInformation("Updating report block {EventId}", blockEventId);
 
 						List<(IIssue, IIssueSpan?)> pairs = new List<(IIssue, IIssueSpan?)>();
 						foreach (int issueId in block.IssueIds)
@@ -1470,7 +1475,6 @@ namespace Horde.Build.Notifications.Sinks
 		async Task UpdateReportBlockAsync(string channel, string eventId, DateTime reportTime, IStream stream, TemplateRefId? templateId, List<(IIssue, IIssueSpan?)> pairs)
 		{
 			StringBuilder body = new StringBuilder();
-			body.Append(DateTime.UtcNow.ToString() + "\n");
 
 			if (templateId != null)
 			{
