Horde: Manually split summary messages to prevent exceeding max Slack message sizes.

diff --git a/Engine/Source/Programs/Horde/Horde.Build/Notifications/Sinks/SlackNotificationSink.cs b/Engine/Source/Programs/Horde/Horde.Build/Notifications/Sinks/SlackNotificationSink.cs
--- a/Engine/Source/Programs/Horde/Horde.Build/Notifications/Sinks/SlackNotificationSink.cs
+++ b/Engine/Source/Programs/Horde/Horde.Build/Notifications/Sinks/SlackNotificationSink.cs
@@ -1295,6 +1295,8 @@ namespace Horde.Build.Notifications.Sinks
 		{
 			if (report.Channel != null)
 			{
+				const int MaxIssuesPerMessage = 8;
+
 				ReportState state = new ReportState();
 				state.Time = report.Time.UtcDateTime;
 
@@ -1326,29 +1328,30 @@ namespace Horde.Build.Notifications.Sinks
 
 						if (pairs.Count > 0)
 						{
-							List<(IIssue Issue, IIssueSpan? Span)> sortedPairs = pairs.OrderBy(x => x.Key.Id).Select(x => (x.Key, (IIssueSpan?)x.Value)).ToList();
-
-							ReportBlock block = new ReportBlock();
-							block.TemplateId = group.Key;
-							block.IssueIds.AddRange(sortedPairs.Select(x => x.Issue.Id));
-							state.Blocks.Add(block);
+							TemplateRefId templateId = group.Key;
+							foreach (IReadOnlyList<KeyValuePair<IIssue, IIssueSpan>> batch in pairs.OrderBy(x => x.Key.Id).Batch(MaxIssuesPerMessage))
+							{
+								ReportBlock block = new ReportBlock();
+								block.TemplateId = templateId;
+								block.IssueIds.AddRange(batch.Select(x => x.Key.Id));
+								state.Blocks.Add(block);
 
-							issuesByBlock.Add(sortedPairs);
+								issuesByBlock.Add(batch.Select(x => (x.Key, (IIssueSpan?)x.Value)).ToList());
+								templateId = TemplateRefId.Empty;
+							}
 						}
 					}
 				}
 				else
 				{
-					List<(IIssue Issue, IIssueSpan? Span)> sortedPairs = new List<(IIssue, IIssueSpan?)>();
-					foreach (IIssue issue in report.Issues.OrderByDescending(x => x.Id))
+					foreach (IReadOnlyList<IIssue> batch in report.Issues.OrderByDescending(x => x.Id).Batch(MaxIssuesPerMessage))
 					{
-						IIssueSpan? span = report.IssueSpans.FirstOrDefault(x => x.IssueId == issue.Id);
-						sortedPairs.Add((issue, span));
-					}
+						ReportBlock block = new ReportBlock();
+						block.IssueIds.AddRange(batch.Select(x => x.Id));
+						state.Blocks.Add(block);
 
-					ReportBlock block = new ReportBlock();
-					block.IssueIds.AddRange(sortedPairs.Select(x => x.Issue.Id));
-					state.Blocks.Add(block);
+						issuesByBlock.Add(batch.Select(x => (x, report.IssueSpans.FirstOrDefault(y => y.IssueId == x.Id))).ToList());
+					}
 				}
 
 				List<BlockBase> blocks = new List<BlockBase>();
