Horde: Show overridden summary text in Slack notifications, and strike through issues in reports once verified.

diff --git a/Engine/Source/Programs/Horde/Horde.Build/Notifications/Sinks/SlackNotificationSink.cs b/Engine/Source/Programs/Horde/Horde.Build/Notifications/Sinks/SlackNotificationSink.cs
--- a/Engine/Source/Programs/Horde/Horde.Build/Notifications/Sinks/SlackNotificationSink.cs
+++ b/Engine/Source/Programs/Horde/Horde.Build/Notifications/Sinks/SlackNotificationSink.cs
@@ -1499,7 +1499,14 @@ namespace Horde.Build.Notifications.Sinks
 				{
 					body.Append('\n');
 				}
-				body.Append(await FormatIssueAsync(issue, span, channel, reportTime));
+
+				string text = await FormatIssueAsync(issue, span, channel, reportTime);
+				if (issue.VerifiedAt != null)
+				{
+					text = $"~{text}~";
+				}
+
+				body.Append(text);
 			}
 
 			await SendOrUpdateMessageAsync(channel, eventId, null, body.ToString(), withEnvironment: false);
@@ -1548,7 +1555,8 @@ namespace Horde.Build.Notifications.Sinks
 				}
 			}
 
-			body.Append($"*: {issue.Summary} [{FormatReadableTimeSpan(reportTime - issue.CreatedAt)}]");
+			string? summary = issue.UserSummary ?? issue.Summary;
+			body.Append($"*: {summary} [{FormatReadableTimeSpan(reportTime - issue.CreatedAt)}]");
 
 			if (!String.IsNullOrEmpty(issue.ExternalIssueKey))
 			{
