Horde: Add a node annotation to flag build blockers. (eg. "BuildBlocker=true")

diff --git a/Engine/Source/Programs/Horde/Horde.Build/Jobs/Graphs/NodeAnnotations.cs b/Engine/Source/Programs/Horde/Horde.Build/Jobs/Graphs/NodeAnnotations.cs
--- a/Engine/Source/Programs/Horde/Horde.Build/Jobs/Graphs/NodeAnnotations.cs
+++ b/Engine/Source/Programs/Horde/Horde.Build/Jobs/Graphs/NodeAnnotations.cs
@@ -45,6 +45,11 @@ namespace Horde.Build.Jobs.Graphs
 		/// Key to use for grouping issues together, preventing them being merged with other groups
 		/// </summary>
 		string? IssueGroup { get; }
+
+		/// <summary>
+		/// Whether failures in this node should be flagged as build blockers
+		/// </summary>
+		bool? BuildBlocker { get; }
 	}
 
 	/// <summary>
@@ -75,6 +80,9 @@ namespace Horde.Build.Jobs.Graphs
 		/// <inheritdoc cref="IReadOnlyNodeAnnotations.IssueGroup"/>
 		public const string IssueGroupKeyName = "IssueGroup";
 
+		/// <inheritdoc cref="IReadOnlyNodeAnnotations.BuildBlocker"/>
+		public const string BuildBlockerKeyName = "BuildBlocker";
+
 		/// <summary>
 		/// Constructor
 		/// </summary>
@@ -141,6 +149,13 @@ namespace Horde.Build.Jobs.Graphs
 			set => SetStringValue(IssueGroupKeyName, value);
 		}
 
+		/// <inheritdoc/>
+		public bool? BuildBlocker
+		{
+			get => GetBoolValue(BuildBlockerKeyName);
+			set => SetBoolValue(BuildBlockerKeyName, value);
+		}
+
 		private bool? GetBoolValue(string key)
 		{
 			string? value = GetStringValue(key);
