Horde: Add yellow/red circles next to issues indicating whether they are warnings or errors.

diff --git a/Engine/Source/Programs/Horde/Horde.Build/Notifications/Sinks/SlackNotificationSink.cs b/Engine/Source/Programs/Horde/Horde.Build/Notifications/Sinks/SlackNotificationSink.cs
--- a/Engine/Source/Programs/Horde/Horde.Build/Notifications/Sinks/SlackNotificationSink.cs
+++ b/Engine/Source/Programs/Horde/Horde.Build/Notifications/Sinks/SlackNotificationSink.cs
@@ -765,7 +765,14 @@ namespace Horde.Build.Notifications.Sinks
 
 				string eventId = GetTriageThreadEventId(issue.Id);
 
-				string text = $"{workflow.TriagePrefix}*New Issue <{issueUrl}|{issue.Id}>*: {issue.Summary}{workflow.TriageSuffix}";
+				string prefix = workflow.TriagePrefix ?? String.Empty;
+				if (issue.Severity == IssueSeverity.Error && (span.LastFailure.Annotations.BuildBlocker ?? false))
+				{
+					prefix = $"{prefix}*[BUILD BLOCKER]* ";
+				}
+
+				string issueSummary = issue.UserSummary ?? issue.Summary;
+				string text = $"{workflow.TriagePrefix}{GetEmojiForSeverity(issue.Severity)}*New Issue <{issueUrl}|{issue.Id}>*: {issueSummary}{workflow.TriageSuffix}";
 				if (!spans.Any(x => x.NextSuccess == null))
 				{
 					text = $"~{text}~";
@@ -1512,6 +1519,11 @@ namespace Horde.Build.Notifications.Sinks
 			await SendOrUpdateMessageAsync(channel, eventId, null, body.ToString(), withEnvironment: false);
 		}
 
+		static string GetEmojiForSeverity(IssueSeverity severity)
+		{
+			return (severity == IssueSeverity.Warning) ? ":large_yellow_circle:" : ":red_circle:";
+		}
+
 		async ValueTask<string> FormatIssueAsync(IIssue issue, IIssueSpan? span, string? triageChannel, DateTime reportTime)
 		{
 			Uri issueUrl = _settings.DashboardUrl;
@@ -1544,7 +1556,8 @@ namespace Horde.Build.Notifications.Sinks
 				status = $"{status} - *Quarantined*";
 			}
 
-			StringBuilder body = new StringBuilder($"\u2022 *Issue <{issueUrl}|{issue.Id}>");
+			string emoji = GetEmojiForSeverity(issue.Severity);
+			StringBuilder body = new StringBuilder($"\u2022 {emoji} *Issue <{issueUrl}|{issue.Id}>");
 
 			if (triageChannel != null)
 			{
