Horde: Allow customizing prefix for errors and warnings. Add some emoji images to use for now.

diff --git a/Engine/Source/Programs/Horde/Horde.Build/Notifications/Sinks/SlackNotificationSink.cs b/Engine/Source/Programs/Horde/Horde.Build/Notifications/Sinks/SlackNotificationSink.cs
--- a/Engine/Source/Programs/Horde/Horde.Build/Notifications/Sinks/SlackNotificationSink.cs
+++ b/Engine/Source/Programs/Horde/Horde.Build/Notifications/Sinks/SlackNotificationSink.cs
@@ -772,7 +772,7 @@ namespace Horde.Build.Notifications.Sinks
 				}
 
 				string issueSummary = issue.UserSummary ?? issue.Summary;
-				string text = $"{workflow.TriagePrefix}{GetEmojiForSeverity(issue.Severity)}*New Issue <{issueUrl}|{issue.Id}>*: {issueSummary}{workflow.TriageSuffix}";
+				string text = $"{workflow.TriagePrefix}{GetPrefixForSeverity(issue.Severity)}*New Issue <{issueUrl}|{issue.Id}>*: {issueSummary}{workflow.TriageSuffix}";
 				if (!spans.Any(x => x.NextSuccess == null))
 				{
 					text = $"~{text}~";
@@ -1519,9 +1519,9 @@ namespace Horde.Build.Notifications.Sinks
 			await SendOrUpdateMessageAsync(channel, eventId, null, body.ToString(), withEnvironment: false);
 		}
 
-		static string GetEmojiForSeverity(IssueSeverity severity)
+		string GetPrefixForSeverity(IssueSeverity severity)
 		{
-			return (severity == IssueSeverity.Warning) ? ":large_yellow_circle:" : ":red_circle:";
+			return (severity == IssueSeverity.Warning) ? _settings.SlackWarningPrefix : _settings.SlackErrorPrefix;
 		}
 
 		async ValueTask<string> FormatIssueAsync(IIssue issue, IIssueSpan? span, string? triageChannel, DateTime reportTime)
@@ -1556,8 +1556,8 @@ namespace Horde.Build.Notifications.Sinks
 				status = $"{status} - *Quarantined*";
 			}
 
-			string emoji = GetEmojiForSeverity(issue.Severity);
-			StringBuilder body = new StringBuilder($"\u2022 {emoji} *Issue <{issueUrl}|{issue.Id}>");
+			string prefix = GetPrefixForSeverity(issue.Severity);
+			StringBuilder body = new StringBuilder($"\u2022 {prefix}*Issue <{issueUrl}|{issue.Id}>");
 
 			if (triageChannel != null)
 			{
diff --git a/Engine/Source/Programs/Horde/Horde.Build/ServerSettings.cs b/Engine/Source/Programs/Horde/Horde.Build/ServerSettings.cs
--- a/Engine/Source/Programs/Horde/Horde.Build/ServerSettings.cs
+++ b/Engine/Source/Programs/Horde/Horde.Build/ServerSettings.cs
@@ -404,6 +404,16 @@ namespace Horde.Build
 		/// </summary>
 		public string? SlackUsers { get; set; }
 
+		/// <summary>
+		/// Prefix to use when reporting errors
+		/// </summary>
+		public string? SlackErrorPrefix { get; set; } = ":horde-error: ";
+
+		/// <summary>
+		/// Prefix to use when reporting warnings
+		/// </summary>
+		public string? SlackWarningPrefix { get; set; } = ":horde-warning: ";
+
 		/// <summary>
 		/// Channel to send stream notification update failures to
 		/// </summary>
