Horde: Fix null reference warnings.

diff --git a/Engine/Source/Programs/Horde/Horde.Build/ServerSettings.cs b/Engine/Source/Programs/Horde/Horde.Build/ServerSettings.cs
--- a/Engine/Source/Programs/Horde/Horde.Build/ServerSettings.cs
+++ b/Engine/Source/Programs/Horde/Horde.Build/ServerSettings.cs
@@ -407,12 +407,12 @@ namespace Horde.Build
 		/// <summary>
 		/// Prefix to use when reporting errors
 		/// </summary>
-		public string? SlackErrorPrefix { get; set; } = ":horde-error: ";
+		public string SlackErrorPrefix { get; set; } = ":horde-error: ";
 
 		/// <summary>
 		/// Prefix to use when reporting warnings
 		/// </summary>
-		public string? SlackWarningPrefix { get; set; } = ":horde-warning: ";
+		public string SlackWarningPrefix { get; set; } = ":horde-warning: ";
 
 		/// <summary>
 		/// Channel to send stream notification update failures to
