Horde: Remove bullets before issues in reports now that we have icons.

diff --git a/Engine/Source/Programs/Horde/Horde.Build/Notifications/Sinks/SlackNotificationSink.cs b/Engine/Source/Programs/Horde/Horde.Build/Notifications/Sinks/SlackNotificationSink.cs
--- a/Engine/Source/Programs/Horde/Horde.Build/Notifications/Sinks/SlackNotificationSink.cs
+++ b/Engine/Source/Programs/Horde/Horde.Build/Notifications/Sinks/SlackNotificationSink.cs
@@ -1557,7 +1557,7 @@ namespace Horde.Build.Notifications.Sinks
 			}
 
 			string prefix = GetPrefixForSeverity(issue.Severity);
-			StringBuilder body = new StringBuilder($"\u2022 {prefix}*Issue <{issueUrl}|{issue.Id}>");
+			StringBuilder body = new StringBuilder($"{prefix}*Issue <{issueUrl}|{issue.Id}>");
 
 			if (triageChannel != null)
 			{
