Horde: Include the changelist number when reporting failed fixes in Slack threads.

diff --git a/Engine/Source/Programs/Horde/Horde.Build/Notifications/Sinks/SlackNotificationSink.cs b/Engine/Source/Programs/Horde/Horde.Build/Notifications/Sinks/SlackNotificationSink.cs
--- a/Engine/Source/Programs/Horde/Horde.Build/Notifications/Sinks/SlackNotificationSink.cs
+++ b/Engine/Source/Programs/Horde/Horde.Build/Notifications/Sinks/SlackNotificationSink.cs
@@ -885,7 +885,7 @@ namespace Horde.Build.Notifications.Sinks
 					else
 					{
 						string fixFailedEventId = $"issue_{issue.Id}_fixfailed_{issue.FixChange}";
-						string fixFailedMessage = $"Issue not fixed by {FormatChange(issue.FixChange.Value)}; see {FormatJobStep(fixFailedSpan.LastFailure, fixFailedSpan.NodeName)} in {fixFailedSpan.StreamName}.";
+						string fixFailedMessage = $"Issue not fixed by {FormatChange(issue.FixChange.Value)}; see {FormatJobStep(fixFailedSpan.LastFailure, fixFailedSpan.NodeName)} at CL {fixFailedSpan.LastFailure.Change} in {fixFailedSpan.StreamName}.";
 						if (issue.OwnerId.HasValue)
 						{
 							string mention = await FormatMentionAsync(issue.OwnerId.Value, workflow.AllowMentions);
