Horde: Put the report summary into the initial report thread, rather than sending separately. Improves formatting.

diff --git a/Engine/Source/Programs/Horde/Horde.Build/Notifications/Sinks/SlackNotificationSink.cs b/Engine/Source/Programs/Horde/Horde.Build/Notifications/Sinks/SlackNotificationSink.cs
--- a/Engine/Source/Programs/Horde/Horde.Build/Notifications/Sinks/SlackNotificationSink.cs
+++ b/Engine/Source/Programs/Horde/Horde.Build/Notifications/Sinks/SlackNotificationSink.cs
@@ -1364,6 +1364,16 @@ namespace Horde.Build.Notifications.Sinks
 				List<BlockBase> blocks = new List<BlockBase>();
 				blocks.Add(new HeaderBlock(AddEnvironmentAnnotation($"{report.Stream.Name}: {report.Time:d}")));
 
+				if (report.Issues.Count == 0)
+				{
+					blocks.Add(new SectionBlock(":tick: No issues open."));
+				}
+				else
+				{
+					TimeSpan averageAge = TimeSpan.FromHours(report.Issues.Select(x => (report.Time - x.CreatedAt).TotalHours).Average());
+					blocks.Add(new SectionBlock($"*{report.Issues.Count} unique issues* currently open (average age {FormatReadableTimeSpan(averageAge)})."));
+				}
+
 				PostMessageResponse? response = await SendMessageAsync(report.Channel, blocks: blocks.ToArray(), withEnvironment: false);
 				if (response != null && response.Ts != null)
 				{
@@ -1371,18 +1381,6 @@ namespace Horde.Build.Notifications.Sinks
 					string json = JsonSerializer.Serialize(state, _jsonSerializerOptions);
 					await AddOrUpdateMessageStateAsync(report.Channel, reportEventId, null, json, response.Ts);
 
-					string summary;
-					if (report.Issues.Count == 0)
-					{
-						summary = ":tick: No issues open.";
-					}
-					else
-					{
-						TimeSpan averageAge = TimeSpan.FromHours(report.Issues.Select(x => (report.Time - x.CreatedAt).TotalHours).Average());
-						summary = $"*{report.Issues.Count} unique issues* currently open (average age {FormatReadableTimeSpan(averageAge)}).";
-					}
-					await SendMessageAsync(report.Channel, text: summary, withEnvironment: false);
-
 					for (int idx = 0; idx < state.Blocks.Count; idx++)
 					{
 						string blockEventId = GetReportBlockEventId(response.Ts, idx);
