Horde Storage - Made the orphan blob GC run in parallell (added setting to control how wide). This should make the GC a bit more aggresive at cleaning up (at the cost of more load on the dbs and cpu usage)

Also enabled ref GC to run in parallel for which should also make that more aggresive in cleaning up out of date refs

diff --git a/Engine/Source/Programs/Horde/HordeStorage/Horde.Storage/Implementation/GC/CacheCleanup.cs b/Engine/Source/Programs/Horde/HordeStorage/Horde.Storage/Implementation/GC/CacheCleanup.cs
--- a/Engine/Source/Programs/Horde/HordeStorage/Horde.Storage/Implementation/GC/CacheCleanup.cs
+++ b/Engine/Source/Programs/Horde/HordeStorage/Horde.Storage/Implementation/GC/CacheCleanup.cs
@@ -80,14 +80,16 @@ namespace Horde.Storage.Implementation
             int countOfDeletedRecords = 0;
             DateTime cutoffTime = DateTime.Now.AddSeconds(-1 * _settings.CurrentValue.LastAccessCutoff.TotalSeconds);
             int consideredCount = 0;
-            await foreach ((BucketId bucket, IoHashKey name, DateTime lastAccessTime) in _referencesStore.GetRecords(ns).WithCancellation(cancellationToken))
+            await Parallel.ForEachAsync(_referencesStore.GetRecords(ns), new ParallelOptions {MaxDegreeOfParallelism = _settings.CurrentValue.OrphanRefMaxParallelOperations, CancellationToken = cancellationToken}, async (tuple, token) =>
             {
+                (BucketId bucket, IoHashKey name, DateTime lastAccessTime) = tuple;
+
                 _logger.Debug("Considering object in {Namespace} {Bucket} {Name} for deletion, was last updated {LastAccessTime}", ns, bucket, name, lastAccessTime);
                 Interlocked.Increment(ref consideredCount);
 
                 if (lastAccessTime > cutoffTime)
                 {
-                    continue;
+                    return;
                 }
 
                 _logger.Information("Attempting to delete object {Namespace} {Bucket} {Name} as it was last updated {LastAccessTime} which is older then {CutoffTime}", ns, bucket, name, lastAccessTime, cutoffTime);
@@ -118,7 +120,7 @@ namespace Horde.Storage.Implementation
                 {
                     _logger.Warning("Failed to delete record {Bucket} {Name} in {Namespace}", bucket, name, ns);            
                 }
-            }
+            });
 
             _logger.Information("Finished cleaning {Namespace}. Refs considered: {ConsideredCount} Refs Deleted: {DeletedCount}", ns, consideredCount, countOfDeletedRecords);
 
diff --git a/Engine/Source/Programs/Horde/HordeStorage/Horde.Storage/Implementation/GC/OrphanBlobCleanupRefs.cs b/Engine/Source/Programs/Horde/HordeStorage/Horde.Storage/Implementation/GC/OrphanBlobCleanupRefs.cs
--- a/Engine/Source/Programs/Horde/HordeStorage/Horde.Storage/Implementation/GC/OrphanBlobCleanupRefs.cs
+++ b/Engine/Source/Programs/Horde/HordeStorage/Horde.Storage/Implementation/GC/OrphanBlobCleanupRefs.cs
@@ -14,12 +14,14 @@ using Horde.Storage.Implementation.Blob;
 using Jupiter;
 using Jupiter.Common;
 using Jupiter.Implementation;
+using Microsoft.Extensions.Options;
 using Serilog;
 
 namespace Horde.Storage.Implementation
 {
     public class OrphanBlobCleanupRefs : IBlobCleanup
     {
+        private readonly IOptionsMonitor<GCSettings> _gcSettings;
         private readonly IBlobService _blobService;
         private readonly IObjectService _objectService;
         private readonly IBlobIndex _blobIndex;
@@ -28,8 +30,9 @@ namespace Horde.Storage.Implementation
         private readonly ILogger _logger = Log.ForContext<OrphanBlobCleanupRefs>();
 
         // ReSharper disable once UnusedMember.Global
-        public OrphanBlobCleanupRefs(IBlobService blobService, IObjectService objectService, IBlobIndex blobIndex, ILeaderElection leaderElection, INamespacePolicyResolver namespacePolicyResolver)
+        public OrphanBlobCleanupRefs(IOptionsMonitor<GCSettings> gcSettings, IBlobService blobService, IObjectService objectService, IBlobIndex blobIndex, ILeaderElection leaderElection, INamespacePolicyResolver namespacePolicyResolver)
         {
+            _gcSettings = gcSettings;
             _blobService = blobService;
             _objectService = objectService;
             _blobIndex = blobIndex;
@@ -56,104 +59,134 @@ namespace Horde.Storage.Implementation
             }
 
             List<NamespaceId> namespaces = await ListNamespaces().Where(NamespaceShouldBeCleaned).ToListAsync();
-            
+            List<NamespaceId> namespacesThatHaveBeenChecked = new List<NamespaceId>();
             // enumerate all namespaces, and check if the old blob is valid in any of them to allow for a blob store to just store them in a single pile if it wants to
             ulong countOfBlobsRemoved = 0;
+            _logger.Information("Started orphan blob");
             foreach (NamespaceId @namespace in namespaces)
             {
+                // if we have already checked this namespace there is no need to repeat it
+                if (namespacesThatHaveBeenChecked.Contains(@namespace))
+                {
+                    continue;
+                }
                 if (cancellationToken.IsCancellationRequested)
                 {
                     break;
                 }
 
+                DateTime startTime = DateTime.Now;
                 NamespacePolicy policy = _namespacePolicyResolver.GetPoliciesForNs(@namespace);
                 List<NamespaceId> namespacesThatSharePool = namespaces.Where(ns => _namespacePolicyResolver.GetPoliciesForNs(ns).StoragePool == policy.StoragePool).ToList();
 
+                _logger.Information("Running Orphan GC For StoragePool: {StoragePool}", policy.StoragePool);
+                namespacesThatHaveBeenChecked.AddRange(namespacesThatSharePool);
                 // only consider blobs that have been around for 60 minutes
                 // this due to cases were blobs are uploaded first
                 DateTime cutoff = DateTime.Now.AddMinutes(-60);
-                await foreach ((BlobIdentifier blob, DateTime lastModified) in _blobService.ListObjects(@namespace).WithCancellation(cancellationToken))
+                await _blobService.ListObjects(@namespace).ParallelForEachAsync(async (tuple) =>
                 {
+                    (BlobIdentifier blob, DateTime lastModified) = tuple;
+
                     if (lastModified > cutoff)
                     {
-                        continue;
+                        return;
                     }
 
                     if (cancellationToken.IsCancellationRequested)
                     {
-                        break;
+                        return;
                     }
 
-                    using IScope removeBlobScope = Tracer.Instance.StartActive("gc.blob");
-                    removeBlobScope.Span.ResourceName = $"{@namespace}.{blob}";
+                    bool removed = await GCBlob(policy.StoragePool, namespacesThatSharePool, blob, lastModified, cancellationToken);
 
-                    bool found = false;
-
-                    // check all other namespaces that share the same storage pool for presence of the blob
-                    await Parallel.ForEachAsync(namespacesThatSharePool, cancellationToken, async (blobNamespace, token) =>
+                    if (removed)
                     {
-                        if (token.IsCancellationRequested)
-                        {
-                            return;
-                        }
-
-                        BlobInfo? blobIndex = await _blobIndex.GetBlobInfo(blobNamespace, blob);
-
-                        if (blobIndex == null)
-                        {
-                            return;
-                        }
-
-                        await Parallel.ForEachAsync(blobIndex.References, cancellationToken, async (tuple, token) =>
-                        {
-                            (BucketId bucket, IoHashKey key) = tuple;
-                            if (token.IsCancellationRequested)
-                            {
-                                return;
-                            }
-
-                            try
-                            {
-                                (ObjectRecord, BlobContents?) _ = await _objectService.Get(blobNamespace,
-                                    bucket, key, new string[] { "name" });
-                                found = true;
-                            }
-                            catch (ObjectNotFoundException)
-                            {
-                                // this is not a valid reference so we should delete
-                            }
-                            catch (MissingBlobsException)
-                            {
-                                // we do not care if there are missing blobs, as long as the record is valid we keep this blob around
-                                found = true;
-                            }
-                        });
-                    });
+                        Interlocked.Increment(ref countOfBlobsRemoved);
+                    }
+                }, cancellationToken: cancellationToken, maxDegreeOfParallelism: _gcSettings.CurrentValue.OrphanGCMaxParallelOperations);
 
-                    if (cancellationToken.IsCancellationRequested)
+                TimeSpan storagePoolGcDuration = DateTime.Now - startTime;
+                _logger.Information("Finished running Orphan GC For StoragePool: {StoragePool}. Took {Duration}", policy.StoragePool, storagePoolGcDuration);
+            }
+
+            _logger.Information("Finished running Orphan GC");
+            return countOfBlobsRemoved;
+        }
+
+        private async Task<bool> GCBlob(string storagePool, List<NamespaceId> namespacesThatSharePool, BlobIdentifier blob, DateTime lastModifiedTime, CancellationToken cancellationToken)
+        {
+            using IScope removeBlobScope = Tracer.Instance.StartActive("gc.blob");
+            removeBlobScope.Span.ResourceName = $"{storagePool}.{blob}";
+
+            bool found = false;
+
+            // check all namespaces that share the same storage pool for presence of the blob
+            await Parallel.ForEachAsync(namespacesThatSharePool, cancellationToken, async (blobNamespace, token) =>
+            {
+                if (token.IsCancellationRequested)
+                {
+                    return;
+                }
+
+                BlobInfo? blobIndex = await _blobIndex.GetBlobInfo(blobNamespace, blob);
+
+                if (blobIndex == null)
+                {
+                    return;
+                }
+
+                await Parallel.ForEachAsync(blobIndex.References, cancellationToken, async (tuple, token) =>
+                {
+                    (BucketId bucket, IoHashKey key) = tuple;
+                    if (token.IsCancellationRequested)
                     {
-                        break;
+                        return;
                     }
 
-                    removeBlobScope.Span.SetTag("removed", (!found).ToString());
-
-                    if (!found)
+                    try
                     {
-                        await Parallel.ForEachAsync(namespacesThatSharePool, cancellationToken, async (ns, _) =>
-                        {
-                            await RemoveBlob(ns, blob);
-                        });
-                        ++countOfBlobsRemoved;
+                        (ObjectRecord, BlobContents?) _ = await _objectService.Get(blobNamespace,
+                            bucket, key, new string[] { "name" });
+                        found = true;
                     }
-                }
+                    catch (ObjectNotFoundException)
+                    {
+                        // this is not a valid reference so we should delete
+                    }
+                    catch (MissingBlobsException)
+                    {
+                        // we do not care if there are missing blobs, as long as the record is valid we keep this blob around
+                        found = true;
+                    }
+                });
+            });
+
+            if (cancellationToken.IsCancellationRequested)
+            {
+                return false;
             }
 
-            return countOfBlobsRemoved;
+            removeBlobScope.Span.SetTag("removed", (!found).ToString());
+
+            // something is still referencing this blob, we should not delete it
+            if (found)
+            {
+                return false;
+            }
+
+            // if the blob was not found to have a reference in any of the namespace that share a storage pool then the blob is not used anymore and should be deleted from all the namespaces
+            await Parallel.ForEachAsync(namespacesThatSharePool, cancellationToken, async (ns, _) =>
+            {
+                await RemoveBlob(ns, blob, lastModifiedTime);
+            });
+            return true;
+
         }
 
-        private async Task RemoveBlob(NamespaceId ns, BlobIdentifier blob)
+        private async Task RemoveBlob(NamespaceId ns, BlobIdentifier blob, DateTime lastModifiedTime)
         {
-            _logger.Information("Attempting to GC Orphan blob {Blob} from {Namespace}", blob, ns);
+            _logger.Information("Attempting to GC Orphan blob {Blob} from {Namespace} which was last modified at {LastModifiedTime}", blob, ns, lastModifiedTime);
             try
             {
                 await _blobService.DeleteObject(ns, blob);
diff --git a/Engine/Source/Programs/Horde/HordeStorage/Horde.Storage/Settings.cs b/Engine/Source/Programs/Horde/HordeStorage/Horde.Storage/Settings.cs
--- a/Engine/Source/Programs/Horde/HordeStorage/Horde.Storage/Settings.cs
+++ b/Engine/Source/Programs/Horde/HordeStorage/Horde.Storage/Settings.cs
@@ -256,7 +256,8 @@ namespace Horde.Storage
 
         public TimeSpan BlobCleanupPollFrequency { get; set; } = TimeSpan.FromMinutes(60);
         public TimeSpan RefCleanupPollFrequency { get; set; } = TimeSpan.FromMinutes(60);
-
+        public int OrphanGCMaxParallelOperations { get; set; } = 8;
+        public int OrphanRefMaxParallelOperations { get; set; } = 8;
     }
 
     public class UpstreamRelaySettings
