Horde Storage - Customzing the user agent to fit when the convention we have established in the cooker

diff --git a/Engine/Source/Programs/Horde/HordeStorage/Horde.Storage/Horde.Storage.Startup.cs b/Engine/Source/Programs/Horde/HordeStorage/Horde.Storage/Horde.Storage.Startup.cs
--- a/Engine/Source/Programs/Horde/HordeStorage/Horde.Storage/Horde.Storage.Startup.cs
+++ b/Engine/Source/Programs/Horde/HordeStorage/Horde.Storage/Horde.Storage.Startup.cs
@@ -5,7 +5,9 @@ using System.Collections.Generic;
 using System.IO;
 using System.Linq;
 using System.Net;
+using System.Net.Http;
 using System.Net.Security;
+using System.Runtime.InteropServices;
 using System.Security.Authentication;
 using Amazon.DAX;
 using Amazon.DynamoDBv2;
@@ -55,7 +57,33 @@ namespace Horde.Storage
 
         protected override void OnAddService(IServiceCollection services)
         {
-            services.AddHttpClient();
+            services.AddHttpClient(Options.DefaultName, (provider, client) =>
+            {
+                string GetPlatformName()
+                {
+                    if (RuntimeInformation.IsOSPlatform(OSPlatform.Linux))
+                    {
+                        return "Linux";
+                    } 
+                    else if (RuntimeInformation.IsOSPlatform(OSPlatform.Windows))
+                    {
+                        return "Windows";
+                    } 
+                    else if (RuntimeInformation.IsOSPlatform(OSPlatform.OSX))
+                    {
+                        return "OSX";
+                    } 
+                    else
+                    {
+                        throw new NotSupportedException("Unknown OS when formatting platform name");
+                    }
+                }
+                IVersionFile? versionFile = provider.GetService<IVersionFile>();
+                string engineVersion = "5.1"; // TODO: we should read this from the version file
+                string? hordeStorageVersion = versionFile?.VersionString ?? "unknown";
+                string platformName = GetPlatformName();
+                client.DefaultRequestHeaders.TryAddWithoutValidation("UserAgent", $"UnrealEngine/{engineVersion} ({platformName}) HordeStorage/{hordeStorageVersion}");
+            });
 
             services.AddOptions<HordeStorageSettings>().Bind(Configuration.GetSection("Horde_Storage")).ValidateDataAnnotations();
             services.AddOptions<MongoSettings>().Bind(Configuration.GetSection("Mongo")).ValidateDataAnnotations();
