Horde Storage - Added option to disable the legacy api (unstructured cache) - enabled it in Dev.

diff --git a/Engine/Source/Programs/Horde/HordeStorage/Horde.Storage/Controllers/DDCRefController.cs b/Engine/Source/Programs/Horde/HordeStorage/Horde.Storage/Controllers/DDCRefController.cs
--- a/Engine/Source/Programs/Horde/HordeStorage/Horde.Storage/Controllers/DDCRefController.cs
+++ b/Engine/Source/Programs/Horde/HordeStorage/Horde.Storage/Controllers/DDCRefController.cs
@@ -18,6 +18,7 @@ using Microsoft.AspNetCore.Authorization;
 using Microsoft.AspNetCore.Http;
 using Microsoft.AspNetCore.Http.Extensions;
 using Microsoft.AspNetCore.Mvc;
+using Microsoft.Extensions.Options;
 using Newtonsoft.Json;
 using Serilog;
 
@@ -51,16 +52,18 @@ namespace Horde.Storage.Controllers
         private readonly IRefsStore _refsStore;
         private readonly IDiagnosticContext _diagnosticContext;
 		private readonly RequestHelper _requestHelper;
+        private readonly IOptionsMonitor<HordeStorageSettings> _settings;
 
         private readonly ILogger _logger = Log.ForContext<DDCRefController>();
         private readonly IDDCRefService _ddcRefService;
 
-        public DDCRefController(IDDCRefService ddcRefService, IRefsStore refsStore, IDiagnosticContext diagnosticContext, RequestHelper requestHelper)
+        public DDCRefController(IDDCRefService ddcRefService, IRefsStore refsStore, IDiagnosticContext diagnosticContext, RequestHelper requestHelper, IOptionsMonitor<HordeStorageSettings> settings)
         {
             _refsStore = refsStore;
             _ddcRefService = ddcRefService;
             _diagnosticContext = diagnosticContext;
             _requestHelper = requestHelper;
+            _settings = settings;
         }
 
         /// <summary>
@@ -115,6 +118,11 @@ namespace Horde.Storage.Controllers
                 return result;
             }
 
+            if (_settings.CurrentValue.DisableLegacyApi)
+            {
+                return NotFound("This api has been removed, you will need to update your code.");
+            }
+
             try
             {
                 bool isRaw = format?.ToUpperInvariant() == "RAW";
@@ -202,6 +210,11 @@ namespace Horde.Storage.Controllers
                 return result;
             }
 
+            if (_settings.CurrentValue.DisableLegacyApi)
+            {
+                return NotFound("This api has been removed, you will need to update your code.");
+            }
+
             try
             {
                 RefRecord record = await _ddcRefService.Exists(ns, bucket, key);
@@ -248,6 +261,11 @@ namespace Horde.Storage.Controllers
                 return result;
             }
 
+            if (_settings.CurrentValue.DisableLegacyApi)
+            {
+                return NotFound("This api has been removed, you will need to update your code.");
+            }
+
             try
             {
                 _diagnosticContext.Set("Content-Length", Request.ContentLength ?? -1);
@@ -285,6 +303,12 @@ namespace Horde.Storage.Controllers
                 return result;
             }
 
+            
+            if (_settings.CurrentValue.DisableLegacyApi)
+            {
+                return NotFound("This api has been removed, you will need to update your code.");
+            }
+
             byte[] blob;
             try
             {
@@ -487,6 +511,11 @@ namespace Horde.Storage.Controllers
                 throw new InvalidOperationException("No operations specified, this is a required field");
             }
 
+            if (_settings.CurrentValue.DisableLegacyApi)
+            {
+                return NotFound("This api has been removed, you will need to update your code.");
+            }
+
             Task<object?>[] tasks = new Task<object?>[batch.Operations.Length];
             for (int index = 0; index < batch.Operations.Length; index++)
             {
diff --git a/Engine/Source/Programs/Horde/HordeStorage/Horde.Storage/Settings.cs b/Engine/Source/Programs/Horde/HordeStorage/Horde.Storage/Settings.cs
--- a/Engine/Source/Programs/Horde/HordeStorage/Horde.Storage/Settings.cs
+++ b/Engine/Source/Programs/Horde/HordeStorage/Horde.Storage/Settings.cs
@@ -141,6 +141,9 @@ namespace Horde.Storage
         public int LastAccessRollupFrequencySeconds { get; set; } = 900; // 15 minutes
         public bool EnableLastAccessTracking { get; set; } = true;
         public bool EnableOnDemandReplication { get; set; } = false;
+
+        // disable the legacy api
+        public bool DisableLegacyApi { get; set; } = false;
     }
 
     public class MongoSettings
