Horde Storage - CbPackage support in refs, supports getting a ref as a package and putting a ref as a package. This is intended to replace the more primitive x-jupiter-inline type.

diff --git a/Engine/Source/Programs/Horde/HordeStorage/Horde.Storage/Controllers/CompressedBlobController.cs b/Engine/Source/Programs/Horde/HordeStorage/Horde.Storage/Controllers/CompressedBlobController.cs
--- a/Engine/Source/Programs/Horde/HordeStorage/Horde.Storage/Controllers/CompressedBlobController.cs
+++ b/Engine/Source/Programs/Horde/HordeStorage/Horde.Storage/Controllers/CompressedBlobController.cs
@@ -4,18 +4,15 @@ using System;
 using System.Collections.Concurrent;
 using System.Collections.Generic;
 using System.ComponentModel.DataAnnotations;
-using System.IO;
 using System.Linq;
 using System.Net.Mime;
 using System.Threading.Tasks;
-using Datadog.Trace;
 using EpicGames.Horde.Storage;
 using EpicGames.Serialization;
 using Horde.Storage.Implementation;
 using Jupiter;
 using Jupiter.Common.Implementation;
 using Jupiter.Implementation;
-using Jupiter.Utils;
 using Microsoft.AspNetCore.Authorization;
 using Microsoft.AspNetCore.Mvc;
 using Microsoft.Extensions.Primitives;
@@ -64,7 +61,7 @@ namespace Horde.Storage.Controllers
 
             try
             {
-                (BlobContents blobContents, string mediaType) = await GetImpl(ns, id);
+                (BlobContents blobContents, string mediaType) = await _storage.GetCompressedObject(ns, id, _contentIdStore);
 
                 StringValues acceptHeader = Request.Headers["Accept"];
                 if (!acceptHeader.Contains("*/*") && acceptHeader.Count != 0 && !acceptHeader.Contains(mediaType))
@@ -204,50 +201,6 @@ namespace Horde.Storage.Controllers
             return Ok(new ExistCheckMultipleContentIdResponse { Needs = needs.ToArray()});
         }
 
-        private async Task<(BlobContents, string)> GetImpl(NamespaceId ns, ContentId contentId)
-        {
-            BlobIdentifier[]? chunks = await _contentIdStore.Resolve(ns, contentId, mustBeContentId: false);
-            if (chunks == null || chunks.Length == 0)
-            {
-                throw new ContentIdResolveException(contentId);
-            }
-
-            // single chunk, we just return that chunk
-            if (chunks.Length == 1)
-            {
-                BlobIdentifier blobToReturn = chunks[0];
-                string mimeType = CustomMediaTypeNames.UnrealCompressedBuffer;
-                if (contentId.Equals(blobToReturn))
-                {
-                    // this was actually the unmapped blob, meaning its not a compressed buffer
-                    mimeType = MediaTypeNames.Application.Octet;
-                }
-
-                return (await _storage.GetObject(ns, blobToReturn), mimeType);
-            }
-
-            // chunked content, combine the chunks into a single stream
-            using IScope _ = Tracer.Instance.StartActive("blob.combine");
-            Task<BlobContents>[] tasks = new Task<BlobContents>[chunks.Length];
-            for (int i = 0; i < chunks.Length; i++)
-            {
-                tasks[i] = _storage.GetObject(ns, chunks[i]);
-            }
-
-            MemoryStream ms = new MemoryStream();
-            foreach (Task<BlobContents> task in tasks)
-            {
-                BlobContents blob = await task;
-                await using Stream s = blob.Stream;
-                await s.CopyToAsync(ms);
-            }
-
-            ms.Seek(0, SeekOrigin.Begin);
-
-            // chunking could not have happened for a non compressed buffer so assume it is compressed
-            return (new BlobContents(ms, ms.Length), CustomMediaTypeNames.UnrealCompressedBuffer);
-        }
-
         [HttpPut("{ns}/{id}")]
         [DisableRequestSizeLimit]
         [RequiredContentType(CustomMediaTypeNames.UnrealCompressedBuffer)]
@@ -267,7 +220,7 @@ namespace Horde.Storage.Controllers
 
             try
             {
-                ContentId identifier = await PutImpl(ns, id, payload);
+                ContentId identifier = await _storage.PutCompressedObject(ns, payload, id, _contentIdStore, _compressedBufferUtils);
 
                 return Ok(new
                 {
@@ -301,7 +254,7 @@ namespace Horde.Storage.Controllers
 
             try
             {
-                ContentId identifier = await PutImpl(ns, null, payload);
+                ContentId identifier = await _storage.PutCompressedObject(ns, payload, null, _contentIdStore, _compressedBufferUtils);
 
                 return Ok(new
                 {
@@ -317,52 +270,6 @@ namespace Horde.Storage.Controllers
             }
         }
 
-        private async Task<ContentId> PutImpl(NamespaceId ns, ContentId? id, IBufferedPayload payload)
-        {
-            // decompress the content and generate a identifier from it to verify the identifier we got if we got one
-            await using Stream decompressStream = payload.GetStream();
-            // TODO: we should add a overload for decompress content that can work on streams, otherwise we are still limited to 2GB compressed blobs
-            byte[] decompressedContent = _compressedBufferUtils.DecompressContent(await decompressStream.ToByteArray());
-
-            await using MemoryStream decompressedStream = new MemoryStream(decompressedContent);
-            ContentId identifierDecompressedPayload;
-            if (id != null)
-            {
-                identifierDecompressedPayload = ContentId.FromContentHash(await _storage.VerifyContentMatchesHash(decompressedStream, id));
-            }
-            else
-            {
-                ContentHash blobHash;
-                {
-                    using IScope _ = Tracer.Instance.StartActive("web.hash");
-                    blobHash = await BlobIdentifier.FromStream(decompressedStream);
-                }
-
-                identifierDecompressedPayload = ContentId.FromContentHash(blobHash);
-            }
-
-            BlobIdentifier identifierCompressedPayload;
-            {
-                using IScope _ = Tracer.Instance.StartActive("web.hash");
-                await using Stream hashStream = payload.GetStream();
-                identifierCompressedPayload = await BlobIdentifier.FromStream(hashStream);
-            }
-
-            // commit the mapping from the decompressed hash to the compressed hash, we run this in parallel with the blob store submit
-            // TODO: let users specify weight of the blob compared to previously submitted content ids
-            int contentIdWeight = (int)payload.Length;
-            Task contentIdStoreTask = _contentIdStore.Put(ns, identifierDecompressedPayload, identifierCompressedPayload, contentIdWeight);
-
-            // we still commit the compressed buffer to the object store using the hash of the compressed content
-            {
-                await _storage.PutObjectKnownHash(ns, payload, identifierCompressedPayload);
-            }
-
-            await contentIdStoreTask;
-
-            return identifierDecompressedPayload;
-        }
-
         /*[HttpDelete("{ns}/{id}")]
         public async Task<IActionResult> Delete(
             [Required] string ns,
diff --git a/Engine/Source/Programs/Horde/HordeStorage/Horde.Storage/Controllers/FormatResolver.cs b/Engine/Source/Programs/Horde/HordeStorage/Horde.Storage/Controllers/FormatResolver.cs
--- a/Engine/Source/Programs/Horde/HordeStorage/Horde.Storage/Controllers/FormatResolver.cs
+++ b/Engine/Source/Programs/Horde/HordeStorage/Horde.Storage/Controllers/FormatResolver.cs
@@ -16,7 +16,11 @@ namespace Horde.Storage.Controllers
         private readonly IOptionsMonitor<MvcOptions> _mvcOptions;
 
         private readonly string[] _validContentTypes = {
-            MediaTypeNames.Application.Octet, MediaTypeNames.Application.Json, CustomMediaTypeNames.UnrealCompactBinary, CustomMediaTypeNames.JupiterInlinedPayload
+            MediaTypeNames.Application.Octet, 
+            MediaTypeNames.Application.Json, 
+            CustomMediaTypeNames.UnrealCompactBinary, 
+            CustomMediaTypeNames.JupiterInlinedPayload,
+			CustomMediaTypeNames.UnrealCompactBinaryPackage
         };
 
         public FormatResolver(IOptionsMonitor<MvcOptions> mvcOptions)
diff --git a/Engine/Source/Programs/Horde/HordeStorage/Horde.Storage/Controllers/ObjectController.cs b/Engine/Source/Programs/Horde/HordeStorage/Horde.Storage/Controllers/ObjectController.cs
--- a/Engine/Source/Programs/Horde/HordeStorage/Horde.Storage/Controllers/ObjectController.cs
+++ b/Engine/Source/Programs/Horde/HordeStorage/Horde.Storage/Controllers/ObjectController.cs
@@ -202,7 +202,7 @@ namespace Horde.Storage.Controllers
 
             try
             {
-                BlobIdentifier[] references = await _referenceResolver.ResolveReferences(ns, compactBinaryObject).ToArrayAsync();
+                BlobIdentifier[] references = await _referenceResolver.GetReferencedBlobs(ns, compactBinaryObject).ToArrayAsync();
                 return Ok(new ResolvedReferencesResult(references));
             }
             catch (PartialReferenceResolveException e)
diff --git a/Engine/Source/Programs/Horde/HordeStorage/Horde.Storage/Controllers/ReferencesController.cs b/Engine/Source/Programs/Horde/HordeStorage/Horde.Storage/Controllers/ReferencesController.cs
--- a/Engine/Source/Programs/Horde/HordeStorage/Horde.Storage/Controllers/ReferencesController.cs
+++ b/Engine/Source/Programs/Horde/HordeStorage/Horde.Storage/Controllers/ReferencesController.cs
@@ -47,24 +47,28 @@ namespace Horde.Storage.Controllers
     public class ReferencesController : ControllerBase
     {
         private readonly IDiagnosticContext _diagnosticContext;
+        private readonly IContentIdStore _contentIdStore;
         private readonly FormatResolver _formatResolver;
         private readonly BufferedPayloadFactory _bufferedPayloadFactory;
         private readonly IReferenceResolver _referenceResolver;
         private readonly RequestHelper _requestHelper;
+        private readonly CompressedBufferUtils _compressedBufferUtils;
 
         private readonly ILogger _logger = Log.ForContext<ReferencesController>();
         private readonly IObjectService _objectService;
         private readonly IBlobService _blobStore;
 
-        public ReferencesController(IObjectService objectService, IBlobService blobStore, IDiagnosticContext diagnosticContext, FormatResolver formatResolver, BufferedPayloadFactory bufferedPayloadFactory, IReferenceResolver referenceResolver, RequestHelper requestHelper)
+        public ReferencesController(IObjectService objectService, IBlobService blobStore, IDiagnosticContext diagnosticContext, IContentIdStore contentIdStore, FormatResolver formatResolver, BufferedPayloadFactory bufferedPayloadFactory, IReferenceResolver referenceResolver, RequestHelper requestHelper, CompressedBufferUtils compressedBufferUtils)
         {
             _objectService = objectService;
             _blobStore = blobStore;
             _diagnosticContext = diagnosticContext;
+            _contentIdStore = contentIdStore;
             _formatResolver = formatResolver;
             _bufferedPayloadFactory = bufferedPayloadFactory;
             _referenceResolver = referenceResolver;
             _requestHelper = requestHelper;
+            _compressedBufferUtils = compressedBufferUtils;
         }
 
         /// <summary>
@@ -100,7 +104,7 @@ namespace Horde.Storage.Controllers
         /// <param name="key">The unique name of this particular key. `iAmAVeryValidKey`</param>
         /// <param name="format">Optional specifier to set which output format is used json/raw/cb</param>
         [HttpGet("{ns}/{bucket}/{key}.{format?}", Order = 500)]
-        [Produces(MediaTypeNames.Application.Json, MediaTypeNames.Application.Octet, CustomMediaTypeNames.UnrealCompactBinary, CustomMediaTypeNames.JupiterInlinedPayload)]
+        [Produces(MediaTypeNames.Application.Json, MediaTypeNames.Application.Octet, CustomMediaTypeNames.UnrealCompactBinary, CustomMediaTypeNames.JupiterInlinedPayload, CustomMediaTypeNames.UnrealCompactBinaryPackage)]
         public async Task<IActionResult> Get(
             [FromRoute] [Required] NamespaceId ns,
             [FromRoute] [Required] BucketId bucket,
@@ -213,6 +217,76 @@ namespace Horde.Storage.Controllers
                         break;
 
                     }
+                    case CustomMediaTypeNames.UnrealCompactBinaryPackage:
+                    {
+                        byte[] blobMemory = await blob.Stream.ToByteArray();
+                        CbObject cb = new CbObject(blobMemory);
+
+                        IAsyncEnumerable<Attachment> attachments = _referenceResolver.GetAttachments(ns, cb);
+
+                        CbPackageBuilder writer = new CbPackageBuilder();
+                        await writer.AddAttachment(objectRecord.BlobIdentifier.AsIoHash(), CbPackageAttachmentFlags.IsObject, blobMemory);
+
+                        List<Task> fetchBlobTasks = new List<Task>();
+                        await foreach (Attachment attachment in attachments)
+                        {
+                            Task blobTask = Task.Run(async () =>
+                            {
+                                IoHash attachmentHash = attachment.AsIoHash();
+                                BlobContents contents;
+                                CbPackageAttachmentFlags flags = 0;
+
+                                try
+                                {
+                                    if (attachment is BlobAttachment blobAttachment)
+                                    {
+                                        BlobIdentifier referencedBlob = blobAttachment.Identifier;
+                                        contents = await _blobStore.GetObject(ns, referencedBlob);
+                                    }
+                                    else if (attachment is ObjectAttachment objectAttachment)
+                                    {
+                                        flags &= CbPackageAttachmentFlags.IsObject;
+                                        BlobIdentifier referencedBlob = objectAttachment.Identifier;
+                                        contents = await _blobStore.GetObject(ns, referencedBlob);
+                                    }
+                                    else if (attachment is ContentIdAttachment contentIdAttachment)
+                                    {
+
+                                        ContentId contentId = contentIdAttachment.Identifier;
+                                        (contents, string mime) = await _blobStore.GetCompressedObject(ns, contentId, _contentIdStore);
+                                        if (mime == CustomMediaTypeNames.UnrealCompressedBuffer)
+                                        {
+                                            flags &= CbPackageAttachmentFlags.IsCompressed;
+                                        }
+                                        else
+                                        {
+                                            // this resolved to a uncompressed blob, the content id existed the the compressed blob didn't
+                                            // so resetting flags to indicate this.
+                                            flags = 0;
+                                        }
+                                    }
+                                    else
+                                    {
+                                        throw new NotSupportedException($"Unknown attachment type {attachment.GetType()}");
+                                    }
+                                }
+                                catch (Exception e)
+                                {
+                                    (CbObject errorObject, HttpStatusCode _) = ToErrorResult(e);
+                                    await writer.AddAttachment(attachmentHash, CbPackageAttachmentFlags.IsError | CbPackageAttachmentFlags.IsObject, errorObject.GetView().ToArray());
+                                    return;
+                                }
+                                await writer.AddAttachment(attachmentHash, flags, contents.Stream, (ulong)contents.Length);
+                            });
+
+                            fetchBlobTasks.Add(blobTask);
+
+                        }
+                        await Task.WhenAll(fetchBlobTasks);
+                        await using BlobContents contents = new BlobContents(writer.ToByteArray());
+                        await WriteBody(contents, CustomMediaTypeNames.UnrealCompactBinaryPackage);
+                        break;
+                    }
                     case CustomMediaTypeNames.JupiterInlinedPayload:
                     {
                         byte[] blobMemory = await blob.Stream.ToByteArray();
@@ -248,7 +322,7 @@ namespace Horde.Storage.Controllers
                             List<BlobIdentifier> referencedBlobs;
                             try
                             {
-                                IAsyncEnumerable<BlobIdentifier> referencedBlobsEnumerable = _referenceResolver.ResolveReferences(ns, cb);
+                                IAsyncEnumerable<BlobIdentifier> referencedBlobsEnumerable = _referenceResolver.GetReferencedBlobs(ns, cb);
                                 referencedBlobs = await referencedBlobsEnumerable.ToListAsync();
                             }
                             catch (PartialReferenceResolveException)
@@ -408,7 +482,7 @@ namespace Horde.Storage.Controllers
                 byte[] blobContents = await blob.Stream.ToByteArray();
                 CbObject compactBinaryObject = new CbObject(blobContents);
                 // the reference resolver will throw if any blob is missing, so no need to do anything other then process each reference
-                IAsyncEnumerable<BlobIdentifier> references = _referenceResolver.ResolveReferences(ns, compactBinaryObject);
+                IAsyncEnumerable<BlobIdentifier> references = _referenceResolver.GetReferencedBlobs(ns, compactBinaryObject);
                 List<BlobIdentifier>? _ = await references.ToListAsync();
 
                 // we have to verify the blobs are available locally, as the record of the key is replicated a head of the content
@@ -486,7 +560,7 @@ namespace Horde.Storage.Controllers
                     byte[] blobContents = await blob.Stream.ToByteArray();
                     CbObject cb = new CbObject(blobContents);
                     // the reference resolver will throw if any blob is missing, so no need to do anything other then process each reference
-                    IAsyncEnumerable<BlobIdentifier> references = _referenceResolver.ResolveReferences(ns, cb);
+                    IAsyncEnumerable<BlobIdentifier> references = _referenceResolver.GetReferencedBlobs(ns, cb);
                     List<BlobIdentifier>? _ = await references.ToListAsync();
                 }
                 catch (ObjectNotFoundException)
@@ -602,6 +676,67 @@ namespace Horde.Storage.Controllers
             return Ok(new PutObjectResponse(missingHashes.ToArray()));
         }
 
+        [HttpPut("{ns}/{bucket}/{key}", Order = 300)]
+        [DisableRequestSizeLimit]
+        [RequiredContentType(CustomMediaTypeNames.UnrealCompactBinaryPackage)]
+        [Authorize("Object.write")]
+        public async Task<IActionResult> PutPackage(
+            [FromRoute][Required] NamespaceId ns,
+            [FromRoute][Required] BucketId bucket,
+            [FromRoute][Required] IoHashKey key)
+        {
+            ActionResult? accessResult = await _requestHelper.HasAccessToNamespace(User, Request, ns, new [] { AclAction.WriteObject });
+            if (accessResult != null)
+            {
+                return accessResult;
+            }
+
+            _diagnosticContext.Set("Content-Length", Request.ContentLength ?? -1);
+
+            CbPackageReader packageReader = await CbPackageReader.Create(Request.Body);
+
+            try
+            {
+                await foreach ((CbPackageAttachmentEntry entry, byte[] blob) in packageReader.IterateAttachments())
+                {
+                    if (entry.Flags.HasFlag(CbPackageAttachmentFlags.IsError))
+                    {
+                        return BadRequest(new ProblemDetails
+                        {
+                            Title = $"Package contained attachment with error {entry.AttachmentHash}\""
+                        });
+                    }
+                    if (entry.Flags.HasFlag(CbPackageAttachmentFlags.IsCompressed))
+                    {
+                        using MemoryBufferedPayload payload = new MemoryBufferedPayload(blob);
+                        await _blobStore.PutCompressedObject(
+                            ns, payload, ContentId.FromIoHash(entry.AttachmentHash),
+                            _contentIdStore, _compressedBufferUtils);
+                    }
+                    else
+                    {
+                        await _blobStore.PutObject(ns, blob, BlobIdentifier.FromIoHash(entry.AttachmentHash));
+                    }
+                }
+            }
+            catch (HashMismatchException e)
+            {
+                return BadRequest(new ProblemDetails
+                {
+                    Title = $"Incorrect hash, got hash \"{e.SuppliedHash}\" but hash of content was determined to be \"{e.ContentHash}\""
+                });
+            }
+            
+            CbObject rootObject = packageReader.RootObject;
+            BlobIdentifier rootObjectHash = BlobIdentifier.FromIoHash(packageReader.RootHash);
+
+            (ContentId[] missingReferences, BlobIdentifier[] missingBlobs) = await _objectService.Put(ns, bucket, key, rootObjectHash, rootObject);
+
+            List<ContentHash> missingHashes = new List<ContentHash>(missingReferences);
+            missingHashes.AddRange(missingBlobs);
+            return Ok(new PutObjectResponse(missingHashes.ToArray()));
+        }
+
         [HttpPost("{ns}/{bucket}/{key}/finalize/{hash}.{format?}")]
         public async Task<IActionResult> FinalizeObject(
             [FromRoute] [Required] NamespaceId ns,
@@ -674,7 +809,7 @@ namespace Horde.Storage.Controllers
 
                     if (op.ResolveAttachments ?? false)
                     {
-                        IAsyncEnumerable<BlobIdentifier> references = _referenceResolver.ResolveReferences(ns, cb);
+                        IAsyncEnumerable<BlobIdentifier> references = _referenceResolver.GetReferencedBlobs(ns, cb);
                         List<BlobIdentifier>? _ = await references.ToListAsync();
                     }
 
@@ -708,7 +843,7 @@ namespace Horde.Storage.Controllers
                         byte[] blobContents = await blob.Stream.ToByteArray();
                         CbObject cb = new CbObject(blobContents);
                         // the reference resolver will throw if any blob is missing, so no need to do anything other then process each reference
-                        IAsyncEnumerable<BlobIdentifier> references = _referenceResolver.ResolveReferences(ns, cb);
+                        IAsyncEnumerable<BlobIdentifier> references = _referenceResolver.GetReferencedBlobs(ns, cb);
                         List<BlobIdentifier>? _ = await references.ToListAsync();
                     }
 
diff --git a/Engine/Source/Programs/Horde/HordeStorage/Horde.Storage/Implementation/Blob/IBlobService.cs b/Engine/Source/Programs/Horde/HordeStorage/Horde.Storage/Implementation/Blob/IBlobService.cs
--- a/Engine/Source/Programs/Horde/HordeStorage/Horde.Storage/Implementation/Blob/IBlobService.cs
+++ b/Engine/Source/Programs/Horde/HordeStorage/Horde.Storage/Implementation/Blob/IBlobService.cs
@@ -7,6 +7,7 @@ using System.IO;
 using System.Linq;
 using System.Net;
 using System.Net.Http;
+using System.Net.Mime;
 using System.Threading.Tasks;
 using Dasync.Collections;
 using Datadog.Trace;
@@ -17,10 +18,12 @@ using Horde.Storage.Implementation.Blob;
 using Jupiter.Common;
 using Jupiter.Common.Implementation;
 using Jupiter.Implementation;
+using Jupiter.Utils;
 using Microsoft.AspNetCore.Http;
 using Microsoft.Extensions.DependencyInjection;
 using Microsoft.Extensions.Options;
 using Serilog;
+using CustomMediaTypeNames = Jupiter.CustomMediaTypeNames;
 
 namespace Horde.Storage.Implementation;
 
@@ -548,3 +551,96 @@ public class BlobService : IBlobService
         return new BlobContents(ms, ms.Length);
     }
 }
+
+public static class BlobServiceExtensions
+{
+    public static async Task<ContentId> PutCompressedObject(this IBlobService blobService, NamespaceId ns, IBufferedPayload payload, ContentId? id, IContentIdStore contentIdStore, CompressedBufferUtils compressedBufferUtils)
+    {
+        // decompress the content and generate a identifier from it to verify the identifier we got
+        await using Stream decompressStream = payload.GetStream();
+        // TODO: we should add a overload for decompress content that can work on streams, otherwise we are still limited to 2GB compressed blobs
+        byte[] decompressedContent = compressedBufferUtils.DecompressContent(await decompressStream.ToByteArray());
+
+        await using MemoryStream decompressedStream = new MemoryStream(decompressedContent);
+        ContentId identifierDecompressedPayload;
+        if (id != null)
+        {
+            identifierDecompressedPayload = ContentId.FromContentHash(await blobService.VerifyContentMatchesHash(decompressedStream, id));
+        }
+        else
+        {
+            ContentHash blobHash;
+            {
+                using IScope _ = Tracer.Instance.StartActive("web.hash");
+                blobHash = await BlobIdentifier.FromStream(decompressedStream);
+            }
+
+            identifierDecompressedPayload = ContentId.FromContentHash(blobHash);
+        }
+
+        BlobIdentifier identifierCompressedPayload;
+        {
+            using IScope _ = Tracer.Instance.StartActive("web.hash");
+            await using Stream hashStream = payload.GetStream();
+            identifierCompressedPayload = await BlobIdentifier.FromStream(hashStream);
+        }
+
+        // commit the mapping from the decompressed hash to the compressed hash, we run this in parallel with the blob store submit
+        // TODO: let users specify weight of the blob compared to previously submitted content ids
+        int contentIdWeight = (int)payload.Length;
+        Task contentIdStoreTask = contentIdStore.Put(ns, identifierDecompressedPayload, identifierCompressedPayload, contentIdWeight);
+
+        // we still commit the compressed buffer to the object store using the hash of the compressed content
+        {
+            await blobService.PutObjectKnownHash(ns, payload, identifierCompressedPayload);
+        }
+
+        await contentIdStoreTask;
+
+        return identifierDecompressedPayload;
+    }
+
+    public static async Task<(BlobContents, string)> GetCompressedObject(this IBlobService blobService, NamespaceId ns, ContentId contentId, IContentIdStore contentIdStore)
+    {
+        BlobIdentifier[]? chunks = await contentIdStore.Resolve(ns, contentId, mustBeContentId: false);
+        if (chunks == null || chunks.Length == 0)
+        {
+            throw new ContentIdResolveException(contentId);
+        }
+
+        // single chunk, we just return that chunk
+        if (chunks.Length == 1)
+        {
+            BlobIdentifier blobToReturn = chunks[0];
+            string mimeType = CustomMediaTypeNames.UnrealCompressedBuffer;
+            if (contentId.Equals(blobToReturn))
+            {
+                // this was actually the unmapped blob, meaning its not a compressed buffer
+                mimeType = MediaTypeNames.Application.Octet;
+            }
+
+            return (await blobService.GetObject(ns, blobToReturn), mimeType);
+        }
+
+        // chunked content, combine the chunks into a single stream
+        using IScope _ = Tracer.Instance.StartActive("blob.combine");
+        Task<BlobContents>[] tasks = new Task<BlobContents>[chunks.Length];
+        for (int i = 0; i < chunks.Length; i++)
+        {
+            tasks[i] = blobService.GetObject(ns, chunks[i]);
+        }
+
+        MemoryStream ms = new MemoryStream();
+        foreach (Task<BlobContents> task in tasks)
+        {
+            BlobContents blob = await task;
+            await using Stream s = blob.Stream;
+            await s.CopyToAsync(ms);
+        }
+
+        ms.Seek(0, SeekOrigin.Begin);
+
+        // chunking could not have happened for a non compressed buffer so assume it is compressed
+        return (new BlobContents(ms, ms.Length), CustomMediaTypeNames.UnrealCompressedBuffer);
+    }
+}
\ No newline at end of file
diff --git a/Engine/Source/Programs/Horde/HordeStorage/Horde.Storage/Implementation/IReferenceResolver.cs b/Engine/Source/Programs/Horde/HordeStorage/Horde.Storage/Implementation/IReferenceResolver.cs
--- a/Engine/Source/Programs/Horde/HordeStorage/Horde.Storage/Implementation/IReferenceResolver.cs
+++ b/Engine/Source/Programs/Horde/HordeStorage/Horde.Storage/Implementation/IReferenceResolver.cs
@@ -13,9 +13,61 @@ using Jupiter.Utils;
 
 namespace Horde.Storage.Implementation
 {
+    public abstract class Attachment
+    {
+        public abstract IoHash AsIoHash();
+    }
+
+    public class BlobAttachment : Attachment
+    {
+        public BlobIdentifier Identifier { get; }
+
+        public BlobAttachment(BlobIdentifier blobIdentifier)
+        {
+            Identifier = blobIdentifier;
+        }
+
+        public override IoHash AsIoHash()
+        {
+            return Identifier.AsIoHash();
+        }
+    }
+
+    public class ObjectAttachment : Attachment
+    {
+        public BlobIdentifier Identifier { get; }
+
+        public ObjectAttachment(BlobIdentifier blobIdentifier)
+        {
+            Identifier = blobIdentifier;
+        }
+
+        public override IoHash AsIoHash()
+        {
+            return Identifier.AsIoHash();
+        }
+    }
+
+    public class ContentIdAttachment : Attachment
+    {
+        public ContentId Identifier { get; }
+
+        public ContentIdAttachment(ContentId contentId)
+        {
+            Identifier = contentId;
+        }
+
+        public override IoHash AsIoHash()
+        {
+            return Identifier.AsIoHash();
+        }
+    }
+
     public interface IReferenceResolver
     {
-        IAsyncEnumerable<BlobIdentifier> ResolveReferences(NamespaceId ns, CbObject cb);
+
+        IAsyncEnumerable<BlobIdentifier> GetReferencedBlobs(NamespaceId ns, CbObject cb);
+        IAsyncEnumerable<Attachment> GetAttachments(NamespaceId ns, CbObject cb);
     }
 
     public class ReferenceResolver : IReferenceResolver
@@ -29,78 +81,70 @@ namespace Horde.Storage.Implementation
             _contentIdStore = contentIdStore;
         }
 
-        public async IAsyncEnumerable<BlobIdentifier> ResolveReferences(NamespaceId ns, CbObject cb)
+        public async IAsyncEnumerable<Attachment> GetAttachments(NamespaceId ns, CbObject cb)
         {
             // TODO: This is cacheable and we should store the result somewhere
             Queue<CbObject> objectsToVisit = new Queue<CbObject>();
             objectsToVisit.Enqueue(cb);
-            List<ContentId> unresolvedContentIdReferences = new List<ContentId>();
             List<BlobIdentifier> unresolvedBlobReferences = new List<BlobIdentifier>();
 
-            List<Task<(ContentId, BlobIdentifier[]?)>> pendingContentIdResolves = new();
             List<Task<CbObject>> pendingCompactBinaryAttachments = new();
 
-            while(pendingCompactBinaryAttachments.Count != 0 || pendingContentIdResolves.Count != 0 || objectsToVisit.Count != 0)
+            while (pendingCompactBinaryAttachments.Count != 0 || objectsToVisit.Count != 0)
             {
                 if (objectsToVisit.TryDequeue(out CbObject? parent))
                 {
-                    List<BlobIdentifier> blobIdentifiers = new List<BlobIdentifier>();
+                    List<Attachment> attachments = new List<Attachment>();
 
                     parent.IterateAttachments(field =>
                     {
                         IoHash attachmentHash = field.AsAttachment();
 
                         BlobIdentifier blobIdentifier = BlobIdentifier.FromIoHash(attachmentHash);
+                        ContentId contentId = ContentId.FromIoHash(attachmentHash);
+
                         if (field.IsBinaryAttachment())
                         {
-                            pendingContentIdResolves.Add(ResolveContentId(ns, ContentId.FromIoHash(attachmentHash)));
+                            bool isContentId = false;
+                            try
+                            {
+                                // TODO: Having to do a proper resolve here is not good as its quite expensive
+                                // it would be much better if the attachment type indicated if it was a content id or not
+                                Task<BlobIdentifier[]?> resolveContentId = _contentIdStore.Resolve(ns, contentId, mustBeContentId: true);
+                                resolveContentId.Wait();
+                                isContentId = resolveContentId.Result != null;
+                            }
+                            catch (InvalidContentIdException)
+                            {
+                                isContentId = false;
+                            }
+                            
+                            if (isContentId)
+                            {
+                                attachments.Add(new ContentIdAttachment(contentId));
+                            }
+                            else
+                            {
+                                attachments.Add(new BlobAttachment(blobIdentifier));
+                            }
                         }
                         else if (field.IsObjectAttachment())
                         {
-                            blobIdentifiers.Add(blobIdentifier);
+                            attachments.Add(new ObjectAttachment(blobIdentifier));
                             pendingCompactBinaryAttachments.Add(ParseCompactBinaryAttachment(ns, blobIdentifier));
                         }
                         else
-                        { 
+                        {
                             throw new NotImplementedException($"Unknown attachment type for field {field}");
                         }
                     });
 
-                    foreach (BlobIdentifier blobIdentifier in blobIdentifiers)
+                    foreach (Attachment attachment in attachments)
                     {
-                        yield return blobIdentifier;
+                        yield return attachment;
                     }
                 }
 
-                List<Task<(ContentId, BlobIdentifier[]?)>> contentIdResolvesToRemove = new();
-
-                foreach (Task<(ContentId, BlobIdentifier[]?)> pendingContentIdResolveTask in pendingContentIdResolves)
-                {
-                    // check for any content id resolve that has finished and return those blobs it found
-                    if (pendingContentIdResolveTask.IsCompleted)
-                    {
-                        (ContentId contentId, BlobIdentifier[]? resolvedBlobs) = await pendingContentIdResolveTask;
-                        if (resolvedBlobs != null)
-                        {
-                            foreach (BlobIdentifier b in resolvedBlobs)
-                            {
-                                yield return b;
-                            }
-                        }
-                        else
-                        {
-                            unresolvedContentIdReferences.Add(contentId);
-                        }
-                        contentIdResolvesToRemove.Add(pendingContentIdResolveTask);
-                    }
-                }
-
-                // cleanup finished tasks
-                foreach (Task<(ContentId, BlobIdentifier[]?)> finishedTask in contentIdResolvesToRemove)
-                {
-                    pendingContentIdResolves.Remove(finishedTask);
-                }
-
                 // check for any compact binary attachment fetches and add those to the objects we are handling
                 List<Task<CbObject>> finishedCompactBinaryResolves = new();
                 foreach (Task<CbObject> pendingCompactBinaryAttachment in pendingCompactBinaryAttachments)
@@ -127,20 +171,67 @@ namespace Horde.Storage.Implementation
                 }
 
                 // if there are pending resolves left, wait for one of them to finish to avoid busy waiting
-                if (pendingContentIdResolves.Any() || pendingCompactBinaryAttachments.Any())
+                if (pendingCompactBinaryAttachments.Any())
                 {
-                    await Task.WhenAny(pendingContentIdResolves.Concat<Task>(pendingCompactBinaryAttachments));
+                    await Task.WhenAny(pendingCompactBinaryAttachments);
                 }
             }
 
-            if (unresolvedContentIdReferences.Count != 0)
+            if (unresolvedBlobReferences.Count != 0)
             {
-                throw new PartialReferenceResolveException(unresolvedContentIdReferences);
+                throw new ReferenceIsMissingBlobsException(unresolvedBlobReferences);
             }
+        }
 
-            if (unresolvedBlobReferences.Count != 0)
+        public async IAsyncEnumerable<BlobIdentifier> GetReferencedBlobs(NamespaceId ns, CbObject cb)
+        {
+            List<Task<(ContentId, BlobIdentifier[]?)>> pendingContentIdResolves = new();
+            List<ContentId> unresolvedContentIdReferences = new List<ContentId>();
+
+            // Resolve all the attachments
+            await foreach (Attachment attachment in GetAttachments(ns, cb))
             {
-                throw new ReferenceIsMissingBlobsException(unresolvedBlobReferences);
+                if (attachment is BlobAttachment blobAttachment)
+                {
+                    yield return blobAttachment.Identifier;
+                }
+                else if (attachment is ContentIdAttachment contentIdAttachment)
+                {
+                    // If we find a content id we resolve that into the actual blobs it references
+                    pendingContentIdResolves.Add(ResolveContentId(ns, contentIdAttachment.Identifier));
+                }
+                else if (attachment is ObjectAttachment objectAttachment)
+                {
+                    // a object just references the same blob, traversing the object attachment is done in GetAttachments
+                    yield return objectAttachment.Identifier;
+                }
+                else
+                {
+                    throw new NotSupportedException($"Unknown attachment type {attachment.GetType()}");
+                }
+            }
+
+            // Return the results of all the content id resolves
+            foreach (Task<(ContentId, BlobIdentifier[]?)> pendingContentIdResolveTask in pendingContentIdResolves)
+            {
+                (ContentId contentId, BlobIdentifier[]? resolvedBlobs) =  await pendingContentIdResolveTask;
+                if (resolvedBlobs != null)
+                {
+                    foreach (BlobIdentifier b in resolvedBlobs)
+                    {
+                        yield return b;
+                    }
+                }
+                else
+                {
+                    unresolvedContentIdReferences.Add(contentId);
+                }
+            }
+
+            // if there were any content ids we did not recognize we throw a partial reference exception
+            if (unresolvedContentIdReferences.Count != 0)
+            {
+                throw new PartialReferenceResolveException(unresolvedContentIdReferences);
             }
         }
 
diff --git a/Engine/Source/Programs/Horde/HordeStorage/Horde.Storage/Implementation/ObjectService.cs b/Engine/Source/Programs/Horde/HordeStorage/Horde.Storage/Implementation/ObjectService.cs
--- a/Engine/Source/Programs/Horde/HordeStorage/Horde.Storage/Implementation/ObjectService.cs
+++ b/Engine/Source/Programs/Horde/HordeStorage/Horde.Storage/Implementation/ObjectService.cs
@@ -189,7 +189,7 @@ namespace Horde.Storage.Implementation
                 using IScope _ = Tracer.Instance.StartActive("ObjectService.ResolveReferences");
                 try
                 {
-                    IAsyncEnumerable<BlobIdentifier> references = _referenceResolver.ResolveReferences(ns, payload);
+                    IAsyncEnumerable<BlobIdentifier> references = _referenceResolver.GetReferencedBlobs(ns, payload);
                     BlobIdentifier[] referencesArray = await references.ToArrayAsync();
                     Task addRefsTask = _blobIndex.AddRefToBlobs(ns, bucket, key, referencesArray);
                     missingBlobs = await _blobService.FilterOutKnownBlobs(ns, referencesArray);
diff --git a/Engine/Source/Programs/Horde/HordeStorage/Horde.Storage/Tests/Functional/Objects/ReferencesTests.cs b/Engine/Source/Programs/Horde/HordeStorage/Horde.Storage/Tests/Functional/Objects/ReferencesTests.cs
--- a/Engine/Source/Programs/Horde/HordeStorage/Horde.Storage/Tests/Functional/Objects/ReferencesTests.cs
+++ b/Engine/Source/Programs/Horde/HordeStorage/Horde.Storage/Tests/Functional/Objects/ReferencesTests.cs
@@ -161,8 +161,8 @@ namespace Horde.Storage.FunctionalTests.References
     
     public abstract class ReferencesTests
     {
-        private static TestServer? _server;
-        private static HttpClient? _httpClient;
+        public static TestServer? Server { get; private set; }
+        public static HttpClient? Client { get; private set; }
 
         protected IBlobService Service { get; set; } = null!;
         protected IReferencesStore ReferencesStore { get; set; } = null!;
@@ -189,11 +189,11 @@ namespace Horde.Storage.FunctionalTests.References
                 .UseSerilog(logger)
                 .UseStartup<HordeStorageStartup>()
             );
-            _httpClient = server.CreateClient();
-            _server = server;
+            Client = server.CreateClient();
+            Server = server;
 
-            Service = _server.Services.GetService<IBlobService>()!;
-            ReferencesStore = _server.Services.GetService<IReferencesStore>()!;
+            Service = Server.Services.GetService<IBlobService>()!;
+            ReferencesStore = Server.Services.GetService<IReferencesStore>()!;
             await SeedDb(server.Services);
         }
 
@@ -205,9 +205,9 @@ namespace Horde.Storage.FunctionalTests.References
         [TestCleanup]
         public async Task Teardown()
         {
-            if (_server != null)
+            if (Server != null)
             {
-                await TeardownDb(_server.Services);
+                await TeardownDb(Server.Services);
             }
         }
 
@@ -219,19 +219,19 @@ namespace Horde.Storage.FunctionalTests.References
         [TestMethod]
         public async Task PutGetBlob()
         {
-            const string objectContents = "This is treated as a opaque blob";
-            byte[] data = Encoding.ASCII.GetBytes(objectContents);
+            const string ObjectContents = "This is treated as a opaque blob";
+            byte[] data = Encoding.ASCII.GetBytes(ObjectContents);
             BlobIdentifier objectHash = BlobIdentifier.FromBlob(data);
             IoHashKey key = IoHashKey.FromName("newBlobObject");
             using HttpContent requestContent = new ByteArrayContent(data);
             requestContent.Headers.ContentType = new MediaTypeHeaderValue(MediaTypeNames.Application.Octet);
             requestContent.Headers.Add(CommonHeaders.HashHeaderName, objectHash.ToString());
 
-            HttpResponseMessage result = await _httpClient!.PutAsync(new Uri($"api/v1/refs/{TestNamespace}/bucket/{key}", UriKind.Relative), requestContent);
+            HttpResponseMessage result = await Client!.PutAsync(new Uri($"api/v1/refs/{TestNamespace}/bucket/{key}", UriKind.Relative), requestContent);
             result.EnsureSuccessStatusCode();
 
             {
-                HttpResponseMessage getResponse = await _httpClient.GetAsync(new Uri($"api/v1/refs/{TestNamespace}/bucket/{key}.raw", UriKind.Relative));
+                HttpResponseMessage getResponse = await Client.GetAsync(new Uri($"api/v1/refs/{TestNamespace}/bucket/{key}.raw", UriKind.Relative));
                 getResponse.EnsureSuccessStatusCode();
                 await using MemoryStream ms = new MemoryStream();
                 await getResponse.Content.CopyToAsync(ms);
@@ -239,7 +239,7 @@ namespace Horde.Storage.FunctionalTests.References
                 byte[] roundTrippedBuffer = ms.ToArray();
                 string roundTrippedPayload = Encoding.ASCII.GetString(roundTrippedBuffer);
 
-                Assert.AreEqual(objectContents, roundTrippedPayload);
+                Assert.AreEqual(ObjectContents, roundTrippedPayload);
                 CollectionAssert.AreEqual(data, roundTrippedBuffer);
                 Assert.AreEqual(objectHash, BlobIdentifier.FromBlob(roundTrippedBuffer));
             }
@@ -247,7 +247,7 @@ namespace Horde.Storage.FunctionalTests.References
             {
                 BlobIdentifier attachment;
                 {
-                    HttpResponseMessage getResponse = await _httpClient.GetAsync(new Uri($"api/v1/refs/{TestNamespace}/bucket/{key}.uecb", UriKind.Relative));
+                    HttpResponseMessage getResponse = await Client.GetAsync(new Uri($"api/v1/refs/{TestNamespace}/bucket/{key}.uecb", UriKind.Relative));
                     getResponse.EnsureSuccessStatusCode();
                     await using MemoryStream ms = new MemoryStream();
                     await getResponse.Content.CopyToAsync(ms);
@@ -264,20 +264,20 @@ namespace Horde.Storage.FunctionalTests.References
                 }
 
                 {
-                    HttpResponseMessage getAttachment = await _httpClient.GetAsync(new Uri($"api/v1/blobs/{TestNamespace}/{attachment}", UriKind.Relative));
+                    HttpResponseMessage getAttachment = await Client.GetAsync(new Uri($"api/v1/blobs/{TestNamespace}/{attachment}", UriKind.Relative));
                     getAttachment.EnsureSuccessStatusCode();
                     await using MemoryStream ms = new MemoryStream();
                     await getAttachment.Content.CopyToAsync(ms);
                     byte[] roundTrippedBuffer = ms.ToArray();
                     string roundTrippedString = Encoding.ASCII.GetString(roundTrippedBuffer);
 
-                    Assert.AreEqual(objectContents, roundTrippedString);
+                    Assert.AreEqual(ObjectContents, roundTrippedString);
                     Assert.AreEqual(objectHash, BlobIdentifier.FromBlob(roundTrippedBuffer));
                 }
             }
 
             {
-                HttpResponseMessage getResponse = await _httpClient.GetAsync(new Uri($"api/v1/refs/{TestNamespace}/bucket/{key}.json", UriKind.Relative));
+                HttpResponseMessage getResponse = await Client.GetAsync(new Uri($"api/v1/refs/{TestNamespace}/bucket/{key}.json", UriKind.Relative));
                 getResponse.EnsureSuccessStatusCode();
                 await using MemoryStream ms = new MemoryStream();
                 await getResponse.Content.CopyToAsync(ms);
@@ -303,7 +303,7 @@ namespace Horde.Storage.FunctionalTests.References
                 // request the object as a json response using accept instead of the format filter
                 using HttpRequestMessage request = new(HttpMethod.Get, new Uri($"api/v1/refs/{TestNamespace}/bucket/{key}", UriKind.Relative));
                 request.Headers.Accept.Add(new MediaTypeWithQualityHeaderValue(MediaTypeNames.Application.Json));
-                HttpResponseMessage getResponse = await _httpClient.SendAsync(request);
+                HttpResponseMessage getResponse = await Client.SendAsync(request);
                 getResponse.EnsureSuccessStatusCode();
                 Assert.AreEqual(MediaTypeNames.Application.Json, getResponse.Content.Headers.ContentType?.MediaType);
 
@@ -331,7 +331,7 @@ namespace Horde.Storage.FunctionalTests.References
                 // request the object as a jupiter inlined payload
                 using HttpRequestMessage request = new (HttpMethod.Get, new Uri($"api/v1/refs/{TestNamespace}/bucket/{key}", UriKind.Relative));
                 request.Headers.Accept.Add(new MediaTypeWithQualityHeaderValue(CustomMediaTypeNames.JupiterInlinedPayload));
-                HttpResponseMessage getResponse = await _httpClient.SendAsync(request);
+                HttpResponseMessage getResponse = await Client.SendAsync(request);
                 getResponse.EnsureSuccessStatusCode();
                 Assert.AreEqual(CustomMediaTypeNames.JupiterInlinedPayload, getResponse.Content.Headers.ContentType?.MediaType);
 
@@ -341,7 +341,7 @@ namespace Horde.Storage.FunctionalTests.References
 
                 string roundTrippedString = Encoding.ASCII.GetString(roundTrippedBuffer);
 
-                Assert.AreEqual(objectContents, roundTrippedString);
+                Assert.AreEqual(ObjectContents, roundTrippedString);
                 Assert.AreEqual(objectHash, BlobIdentifier.FromBlob(roundTrippedBuffer));
             }
         }
@@ -362,7 +362,7 @@ namespace Horde.Storage.FunctionalTests.References
             requestContent.Headers.ContentType = new MediaTypeHeaderValue(CustomMediaTypeNames.UnrealCompactBinary);
             requestContent.Headers.Add(CommonHeaders.HashHeaderName, objectHash.ToString());
 
-            using HttpResponseMessage result = await _httpClient!.PutAsync(new Uri($"api/v1/refs/{TestNamespace}/bucket/{key}.uecb", UriKind.Relative), requestContent);
+            using HttpResponseMessage result = await Client!.PutAsync(new Uri($"api/v1/refs/{TestNamespace}/bucket/{key}.uecb", UriKind.Relative), requestContent);
             result.EnsureSuccessStatusCode();
 
             {
@@ -390,7 +390,7 @@ namespace Horde.Storage.FunctionalTests.References
             }
 
             {
-                HttpResponseMessage getResponse = await _httpClient.GetAsync(new Uri($"api/v1/refs/{TestNamespace}/bucket/{key}.raw", UriKind.Relative));
+                HttpResponseMessage getResponse = await Client.GetAsync(new Uri($"api/v1/refs/{TestNamespace}/bucket/{key}.raw", UriKind.Relative));
                 getResponse.EnsureSuccessStatusCode();
                 await using MemoryStream ms = new MemoryStream();
                 await getResponse.Content.CopyToAsync(ms);
@@ -402,7 +402,7 @@ namespace Horde.Storage.FunctionalTests.References
             }
 
             {
-                HttpResponseMessage getResponse = await _httpClient.GetAsync(new Uri($"api/v1/refs/{TestNamespace}/bucket/{key}.uecb", UriKind.Relative));
+                HttpResponseMessage getResponse = await Client.GetAsync(new Uri($"api/v1/refs/{TestNamespace}/bucket/{key}.uecb", UriKind.Relative));
                 getResponse.EnsureSuccessStatusCode();
                 await using MemoryStream ms = new MemoryStream();
                 await getResponse.Content.CopyToAsync(ms);
@@ -418,7 +418,7 @@ namespace Horde.Storage.FunctionalTests.References
             }
 
             {
-                HttpResponseMessage getResponse = await _httpClient.GetAsync(new Uri($"api/v1/refs/{TestNamespace}/bucket/{key}.json", UriKind.Relative));
+                HttpResponseMessage getResponse = await Client.GetAsync(new Uri($"api/v1/refs/{TestNamespace}/bucket/{key}.json", UriKind.Relative));
                 getResponse.EnsureSuccessStatusCode();
                 await using MemoryStream ms = new MemoryStream();
                 await getResponse.Content.CopyToAsync(ms);
@@ -460,7 +460,7 @@ namespace Horde.Storage.FunctionalTests.References
             requestContent.Headers.ContentType = new MediaTypeHeaderValue(CustomMediaTypeNames.UnrealCompactBinary);
             requestContent.Headers.Add(CommonHeaders.HashHeaderName, objectHash.ToString());
 
-            HttpResponseMessage result = await _httpClient!.PutAsync(new Uri($"api/v1/refs/{TestNamespace}/bucket/{key}.uecb", UriKind.Relative), requestContent);
+            HttpResponseMessage result = await Client!.PutAsync(new Uri($"api/v1/refs/{TestNamespace}/bucket/{key}.uecb", UriKind.Relative), requestContent);
             result.EnsureSuccessStatusCode();
 
             {
@@ -488,7 +488,7 @@ namespace Horde.Storage.FunctionalTests.References
             }
             
             {
-                HttpResponseMessage getResponse = await _httpClient.GetAsync(new Uri($"api/v1/refs/{TestNamespace}/bucket/{key}.json?fields=name", UriKind.Relative));
+                HttpResponseMessage getResponse = await Client.GetAsync(new Uri($"api/v1/refs/{TestNamespace}/bucket/{key}.json?fields=name", UriKind.Relative));
                 getResponse.EnsureSuccessStatusCode();
                 await using MemoryStream ms = new MemoryStream();
                 await getResponse.Content.CopyToAsync(ms);
@@ -538,7 +538,7 @@ namespace Horde.Storage.FunctionalTests.References
                 requestContent.Headers.ContentType = new MediaTypeHeaderValue(CustomMediaTypeNames.UnrealCompactBinary);
                 requestContent.Headers.Add(CommonHeaders.HashHeaderName, parentObjectHash.ToString());
 
-                HttpResponseMessage result = await _httpClient!.PutAsync(new Uri($"api/v1/refs/{TestNamespace}/bucket/{key}.uecb", UriKind.Relative), requestContent);
+                HttpResponseMessage result = await Client!.PutAsync(new Uri($"api/v1/refs/{TestNamespace}/bucket/{key}.uecb", UriKind.Relative), requestContent);
                 result.EnsureSuccessStatusCode();
 
                 Assert.AreEqual(CustomMediaTypeNames.UnrealCompactBinary, result!.Content.Headers.ContentType!.MediaType);
@@ -564,7 +564,7 @@ namespace Horde.Storage.FunctionalTests.References
                 requestContent.Headers.Add(CommonHeaders.HashHeaderName, childObjectHash.ToString());
                 request.Content = requestContent;
 
-                HttpResponseMessage result = await _httpClient!.SendAsync(request);
+                HttpResponseMessage result = await Client!.SendAsync(request);
                 result.EnsureSuccessStatusCode();
 
                 Assert.AreEqual(CustomMediaTypeNames.UnrealCompactBinary, result!.Content.Headers.ContentType!.MediaType);
@@ -584,7 +584,7 @@ namespace Horde.Storage.FunctionalTests.References
                 requestContent.Headers.ContentType = new MediaTypeHeaderValue(CustomMediaTypeNames.UnrealCompactBinary);
                 requestContent.Headers.Add(CommonHeaders.HashHeaderName, parentObjectHash.ToString());
 
-                HttpResponseMessage result = await _httpClient!.PutAsync(new Uri($"api/v1/refs/{TestNamespace}/bucket/{key}.uecb", UriKind.Relative), requestContent);
+                HttpResponseMessage result = await Client!.PutAsync(new Uri($"api/v1/refs/{TestNamespace}/bucket/{key}.uecb", UriKind.Relative), requestContent);
                 result.EnsureSuccessStatusCode();
 
                 Assert.AreEqual(result!.Content.Headers.ContentType!.MediaType, CustomMediaTypeNames.UnrealCompactBinary);
@@ -610,7 +610,7 @@ namespace Horde.Storage.FunctionalTests.References
             }
 
             {
-                HttpResponseMessage getResponse = await _httpClient.GetAsync(new Uri($"api/v1/refs/{TestNamespace}/bucket/{key}.raw", UriKind.Relative));
+                HttpResponseMessage getResponse = await Client.GetAsync(new Uri($"api/v1/refs/{TestNamespace}/bucket/{key}.raw", UriKind.Relative));
                 getResponse.EnsureSuccessStatusCode();
                 await using MemoryStream ms = new MemoryStream();
                 await getResponse.Content.CopyToAsync(ms);
@@ -622,7 +622,7 @@ namespace Horde.Storage.FunctionalTests.References
             }
 
             {
-                HttpResponseMessage getResponse = await _httpClient.GetAsync(new Uri($"api/v1/refs/{TestNamespace}/bucket/{key}.uecb", UriKind.Relative));
+                HttpResponseMessage getResponse = await Client.GetAsync(new Uri($"api/v1/refs/{TestNamespace}/bucket/{key}.uecb", UriKind.Relative));
                 getResponse.EnsureSuccessStatusCode();
                 await using MemoryStream ms = new MemoryStream();
                 await getResponse.Content.CopyToAsync(ms);
@@ -638,7 +638,7 @@ namespace Horde.Storage.FunctionalTests.References
             }
 
             {
-                HttpResponseMessage getResponse = await _httpClient.GetAsync(new Uri($"api/v1/refs/{TestNamespace}/bucket/{key}.json", UriKind.Relative));
+                HttpResponseMessage getResponse = await Client.GetAsync(new Uri($"api/v1/refs/{TestNamespace}/bucket/{key}.json", UriKind.Relative));
                 getResponse.EnsureSuccessStatusCode();
                 await using MemoryStream ms = new MemoryStream();
                 await getResponse.Content.CopyToAsync(ms);
@@ -663,7 +663,7 @@ namespace Horde.Storage.FunctionalTests.References
             }
 
             {
-                HttpResponseMessage getResponse = await _httpClient.GetAsync(new Uri($"api/v1/objects/{TestNamespace}/{childObjectHash}", UriKind.Relative));
+                HttpResponseMessage getResponse = await Client.GetAsync(new Uri($"api/v1/objects/{TestNamespace}/{childObjectHash}", UriKind.Relative));
                 getResponse.EnsureSuccessStatusCode();
                 await using MemoryStream ms = new MemoryStream();
                 await getResponse.Content.CopyToAsync(ms);
@@ -678,8 +678,8 @@ namespace Horde.Storage.FunctionalTests.References
         [TestMethod]
         public async Task ExistsChecks()
         {
-            const string objectContents = "This is treated as a opaque blob";
-            byte[] data = Encoding.ASCII.GetBytes(objectContents);
+            const string ObjectContents = "This is treated as a opaque blob";
+            byte[] data = Encoding.ASCII.GetBytes(ObjectContents);
             BlobIdentifier objectHash = BlobIdentifier.FromBlob(data);
             IoHashKey key = IoHashKey.FromName("newObject");
             using HttpContent requestContent = new ByteArrayContent(data);
@@ -687,19 +687,19 @@ namespace Horde.Storage.FunctionalTests.References
             requestContent.Headers.Add(CommonHeaders.HashHeaderName, objectHash.ToString());
 
             {
-                HttpResponseMessage result = await _httpClient!.PutAsync(new Uri($"api/v1/refs/{TestNamespace}/bucket/{key}", UriKind.Relative), requestContent);
+                HttpResponseMessage result = await Client!.PutAsync(new Uri($"api/v1/refs/{TestNamespace}/bucket/{key}", UriKind.Relative), requestContent);
                 result.EnsureSuccessStatusCode();
             }
 
             {
                 using HttpRequestMessage message = new(HttpMethod.Head, new Uri($"api/v1/refs/{TestNamespace}/bucket/{key}", UriKind.Relative));
-                HttpResponseMessage result = await _httpClient!.SendAsync(message);
+                HttpResponseMessage result = await Client!.SendAsync(message);
                 result.EnsureSuccessStatusCode();
             }
 
             {
                 using HttpRequestMessage message = new(HttpMethod.Head, new Uri($"api/v1/refs/{TestNamespace}/bucket/{IoHashKey.FromName("missingObject")}", UriKind.Relative));
-                HttpResponseMessage result = await _httpClient!.SendAsync(message);
+                HttpResponseMessage result = await Client!.SendAsync(message);
                 Assert.AreEqual(HttpStatusCode.NotFound, result.StatusCode);
             }
         }
@@ -710,15 +710,15 @@ namespace Horde.Storage.FunctionalTests.References
             BucketId bucket = new BucketId("bucket");
             IoHashKey existingObject = IoHashKey.FromName("existingObject");
 
-            const string objectContents = "This is treated as a opaque blob";
-            byte[] data = Encoding.ASCII.GetBytes(objectContents);
+            const string ObjectContents = "This is treated as a opaque blob";
+            byte[] data = Encoding.ASCII.GetBytes(ObjectContents);
             BlobIdentifier objectHash = BlobIdentifier.FromBlob(data);
             using HttpContent requestContent = new ByteArrayContent(data);
             requestContent.Headers.ContentType = new MediaTypeHeaderValue(MediaTypeNames.Application.Octet);
             requestContent.Headers.Add(CommonHeaders.HashHeaderName, objectHash.ToString());
 
             {
-                HttpResponseMessage result = await _httpClient!.PutAsync(new Uri($"api/v1/refs/{TestNamespace}/{bucket}/{existingObject}", UriKind.Relative), requestContent);
+                HttpResponseMessage result = await Client!.PutAsync(new Uri($"api/v1/refs/{TestNamespace}/{bucket}/{existingObject}", UriKind.Relative), requestContent);
                 result.EnsureSuccessStatusCode();
             }
 
@@ -727,7 +727,7 @@ namespace Horde.Storage.FunctionalTests.References
             string queryString = $"?names={bucket}.{existingObject}&names={bucket}.{missingObject}";
 
             {
-                HttpResponseMessage result = await _httpClient!.GetAsync(new Uri($"api/v1/refs/{TestNamespace}/exists" + queryString, UriKind.Relative));
+                HttpResponseMessage result = await Client!.GetAsync(new Uri($"api/v1/refs/{TestNamespace}/exists" + queryString, UriKind.Relative));
                 result.EnsureSuccessStatusCode();
                 ExistCheckMultipleRefsResponse response = await result.Content.ReadAsAsync<ExistCheckMultipleRefsResponse>();
 
@@ -775,7 +775,7 @@ namespace Horde.Storage.FunctionalTests.References
             requestContent.Headers.ContentType = new MediaTypeHeaderValue(CustomMediaTypeNames.UnrealCompactBinary);
             requestContent.Headers.Add(CommonHeaders.HashHeaderName, objectHash.ToString());
 
-            HttpResponseMessage result = await _httpClient!.PutAsync(new Uri($"api/v1/refs/{TestNamespace}/bucket/{key}.uecb", UriKind.Relative), requestContent);
+            HttpResponseMessage result = await Client!.PutAsync(new Uri($"api/v1/refs/{TestNamespace}/bucket/{key}.uecb", UriKind.Relative), requestContent);
             result.EnsureSuccessStatusCode();
 
             // check the response
@@ -805,7 +805,7 @@ namespace Horde.Storage.FunctionalTests.References
 
             // verify attachments
             {
-                HttpResponseMessage getResponse = await _httpClient.GetAsync(new Uri($"api/v1/refs/{TestNamespace}/bucket/{key}.raw", UriKind.Relative));
+                HttpResponseMessage getResponse = await Client.GetAsync(new Uri($"api/v1/refs/{TestNamespace}/bucket/{key}.raw", UriKind.Relative));
                 getResponse.EnsureSuccessStatusCode();
                 await using MemoryStream ms = new MemoryStream();
                 await getResponse.Content.CopyToAsync(ms);
@@ -820,7 +820,7 @@ namespace Horde.Storage.FunctionalTests.References
                 BlobIdentifier blobAttachment;
                 BlobIdentifier objectAttachment;
                 {
-                    HttpResponseMessage getResponse = await _httpClient.GetAsync(new Uri($"api/v1/refs/{TestNamespace}/bucket/{key}.uecb", UriKind.Relative));
+                    HttpResponseMessage getResponse = await Client.GetAsync(new Uri($"api/v1/refs/{TestNamespace}/bucket/{key}.uecb", UriKind.Relative));
                     getResponse.EnsureSuccessStatusCode();
                     await using MemoryStream ms = new MemoryStream();
                     await getResponse.Content.CopyToAsync(ms);
@@ -839,7 +839,7 @@ namespace Horde.Storage.FunctionalTests.References
                 }
 
                 {
-                    HttpResponseMessage getAttachment = await _httpClient.GetAsync(new Uri($"api/v1/blobs/{TestNamespace}/{blobAttachment}", UriKind.Relative));
+                    HttpResponseMessage getAttachment = await Client.GetAsync(new Uri($"api/v1/blobs/{TestNamespace}/{blobAttachment}", UriKind.Relative));
                     getAttachment.EnsureSuccessStatusCode();
                     await using MemoryStream ms = new MemoryStream();
                     await getAttachment.Content.CopyToAsync(ms);
@@ -852,7 +852,7 @@ namespace Horde.Storage.FunctionalTests.References
 
                 BlobIdentifier attachedBlobIdentifier;
                 {
-                    HttpResponseMessage getAttachment = await _httpClient.GetAsync(new Uri($"api/v1/blobs/{TestNamespace}/{objectAttachment}", UriKind.Relative));
+                    HttpResponseMessage getAttachment = await Client.GetAsync(new Uri($"api/v1/blobs/{TestNamespace}/{objectAttachment}", UriKind.Relative));
                     getAttachment.EnsureSuccessStatusCode();
                     await using MemoryStream ms = new MemoryStream();
                     await getAttachment.Content.CopyToAsync(ms);
@@ -868,7 +868,7 @@ namespace Horde.Storage.FunctionalTests.References
                 }
 
                 {
-                    HttpResponseMessage getAttachment = await _httpClient.GetAsync(new Uri($"api/v1/blobs/{TestNamespace}/{attachedBlobIdentifier}", UriKind.Relative));
+                    HttpResponseMessage getAttachment = await Client.GetAsync(new Uri($"api/v1/blobs/{TestNamespace}/{attachedBlobIdentifier}", UriKind.Relative));
                     getAttachment.EnsureSuccessStatusCode();
                     await using MemoryStream ms = new MemoryStream();
                     await getAttachment.Content.CopyToAsync(ms);
@@ -882,7 +882,7 @@ namespace Horde.Storage.FunctionalTests.References
 
             // check json representation
             {
-                HttpResponseMessage getResponse = await _httpClient.GetAsync(new Uri($"api/v1/refs/{TestNamespace}/bucket/{key}.json", UriKind.Relative));
+                HttpResponseMessage getResponse = await Client.GetAsync(new Uri($"api/v1/refs/{TestNamespace}/bucket/{key}.json", UriKind.Relative));
                 getResponse.EnsureSuccessStatusCode();
                 await using MemoryStream ms = new MemoryStream();
                 await getResponse.Content.CopyToAsync(ms);
@@ -941,7 +941,7 @@ namespace Horde.Storage.FunctionalTests.References
                 requestContent.Headers.ContentType = new MediaTypeHeaderValue(CustomMediaTypeNames.UnrealCompactBinary);
                 requestContent.Headers.Add(CommonHeaders.HashHeaderName, objectHash.ToString());
 
-                HttpResponseMessage result = await _httpClient!.PutAsync(new Uri($"api/v1/refs/{TestNamespace}/bucket/{key}.uecb", UriKind.Relative), requestContent);
+                HttpResponseMessage result = await Client!.PutAsync(new Uri($"api/v1/refs/{TestNamespace}/bucket/{key}.uecb", UriKind.Relative), requestContent);
                 result.EnsureSuccessStatusCode();
 
                 {
@@ -964,7 +964,7 @@ namespace Horde.Storage.FunctionalTests.References
                 requestContent.Headers.ContentType = new MediaTypeHeaderValue(CustomMediaTypeNames.UnrealCompactBinary);
                 requestContent.Headers.Add(CommonHeaders.HashHeaderName, objectHash.ToString());
 
-                HttpResponseMessage result = await _httpClient!.PutAsync( new Uri($"api/v1/refs/{TestNamespace}/bucket/{key}.json", UriKind.Relative), requestContent);
+                HttpResponseMessage result = await Client!.PutAsync( new Uri($"api/v1/refs/{TestNamespace}/bucket/{key}.json", UriKind.Relative), requestContent);
                 result.EnsureSuccessStatusCode();
 
                 {
@@ -989,7 +989,7 @@ namespace Horde.Storage.FunctionalTests.References
         [TestMethod]
         public async Task PutContentIdMissingBlob()
         {
-            IContentIdStore? contentIdStore = _server!.Services.GetService<IContentIdStore>();
+            IContentIdStore? contentIdStore = Server!.Services.GetService<IContentIdStore>();
             Assert.IsNotNull(contentIdStore);
 
             // submit a object which contains a content id, which exists but points to a blob that does not exist
@@ -1015,7 +1015,7 @@ namespace Horde.Storage.FunctionalTests.References
                 requestContent.Headers.ContentType = new MediaTypeHeaderValue(CustomMediaTypeNames.UnrealCompactBinary);
                 requestContent.Headers.Add(CommonHeaders.HashHeaderName, objectHash.ToString());
 
-                HttpResponseMessage result = await _httpClient!.PutAsync(new Uri($"api/v1/refs/{TestNamespace}/bucket/{key}.uecb", UriKind.Relative), requestContent);
+                HttpResponseMessage result = await Client!.PutAsync(new Uri($"api/v1/refs/{TestNamespace}/bucket/{key}.uecb", UriKind.Relative), requestContent);
                 result.EnsureSuccessStatusCode();
 
                 {
@@ -1039,7 +1039,7 @@ namespace Horde.Storage.FunctionalTests.References
                 requestContent.Headers.ContentType = new MediaTypeHeaderValue(CustomMediaTypeNames.UnrealCompactBinary);
                 requestContent.Headers.Add(CommonHeaders.HashHeaderName, objectHash.ToString());
 
-                HttpResponseMessage result = await _httpClient!.PutAsync(new Uri($"api/v1/refs/{TestNamespace}/bucket/{key}.json", UriKind.Relative), requestContent);
+                HttpResponseMessage result = await Client!.PutAsync(new Uri($"api/v1/refs/{TestNamespace}/bucket/{key}.json", UriKind.Relative), requestContent);
                 result.EnsureSuccessStatusCode();
 
                 {
@@ -1088,7 +1088,7 @@ namespace Horde.Storage.FunctionalTests.References
                 requestContent.Headers.ContentType = new MediaTypeHeaderValue(CustomMediaTypeNames.UnrealCompactBinary);
                 requestContent.Headers.Add(CommonHeaders.HashHeaderName, objectHash.ToString());
 
-                HttpResponseMessage result = await _httpClient!.PutAsync(new Uri($"api/v1/refs/{TestNamespace}/bucket/{key}.uecb", UriKind.Relative), requestContent);
+                HttpResponseMessage result = await Client!.PutAsync(new Uri($"api/v1/refs/{TestNamespace}/bucket/{key}.uecb", UriKind.Relative), requestContent);
                 result.EnsureSuccessStatusCode();
 
                 {
@@ -1146,7 +1146,7 @@ namespace Horde.Storage.FunctionalTests.References
                 requestContent.Headers.ContentType = new MediaTypeHeaderValue(CustomMediaTypeNames.UnrealCompactBinary);
                 requestContent.Headers.Add(CommonHeaders.HashHeaderName, objectHash.ToString());
 
-                HttpResponseMessage result = await _httpClient!.PutAsync(new Uri($"api/v1/refs/{TestNamespace}/bucket/{key}.uecb", UriKind.Relative), requestContent);
+                HttpResponseMessage result = await Client!.PutAsync(new Uri($"api/v1/refs/{TestNamespace}/bucket/{key}.uecb", UriKind.Relative), requestContent);
                 result.EnsureSuccessStatusCode();
 
                 {
@@ -1183,7 +1183,7 @@ namespace Horde.Storage.FunctionalTests.References
             {
                 using HttpContent requestContent = new ByteArrayContent(Array.Empty<byte>());
 
-                HttpResponseMessage result = await _httpClient!.PostAsync(new Uri($"api/v1/refs/{TestNamespace}/bucket/{key}/finalize/{objectHash}.uecb", UriKind.Relative), requestContent);
+                HttpResponseMessage result = await Client!.PostAsync(new Uri($"api/v1/refs/{TestNamespace}/bucket/{key}/finalize/{objectHash}.uecb", UriKind.Relative), requestContent);
                 result.EnsureSuccessStatusCode();
 
                 {
@@ -1221,7 +1221,7 @@ namespace Horde.Storage.FunctionalTests.References
             requestContent.Headers.ContentType = new MediaTypeHeaderValue(MediaTypeNames.Application.Octet);
             requestContent.Headers.Add(CommonHeaders.HashHeaderName, blobHash.ToString());
 
-            HttpResponseMessage result = await _httpClient!.PutAsync(new Uri($"api/v1/refs/{TestNamespace}/bucket/{key}.uecb", UriKind.Relative), requestContent);
+            HttpResponseMessage result = await Client!.PutAsync(new Uri($"api/v1/refs/{TestNamespace}/bucket/{key}.uecb", UriKind.Relative), requestContent);
             result.EnsureSuccessStatusCode();
 
             {
@@ -1238,7 +1238,7 @@ namespace Horde.Storage.FunctionalTests.References
 
             {
                 // verify we can fetch the blob properly
-                HttpResponseMessage getResponse = await _httpClient.GetAsync(new Uri($"api/v1/refs/{TestNamespace}/bucket/{key}.raw", UriKind.Relative));
+                HttpResponseMessage getResponse = await Client.GetAsync(new Uri($"api/v1/refs/{TestNamespace}/bucket/{key}.raw", UriKind.Relative));
                 getResponse.EnsureSuccessStatusCode();
                 await using MemoryStream ms = new MemoryStream();
                 await getResponse.Content.CopyToAsync(ms);
@@ -1256,19 +1256,19 @@ namespace Horde.Storage.FunctionalTests.References
 
             {
                 // the compact binary object still exists, so this returns a success (trying to resolve the attachment will fail)
-                HttpResponseMessage getResponse = await _httpClient.GetAsync(new Uri($"api/v1/refs/{TestNamespace}/bucket/{key}.uecb", UriKind.Relative));
+                HttpResponseMessage getResponse = await Client.GetAsync(new Uri($"api/v1/refs/{TestNamespace}/bucket/{key}.uecb", UriKind.Relative));
                 getResponse.EnsureSuccessStatusCode();
             }
 
             {
                 // the compact binary object still exists, so this returns a success (trying to resolve the attachment will fail)
-                HttpResponseMessage getResponse = await _httpClient.GetAsync(new Uri($"api/v1/refs/{TestNamespace}/bucket/{key}.json", UriKind.Relative));
+                HttpResponseMessage getResponse = await Client.GetAsync(new Uri($"api/v1/refs/{TestNamespace}/bucket/{key}.json", UriKind.Relative));
                 getResponse.EnsureSuccessStatusCode();
             }
 
             {
                 // we should now see a 404 as the blob is missing
-                HttpResponseMessage getResponse = await _httpClient.GetAsync(new Uri($"api/v1/refs/{TestNamespace}/bucket/{key}.raw", UriKind.Relative));
+                HttpResponseMessage getResponse = await Client.GetAsync(new Uri($"api/v1/refs/{TestNamespace}/bucket/{key}.raw", UriKind.Relative));
                 Assert.AreEqual(HttpStatusCode.NotFound, getResponse.StatusCode);
                 Assert.AreEqual("application/problem+json", getResponse.Content.Headers.ContentType!.MediaType);
                 string s = await getResponse.Content.ReadAsStringAsync();
@@ -1281,8 +1281,8 @@ namespace Horde.Storage.FunctionalTests.References
         [TestMethod]
         public async Task DeleteObject()
         {
-            const string objectContents = "This is treated as a opaque blob";
-            byte[] data = Encoding.ASCII.GetBytes(objectContents);
+            const string ObjectContents = "This is treated as a opaque blob";
+            byte[] data = Encoding.ASCII.GetBytes(ObjectContents);
             BlobIdentifier objectHash = BlobIdentifier.FromBlob(data);
 
             using HttpContent requestContent = new ByteArrayContent(data);
@@ -1292,31 +1292,31 @@ namespace Horde.Storage.FunctionalTests.References
             IoHashKey key = IoHashKey.FromName("deletableObject");
             // submit the object
             {
-                HttpResponseMessage result = await _httpClient!.PutAsync(new Uri($"api/v1/refs/{TestNamespace}/bucket/{key}", UriKind.Relative), requestContent);
+                HttpResponseMessage result = await Client!.PutAsync(new Uri($"api/v1/refs/{TestNamespace}/bucket/{key}", UriKind.Relative), requestContent);
                 result.EnsureSuccessStatusCode();
             }
 
             // verify it is present
             {
-                HttpResponseMessage result = await _httpClient!.GetAsync(new Uri($"api/v1/refs/{TestNamespace}/bucket/{key}.raw", UriKind.Relative));
+                HttpResponseMessage result = await Client!.GetAsync(new Uri($"api/v1/refs/{TestNamespace}/bucket/{key}.raw", UriKind.Relative));
                 result.EnsureSuccessStatusCode();
             }
 
             // delete the object
             {
-                HttpResponseMessage result = await _httpClient!.DeleteAsync(new Uri($"api/v1/refs/{TestNamespace}/bucket/{key}", UriKind.Relative));
+                HttpResponseMessage result = await Client!.DeleteAsync(new Uri($"api/v1/refs/{TestNamespace}/bucket/{key}", UriKind.Relative));
                 result.EnsureSuccessStatusCode();
             }
 
             // ensure the object is not present anymore
             {
-                HttpResponseMessage result = await _httpClient!.GetAsync(new Uri($"api/v1/refs/{TestNamespace}/bucket/{key}.raw", UriKind.Relative));
+                HttpResponseMessage result = await Client!.GetAsync(new Uri($"api/v1/refs/{TestNamespace}/bucket/{key}.raw", UriKind.Relative));
                 Assert.AreEqual(HttpStatusCode.NotFound, result.StatusCode);
             }
 
             // delete the object again which doesn't exist anymore
             {
-                HttpResponseMessage result = await _httpClient!.DeleteAsync(new Uri($"api/v1/refs/{TestNamespace}/bucket/{key}", UriKind.Relative));
+                HttpResponseMessage result = await Client!.DeleteAsync(new Uri($"api/v1/refs/{TestNamespace}/bucket/{key}", UriKind.Relative));
                 Assert.AreEqual(HttpStatusCode.NotFound, result.StatusCode);
             }
         }
@@ -1327,8 +1327,8 @@ namespace Horde.Storage.FunctionalTests.References
         {
             const string BucketToDelete = "delete-bucket";
 
-            const string objectContents = "This is treated as a opaque blob";
-            byte[] data = Encoding.ASCII.GetBytes(objectContents);
+            const string ObjectContents = "This is treated as a opaque blob";
+            byte[] data = Encoding.ASCII.GetBytes(ObjectContents);
             BlobIdentifier objectHash = BlobIdentifier.FromBlob(data);
 
             using HttpContent requestContent = new ByteArrayContent(data);
@@ -1339,42 +1339,42 @@ namespace Horde.Storage.FunctionalTests.References
 
             // submit the object into multiple buckets
             {
-                HttpResponseMessage result = await _httpClient!.PutAsync(new Uri($"api/v1/refs/{TestNamespace}/{BucketToDelete}/{key}", UriKind.Relative), requestContent);
+                HttpResponseMessage result = await Client!.PutAsync(new Uri($"api/v1/refs/{TestNamespace}/{BucketToDelete}/{key}", UriKind.Relative), requestContent);
                 result.EnsureSuccessStatusCode();
             }
 
             {
-                HttpResponseMessage result = await _httpClient!.PutAsync(new Uri($"api/v1/refs/{TestNamespace}/bucket/{key}", UriKind.Relative), requestContent);
+                HttpResponseMessage result = await Client!.PutAsync(new Uri($"api/v1/refs/{TestNamespace}/bucket/{key}", UriKind.Relative), requestContent);
                 result.EnsureSuccessStatusCode();
             }
 
             // verify it is present
             {
-                HttpResponseMessage result = await _httpClient!.GetAsync(new Uri($"api/v1/refs/{TestNamespace}/bucket/{key}.raw", UriKind.Relative));
+                HttpResponseMessage result = await Client!.GetAsync(new Uri($"api/v1/refs/{TestNamespace}/bucket/{key}.raw", UriKind.Relative));
                 result.EnsureSuccessStatusCode();
             }
 
             // verify it is present
             {
-                HttpResponseMessage result = await _httpClient!.GetAsync(new Uri($"api/v1/refs/{TestNamespace}/{BucketToDelete}/{key}.raw", UriKind.Relative));
+                HttpResponseMessage result = await Client!.GetAsync(new Uri($"api/v1/refs/{TestNamespace}/{BucketToDelete}/{key}.raw", UriKind.Relative));
                 result.EnsureSuccessStatusCode();
             }
 
             // delete the bucket
             {
-                HttpResponseMessage result = await _httpClient!.DeleteAsync(new Uri($"api/v1/refs/{TestNamespace}/{BucketToDelete}", UriKind.Relative));
+                HttpResponseMessage result = await Client!.DeleteAsync(new Uri($"api/v1/refs/{TestNamespace}/{BucketToDelete}", UriKind.Relative));
                 result.EnsureSuccessStatusCode();
             }
 
             // ensure the object is present in not deleted bucket
             {
-                HttpResponseMessage result = await _httpClient!.GetAsync(new Uri($"api/v1/refs/{TestNamespace}/bucket/{key}.raw", UriKind.Relative));
+                HttpResponseMessage result = await Client!.GetAsync(new Uri($"api/v1/refs/{TestNamespace}/bucket/{key}.raw", UriKind.Relative));
                 result.EnsureSuccessStatusCode();
             }
 
             // ensure the object is not present anymore
             {
-                HttpResponseMessage result = await _httpClient!.GetAsync(new Uri($"api/v1/refs/{TestNamespace}/{BucketToDelete}/{key}.raw", UriKind.Relative));
+                HttpResponseMessage result = await Client!.GetAsync(new Uri($"api/v1/refs/{TestNamespace}/{BucketToDelete}/{key}.raw", UriKind.Relative));
                 Assert.AreEqual(HttpStatusCode.NotFound, result.StatusCode);
             }
         }
@@ -1385,8 +1385,8 @@ namespace Horde.Storage.FunctionalTests.References
         {
             const string NamespaceToBeDeleted = "test-delete-namespace";
 
-            const string objectContents = "This is treated as a opaque blob";
-            byte[] data = Encoding.ASCII.GetBytes(objectContents);
+            const string ObjectContents = "This is treated as a opaque blob";
+            byte[] data = Encoding.ASCII.GetBytes(ObjectContents);
             BlobIdentifier objectHash = BlobIdentifier.FromBlob(data);
 
             using HttpContent requestContent = new ByteArrayContent(data);
@@ -1396,48 +1396,48 @@ namespace Horde.Storage.FunctionalTests.References
             IoHashKey key = IoHashKey.FromName("deletableObject");
             // submit the object into multiple namespaces
             {
-                HttpResponseMessage result = await _httpClient!.PutAsync(new Uri($"api/v1/refs/{TestNamespace}/bucket/{key}", UriKind.Relative), requestContent);
+                HttpResponseMessage result = await Client!.PutAsync(new Uri($"api/v1/refs/{TestNamespace}/bucket/{key}", UriKind.Relative), requestContent);
                 result.EnsureSuccessStatusCode();
             }
 
             {
-                HttpResponseMessage result = await _httpClient!.PutAsync(new Uri($"api/v1/refs/{NamespaceToBeDeleted}/bucket/{key}", UriKind.Relative), requestContent);
+                HttpResponseMessage result = await Client!.PutAsync(new Uri($"api/v1/refs/{NamespaceToBeDeleted}/bucket/{key}", UriKind.Relative), requestContent);
                 result.EnsureSuccessStatusCode();
             }
 
             // verify it is present
             {
-                HttpResponseMessage result = await _httpClient!.GetAsync(new Uri($"api/v1/refs/{TestNamespace}/bucket/{key}.raw", UriKind.Relative));
+                HttpResponseMessage result = await Client!.GetAsync(new Uri($"api/v1/refs/{TestNamespace}/bucket/{key}.raw", UriKind.Relative));
                 result.EnsureSuccessStatusCode();
             }
 
             // verify it is present
             {
-                HttpResponseMessage result = await _httpClient!.GetAsync(new Uri($"api/v1/refs/{NamespaceToBeDeleted}/bucket/{key}.raw", UriKind.Relative));
+                HttpResponseMessage result = await Client!.GetAsync(new Uri($"api/v1/refs/{NamespaceToBeDeleted}/bucket/{key}.raw", UriKind.Relative));
                 result.EnsureSuccessStatusCode();
             }
 
             // delete the namespace
             {
-                HttpResponseMessage result = await _httpClient!.DeleteAsync(new Uri($"api/v1/refs/{NamespaceToBeDeleted}", UriKind.Relative));
+                HttpResponseMessage result = await Client!.DeleteAsync(new Uri($"api/v1/refs/{NamespaceToBeDeleted}", UriKind.Relative));
                 result.EnsureSuccessStatusCode();
             }
 
             // ensure the object is present in not deleted namespace
             {
-                HttpResponseMessage result = await _httpClient!.GetAsync(new Uri($"api/v1/refs/{TestNamespace}/bucket/{key}.raw", UriKind.Relative));
+                HttpResponseMessage result = await Client!.GetAsync(new Uri($"api/v1/refs/{TestNamespace}/bucket/{key}.raw", UriKind.Relative));
                 result.EnsureSuccessStatusCode();
             }
 
             // ensure the object is not present anymore
             {
-                HttpResponseMessage result = await _httpClient!.GetAsync(new Uri($"api/v1/refs/{NamespaceToBeDeleted}/bucket/{key}.raw", UriKind.Relative));
+                HttpResponseMessage result = await Client!.GetAsync(new Uri($"api/v1/refs/{NamespaceToBeDeleted}/bucket/{key}.raw", UriKind.Relative));
                 Assert.AreEqual(HttpStatusCode.NotFound, result.StatusCode);
             }
 
             // make sure the namespace is not considered a valid namespace anymore
             {
-                HttpResponseMessage result = await _httpClient!.GetAsync(new Uri($"api/v1/refs", UriKind.Relative));
+                HttpResponseMessage result = await Client!.GetAsync(new Uri($"api/v1/refs", UriKind.Relative));
                 result.EnsureSuccessStatusCode();
                 Assert.AreEqual(result!.Content.Headers.ContentType!.MediaType, MediaTypeNames.Application.Json);
                 GetNamespacesResponse response = await result.Content.ReadAsAsync<GetNamespacesResponse>();
@@ -1449,8 +1449,8 @@ namespace Horde.Storage.FunctionalTests.References
         [TestMethod]
         public async Task ListNamespaces()
         {
-            const string objectContents = "This is treated as a opaque blob";
-            byte[] data = Encoding.ASCII.GetBytes(objectContents);
+            const string ObjectContents = "This is treated as a opaque blob";
+            byte[] data = Encoding.ASCII.GetBytes(ObjectContents);
             BlobIdentifier objectHash = BlobIdentifier.FromBlob(data);
             IoHashKey key = IoHashKey.FromName("notUsedObject");
 
@@ -1460,12 +1460,12 @@ namespace Horde.Storage.FunctionalTests.References
 
             // submit a object to make sure a namespace is created
             {
-                HttpResponseMessage result = await _httpClient!.PutAsync(new Uri($"api/v1/refs/{TestNamespace}/bucket/{key}", UriKind.Relative), requestContent);
+                HttpResponseMessage result = await Client!.PutAsync(new Uri($"api/v1/refs/{TestNamespace}/bucket/{key}", UriKind.Relative), requestContent);
                 result.EnsureSuccessStatusCode();
             }
 
             {
-                HttpResponseMessage result = await _httpClient!.GetAsync(new Uri($"api/v1/refs", UriKind.Relative));
+                HttpResponseMessage result = await Client!.GetAsync(new Uri($"api/v1/refs", UriKind.Relative));
                 result.EnsureSuccessStatusCode();
                 Assert.AreEqual(result!.Content.Headers.ContentType!.MediaType, MediaTypeNames.Application.Json);
                 GetNamespacesResponse response = await result.Content.ReadAsAsync<GetNamespacesResponse>();
@@ -1476,8 +1476,8 @@ namespace Horde.Storage.FunctionalTests.References
         [TestMethod]
         public async Task GetOldRecords()
         {
-            const string objectContents = "This is treated as a opaque blob";
-            byte[] data = Encoding.ASCII.GetBytes(objectContents);
+            const string ObjectContents = "This is treated as a opaque blob";
+            byte[] data = Encoding.ASCII.GetBytes(ObjectContents);
             BlobIdentifier objectHash = BlobIdentifier.FromBlob(data);
 
             using HttpContent requestContent = new ByteArrayContent(data);
@@ -1487,7 +1487,7 @@ namespace Horde.Storage.FunctionalTests.References
             IoHashKey key = IoHashKey.FromName("oldRecord");
             // submit some contents
             {
-                HttpResponseMessage result = await _httpClient!.PutAsync(new Uri($"api/v1/refs/{TestNamespace}/bucket/{key}", UriKind.Relative), requestContent);
+                HttpResponseMessage result = await Client!.PutAsync(new Uri($"api/v1/refs/{TestNamespace}/bucket/{key}", UriKind.Relative), requestContent);
                 result.EnsureSuccessStatusCode();
             }
 
@@ -1507,7 +1507,7 @@ namespace Horde.Storage.FunctionalTests.References
                 using HttpContent requestContent = new StringContent("{}");
                 requestContent.Headers.ContentType = new MediaTypeHeaderValue(MediaTypeNames.Application.Json);
 
-                HttpResponseMessage result = await _httpClient!.PostAsync(new Uri($"api/v1/refs/{TestNamespace}", UriKind.Relative), requestContent);
+                HttpResponseMessage result = await Client!.PostAsync(new Uri($"api/v1/refs/{TestNamespace}", UriKind.Relative), requestContent);
 
                 Assert.AreEqual(HttpStatusCode.UnsupportedMediaType, result.StatusCode);
             }
@@ -1562,7 +1562,7 @@ namespace Horde.Storage.FunctionalTests.References
                 using HttpContent requestContent = new ByteArrayContent(batchRequestData);
                 requestContent.Headers.ContentType = new MediaTypeHeaderValue(CustomMediaTypeNames.UnrealCompactBinary);
 
-                HttpResponseMessage result = await _httpClient!.PostAsync(new Uri($"api/v1/refs/{TestNamespace}", UriKind.Relative), requestContent);
+                HttpResponseMessage result = await Client!.PostAsync(new Uri($"api/v1/refs/{TestNamespace}", UriKind.Relative), requestContent);
                 result.EnsureSuccessStatusCode();
 
                 await using MemoryStream ms = new MemoryStream();
@@ -1602,7 +1602,7 @@ namespace Horde.Storage.FunctionalTests.References
                 requestContent.Headers.ContentType = new MediaTypeHeaderValue(CustomMediaTypeNames.UnrealCompactBinary);
                 requestContent.Headers.Add(CommonHeaders.HashHeaderName, blobHash.ToString());
 
-                HttpResponseMessage result = await _httpClient!.PutAsync(new Uri($"api/v1/refs/{TestNamespace}/{bucket}/{newBlobObjectKey}.uecb", UriKind.Relative), requestContent);
+                HttpResponseMessage result = await Client!.PutAsync(new Uri($"api/v1/refs/{TestNamespace}/{bucket}/{newBlobObjectKey}.uecb", UriKind.Relative), requestContent);
                 result.EnsureSuccessStatusCode();
                 byte[] content = await result.Content.ReadAsByteArrayAsync();
                 PutObjectResponse? response = CbSerializer.Deserialize<PutObjectResponse?>(content);
@@ -1621,7 +1621,7 @@ namespace Horde.Storage.FunctionalTests.References
                 requestContent.Headers.ContentType = new MediaTypeHeaderValue(MediaTypeNames.Application.Octet);
                 requestContent.Headers.Add(CommonHeaders.HashHeaderName, blobHash.ToString());
 
-                HttpResponseMessage result = await _httpClient!.PutAsync(new Uri($"api/v1/blobs/{TestNamespace}/{blobHash}", UriKind.Relative), requestContent);
+                HttpResponseMessage result = await Client!.PutAsync(new Uri($"api/v1/blobs/{TestNamespace}/{blobHash}", UriKind.Relative), requestContent);
                 result.EnsureSuccessStatusCode();
             }
 
@@ -1636,7 +1636,7 @@ namespace Horde.Storage.FunctionalTests.References
                 requestContent.Headers.ContentType = new MediaTypeHeaderValue(MediaTypeNames.Application.Octet);
                 requestContent.Headers.Add(CommonHeaders.HashHeaderName, blobHash.ToString());
 
-                HttpResponseMessage result = await _httpClient!.PutAsync(new Uri($"api/v1/blobs/{TestNamespace}/{blobHash}", UriKind.Relative), requestContent);
+                HttpResponseMessage result = await Client!.PutAsync(new Uri($"api/v1/blobs/{TestNamespace}/{blobHash}", UriKind.Relative), requestContent);
                 result.EnsureSuccessStatusCode();
             }
 
@@ -1648,7 +1648,7 @@ namespace Horde.Storage.FunctionalTests.References
                 requestContent.Headers.ContentType = new MediaTypeHeaderValue(CustomMediaTypeNames.UnrealCompactBinary);
                 requestContent.Headers.Add(CommonHeaders.HashHeaderName, blobHash.ToString());
 
-                HttpResponseMessage result = await _httpClient!.PutAsync(new Uri($"api/v1/refs/{TestNamespace}/{bucket}/{newReferenceObjectKey}.uecb", UriKind.Relative), requestContent);
+                HttpResponseMessage result = await Client!.PutAsync(new Uri($"api/v1/refs/{TestNamespace}/{bucket}/{newReferenceObjectKey}.uecb", UriKind.Relative), requestContent);
                 result.EnsureSuccessStatusCode();
 
                 result.EnsureSuccessStatusCode();
@@ -1704,7 +1704,7 @@ namespace Horde.Storage.FunctionalTests.References
                 using HttpContent requestContent = new ByteArrayContent(batchRequestData);
                 requestContent.Headers.ContentType = new MediaTypeHeaderValue(CustomMediaTypeNames.UnrealCompactBinary);
 
-                HttpResponseMessage result = await _httpClient!.PostAsync(new Uri($"api/v1/refs/{TestNamespace}", UriKind.Relative), requestContent);
+                HttpResponseMessage result = await Client!.PostAsync(new Uri($"api/v1/refs/{TestNamespace}", UriKind.Relative), requestContent);
                 result.EnsureSuccessStatusCode();
 
                 await using MemoryStream ms = new MemoryStream();
@@ -1747,7 +1747,7 @@ namespace Horde.Storage.FunctionalTests.References
                 requestContent.Headers.ContentType = new MediaTypeHeaderValue(CustomMediaTypeNames.UnrealCompactBinary);
                 requestContent.Headers.Add(CommonHeaders.HashHeaderName, blobHash.ToString());
 
-                HttpResponseMessage result = await _httpClient!.PutAsync(new Uri($"api/v1/refs/{TestNamespace}/{bucket}/{newBlobObjectKey}.uecb", UriKind.Relative), requestContent);
+                HttpResponseMessage result = await Client!.PutAsync(new Uri($"api/v1/refs/{TestNamespace}/{bucket}/{newBlobObjectKey}.uecb", UriKind.Relative), requestContent);
                 result.EnsureSuccessStatusCode();
                 byte[] content = await result.Content.ReadAsByteArrayAsync();
                 PutObjectResponse? response = CbSerializer.Deserialize<PutObjectResponse?>(content);
@@ -1766,7 +1766,7 @@ namespace Horde.Storage.FunctionalTests.References
                 requestContent.Headers.ContentType = new MediaTypeHeaderValue(MediaTypeNames.Application.Octet);
                 requestContent.Headers.Add(CommonHeaders.HashHeaderName, blobHash.ToString());
 
-                HttpResponseMessage result = await _httpClient!.PutAsync(new Uri($"api/v1/blobs/{TestNamespace}/{blobHash}", UriKind.Relative), requestContent);
+                HttpResponseMessage result = await Client!.PutAsync(new Uri($"api/v1/blobs/{TestNamespace}/{blobHash}", UriKind.Relative), requestContent);
                 result.EnsureSuccessStatusCode();
             }
 
@@ -1778,7 +1778,7 @@ namespace Horde.Storage.FunctionalTests.References
                 requestContent.Headers.ContentType = new MediaTypeHeaderValue(CustomMediaTypeNames.UnrealCompactBinary);
                 requestContent.Headers.Add(CommonHeaders.HashHeaderName, blobHash.ToString());
 
-                HttpResponseMessage result = await _httpClient!.PutAsync(new Uri($"api/v1/refs/{TestNamespace}/{bucket}/{newReferenceObjectKey}.uecb", UriKind.Relative), requestContent);
+                HttpResponseMessage result = await Client!.PutAsync(new Uri($"api/v1/refs/{TestNamespace}/{bucket}/{newReferenceObjectKey}.uecb", UriKind.Relative), requestContent);
                 result.EnsureSuccessStatusCode();
 
                 result.EnsureSuccessStatusCode();
@@ -1823,7 +1823,7 @@ namespace Horde.Storage.FunctionalTests.References
                 using HttpContent requestContent = new ByteArrayContent(batchRequestData);
                 requestContent.Headers.ContentType = new MediaTypeHeaderValue(CustomMediaTypeNames.UnrealCompactBinary);
 
-                HttpResponseMessage result = await _httpClient!.PostAsync(new Uri($"api/v1/refs/{TestNamespace}", UriKind.Relative), requestContent);
+                HttpResponseMessage result = await Client!.PostAsync(new Uri($"api/v1/refs/{TestNamespace}", UriKind.Relative), requestContent);
                 result.EnsureSuccessStatusCode();
 
                 await using MemoryStream ms = new MemoryStream();
@@ -1898,7 +1898,7 @@ namespace Horde.Storage.FunctionalTests.References
                 using HttpContent requestContent = new ByteArrayContent(batchRequestData);
                 requestContent.Headers.ContentType = new MediaTypeHeaderValue(CustomMediaTypeNames.UnrealCompactBinary);
 
-                HttpResponseMessage result = await _httpClient!.PostAsync(new Uri($"api/v1/refs/{TestNamespace}", UriKind.Relative), requestContent);
+                HttpResponseMessage result = await Client!.PostAsync(new Uri($"api/v1/refs/{TestNamespace}", UriKind.Relative), requestContent);
                 result.EnsureSuccessStatusCode();
 
                 await using MemoryStream ms = new MemoryStream();
@@ -1943,7 +1943,7 @@ namespace Horde.Storage.FunctionalTests.References
                 requestContent.Headers.ContentType = new MediaTypeHeaderValue(CustomMediaTypeNames.UnrealCompactBinary);
                 requestContent.Headers.Add(CommonHeaders.HashHeaderName, blobHash.ToString());
 
-                HttpResponseMessage result = await _httpClient!.PutAsync(new Uri($"api/v1/refs/{TestNamespace}/{bucket}/{getObjectKey}.uecb", UriKind.Relative), requestContent);
+                HttpResponseMessage result = await Client!.PutAsync(new Uri($"api/v1/refs/{TestNamespace}/{bucket}/{getObjectKey}.uecb", UriKind.Relative), requestContent);
                 result.EnsureSuccessStatusCode();
                 byte[] content = await result.Content.ReadAsByteArrayAsync();
                 PutObjectResponse? response = CbSerializer.Deserialize<PutObjectResponse?>(content);
@@ -2008,7 +2008,7 @@ namespace Horde.Storage.FunctionalTests.References
                 using HttpContent requestContent = new ByteArrayContent(batchRequestData);
                 requestContent.Headers.ContentType = new MediaTypeHeaderValue(CustomMediaTypeNames.UnrealCompactBinary);
 
-                HttpResponseMessage result = await _httpClient!.PostAsync(new Uri($"api/v1/refs/{TestNamespace}", UriKind.Relative), requestContent);
+                HttpResponseMessage result = await Client!.PostAsync(new Uri($"api/v1/refs/{TestNamespace}", UriKind.Relative), requestContent);
                 result.EnsureSuccessStatusCode();
 
                 await using MemoryStream ms = new MemoryStream();
@@ -2043,5 +2043,221 @@ namespace Horde.Storage.FunctionalTests.References
                 Assert.IsTrue(headOp.Response["exists"].AsBool());
             }
         }
+
+        [TestMethod]
+        public async Task GetPackage()
+        {
+            string blobContents = "This is a string that is referenced as a blob";
+            byte[] blobData = Encoding.ASCII.GetBytes(blobContents);
+            BlobIdentifier blobHash = BlobIdentifier.FromBlob(blobData);
+            await Service.PutObject(TestNamespace, blobData, blobHash);
+
+            string blobContentsChild = "This string is also referenced as a blob but from a child object";
+            byte[] dataChild = Encoding.ASCII.GetBytes(blobContentsChild);
+            BlobIdentifier blobHashChild = BlobIdentifier.FromBlob(dataChild);
+            await Service.PutObject(TestNamespace, dataChild, blobHashChild);
+
+            CbWriter writerChild = new CbWriter();
+            writerChild.BeginObject();
+            writerChild.WriteBinaryAttachment("blob", blobHashChild.AsIoHash());
+            writerChild.EndObject();
+
+            byte[] childDataObject = writerChild.ToByteArray();
+            BlobIdentifier childDataObjectHash = BlobIdentifier.FromBlob(childDataObject);
+            await Service.PutObject(TestNamespace, childDataObject, childDataObjectHash);
+
+            CbWriter writerParent = new CbWriter();
+            writerParent.BeginObject();
+
+            writerParent.WriteBinaryAttachment("blobAttachment", blobHash.AsIoHash());
+            writerParent.WriteObjectAttachment("objectAttachment", childDataObjectHash.AsIoHash());
+            writerParent.EndObject();
+
+            byte[] objectData = writerParent.ToByteArray();
+            BlobIdentifier objectHash = BlobIdentifier.FromBlob(objectData);
+
+            IoHashKey key = IoHashKey.FromName("newHierarchyObject");
+
+            using HttpContent requestContent = new ByteArrayContent(objectData);
+            requestContent.Headers.ContentType = new MediaTypeHeaderValue(CustomMediaTypeNames.UnrealCompactBinary);
+            requestContent.Headers.Add(CommonHeaders.HashHeaderName, objectHash.ToString());
+
+            {
+                HttpResponseMessage result = await Client!.PutAsync(new Uri($"api/v1/refs/{TestNamespace}/bucket/{key}.uecb", UriKind.Relative), requestContent);
+                result.EnsureSuccessStatusCode();
+            }
+
+            {
+                HttpResponseMessage result = await Client!.GetAsync(new Uri($"api/v1/refs/{TestNamespace}/bucket/{key}.uecbpkg", UriKind.Relative));
+                result.EnsureSuccessStatusCode();
+
+                CbPackageReader packageReader = await CbPackageReader.Create(await result.Content.ReadAsStreamAsync());
+                List<(CbPackageAttachmentEntry, byte[])>? attachments = await packageReader.IterateAttachments().ToListAsync();
+
+                Assert.AreEqual(3, attachments.Count);
+
+                foreach ((CbPackageAttachmentEntry entry, byte[] blob) in attachments)
+                {
+                    if (entry.AttachmentHash == blobHash.AsIoHash())
+                    {
+                        CollectionAssert.AreEqual(blobData, blob);
+                    }
+                    else if (entry.AttachmentHash == blobHashChild.AsIoHash())
+                    {
+                        CollectionAssert.AreEqual(dataChild,  blob);
+                    }
+                    else if (entry.AttachmentHash == childDataObjectHash.AsIoHash())
+                    {
+                        CollectionAssert.AreEqual(childDataObject, blob);
+                    }
+                    else
+                    {
+                        Assert.Fail($"Unknown attachment {entry.AttachmentHash}");
+                    }
+                }
+            }
+        }
+
+        [TestMethod]
+        public async Task PutPackage()
+        {
+            string blobContents = "This is a string that is referenced as a blob";
+            byte[] blobData = Encoding.ASCII.GetBytes(blobContents);
+            BlobIdentifier blobHash = BlobIdentifier.FromBlob(blobData);
+            await Service.PutObject(TestNamespace, blobData, blobHash);
+
+            string blobContentsChild = "This string is also referenced as a blob but from a child object";
+            byte[] dataChild = Encoding.ASCII.GetBytes(blobContentsChild);
+            BlobIdentifier blobHashChild = BlobIdentifier.FromBlob(dataChild);
+            await Service.PutObject(TestNamespace, dataChild, blobHashChild);
+
+            CbWriter writerChild = new CbWriter();
+            writerChild.BeginObject();
+            writerChild.WriteBinaryAttachment("blob", blobHashChild.AsIoHash());
+            writerChild.EndObject();
+
+            byte[] childDataObject = writerChild.ToByteArray();
+            BlobIdentifier childDataObjectHash = BlobIdentifier.FromBlob(childDataObject);
+            await Service.PutObject(TestNamespace, childDataObject, childDataObjectHash);
+
+            CbWriter writerParent = new CbWriter();
+            writerParent.BeginObject();
+
+            writerParent.WriteBinaryAttachment("blobAttachment", blobHash.AsIoHash());
+            writerParent.WriteObjectAttachment("objectAttachment", childDataObjectHash.AsIoHash());
+            writerParent.EndObject();
+
+            byte[] objectData = writerParent.ToByteArray();
+            BlobIdentifier objectHash = BlobIdentifier.FromBlob(objectData);
+
+            IoHashKey key = IoHashKey.FromName("newHierarchyObject");
+
+            CbPackageBuilder builder = new CbPackageBuilder();
+            await builder.AddAttachment(objectHash.AsIoHash(), CbPackageAttachmentFlags.IsObject, objectData);
+            await builder.AddAttachment(blobHash.AsIoHash(), 0, blobData);
+            await builder.AddAttachment(blobHashChild.AsIoHash(), 0, dataChild);
+            await builder.AddAttachment(childDataObjectHash.AsIoHash(), 0, childDataObject);
+            using HttpContent requestContent = new ByteArrayContent(builder.ToByteArray());
+            requestContent.Headers.ContentType = new MediaTypeHeaderValue(CustomMediaTypeNames.UnrealCompactBinaryPackage);
+
+            {
+                HttpResponseMessage result = await Client!.PutAsync(new Uri($"api/v1/refs/{TestNamespace}/bucket/{key}", UriKind.Relative), requestContent);
+                result.EnsureSuccessStatusCode();
+            }
+            
+            {
+                BlobIdentifier blobAttachment;
+                BlobIdentifier objectAttachment;
+                {
+                    HttpResponseMessage getResponse = await Client.GetAsync(new Uri($"api/v1/refs/{TestNamespace}/bucket/{key}.uecb", UriKind.Relative));
+                    getResponse.EnsureSuccessStatusCode();
+                    await using MemoryStream ms = new MemoryStream();
+                    await getResponse.Content.CopyToAsync(ms);
+
+                    byte[] roundTrippedBuffer = ms.ToArray();
+                    ReadOnlyMemory<byte> localMemory = new ReadOnlyMemory<byte>(roundTrippedBuffer);
+                    CbObject cb = new CbObject(roundTrippedBuffer);
+                    Assert.AreEqual(2, cb.Count());
+
+                    CbField blobAttachmentField = cb["blobAttachment"];
+                    Assert.AreNotEqual(CbField.Empty, blobAttachmentField);
+                    blobAttachment = BlobIdentifier.FromIoHash(blobAttachmentField.AsBinaryAttachment());
+                    CbField objectAttachmentField = cb["objectAttachment"];
+                    Assert.AreNotEqual(CbField.Empty, objectAttachmentField);
+                    objectAttachment = BlobIdentifier.FromIoHash(objectAttachmentField.AsObjectAttachment().Hash);
+                }
+
+                {
+                    HttpResponseMessage getAttachment = await Client.GetAsync(new Uri($"api/v1/blobs/{TestNamespace}/{blobAttachment}", UriKind.Relative));
+                    getAttachment.EnsureSuccessStatusCode();
+                    await using MemoryStream ms = new MemoryStream();
+                    await getAttachment.Content.CopyToAsync(ms);
+                    byte[] roundTrippedBuffer = ms.ToArray();
+                    string roundTrippedString = Encoding.ASCII.GetString(roundTrippedBuffer);
+
+                    Assert.AreEqual(blobContents, roundTrippedString);
+                    Assert.AreEqual(blobHash, BlobIdentifier.FromBlob(roundTrippedBuffer));
+                }
+
+                BlobIdentifier attachedBlobIdentifier;
+                {
+                    HttpResponseMessage getAttachment = await Client.GetAsync(new Uri($"api/v1/blobs/{TestNamespace}/{objectAttachment}", UriKind.Relative));
+                    getAttachment.EnsureSuccessStatusCode();
+                    await using MemoryStream ms = new MemoryStream();
+                    await getAttachment.Content.CopyToAsync(ms);
+                    byte[] roundTrippedBuffer = ms.ToArray();
+                    ReadOnlyMemory<byte> localMemory = new ReadOnlyMemory<byte>(roundTrippedBuffer);
+                    CbObject cb = new CbObject(roundTrippedBuffer);
+                    Assert.AreEqual(1, cb.Count());
+
+                    CbField blobField = cb["blob"];
+                    Assert.AreNotEqual(CbField.Empty, blobField);
+
+                    attachedBlobIdentifier = BlobIdentifier.FromIoHash(blobField!.AsBinaryAttachment());
+                }
+
+                {
+                    HttpResponseMessage getAttachment = await Client.GetAsync(new Uri($"api/v1/blobs/{TestNamespace}/{attachedBlobIdentifier}", UriKind.Relative));
+                    getAttachment.EnsureSuccessStatusCode();
+                    await using MemoryStream ms = new MemoryStream();
+                    await getAttachment.Content.CopyToAsync(ms);
+                    byte[] roundTrippedBuffer = ms.ToArray();
+                    string roundTrippedString = Encoding.ASCII.GetString(roundTrippedBuffer);
+
+                    Assert.AreEqual(blobContentsChild, roundTrippedString);
+                    Assert.AreEqual(blobHashChild, BlobIdentifier.FromBlob(roundTrippedBuffer));
+                }
+            }
+        }
+
+        [TestMethod]
+        public async Task PutPackageWithError()
+        {
+            string blobContents = "This is a string that is referenced as a blob";
+            byte[] blobData = Encoding.ASCII.GetBytes(blobContents);
+            BlobIdentifier blobHash = BlobIdentifier.FromBlob(blobData);
+            await Service.PutObject(TestNamespace, blobData, blobHash);
+
+            CbWriter writerParent = new CbWriter();
+            writerParent.BeginObject();
+            writerParent.WriteBinaryAttachment("blobAttachment", blobHash.AsIoHash());
+            writerParent.EndObject();
+
+            byte[] objectData = writerParent.ToByteArray();
+            BlobIdentifier objectHash = BlobIdentifier.FromBlob(objectData);
+
+            IoHashKey key = IoHashKey.FromName("brokenObject");
+
+            CbPackageBuilder builder = new CbPackageBuilder();
+            await builder.AddAttachment(objectHash.AsIoHash(), CbPackageAttachmentFlags.IsObject, objectData);
+            await builder.AddAttachment(blobHash.AsIoHash(), CbPackageAttachmentFlags.IsError, blobData);
+            using HttpContent requestContent = new ByteArrayContent(builder.ToByteArray());
+            requestContent.Headers.ContentType = new MediaTypeHeaderValue(CustomMediaTypeNames.UnrealCompactBinaryPackage);
+
+            {
+                HttpResponseMessage result = await Client!.PutAsync(new Uri($"api/v1/refs/{TestNamespace}/bucket/{key}", UriKind.Relative), requestContent);
+                Assert.AreEqual(HttpStatusCode.BadRequest, result.StatusCode);
+            }
+        }
     }
 }
diff --git a/Engine/Source/Programs/Horde/HordeStorage/Jupiter.Common/BaseStartup.cs b/Engine/Source/Programs/Horde/HordeStorage/Jupiter.Common/BaseStartup.cs
--- a/Engine/Source/Programs/Horde/HordeStorage/Jupiter.Common/BaseStartup.cs
+++ b/Engine/Source/Programs/Horde/HordeStorage/Jupiter.Common/BaseStartup.cs
@@ -94,6 +94,7 @@ namespace Jupiter
 
                     options.FormatterMappings.SetMediaTypeMappingForFormat("raw", MediaTypeNames.Application.Octet);
                     options.FormatterMappings.SetMediaTypeMappingForFormat("uecb", CustomMediaTypeNames.UnrealCompactBinary);
+                    options.FormatterMappings.SetMediaTypeMappingForFormat("uecbpkg", CustomMediaTypeNames.UnrealCompactBinaryPackage);
 
                     OnAddControllers(options);
                 }).ConfigureApiBehaviorOptions(options =>
diff --git a/Engine/Source/Programs/Horde/HordeStorage/Jupiter.Common/CustomMediaTypeNames.cs b/Engine/Source/Programs/Horde/HordeStorage/Jupiter.Common/CustomMediaTypeNames.cs
--- a/Engine/Source/Programs/Horde/HordeStorage/Jupiter.Common/CustomMediaTypeNames.cs
+++ b/Engine/Source/Programs/Horde/HordeStorage/Jupiter.Common/CustomMediaTypeNames.cs
@@ -8,5 +8,7 @@ namespace Jupiter
         public const string UnrealCompressedBuffer = "application/x-ue-comp";
 
         public const string JupiterInlinedPayload = "application/x-jupiter-inline";
+        public const string UnrealCompactBinaryPackage = "application/x-ue-cbpkg";
+
     }
 }
