Horde Storage - Fixed issue with put package endpoint (causing tests to fail). Removed extra auth attribute which we no longer use.

diff --git a/Engine/Source/Programs/Horde/HordeStorage/Horde.Storage/Controllers/ReferencesController.cs b/Engine/Source/Programs/Horde/HordeStorage/Horde.Storage/Controllers/ReferencesController.cs
--- a/Engine/Source/Programs/Horde/HordeStorage/Horde.Storage/Controllers/ReferencesController.cs
+++ b/Engine/Source/Programs/Horde/HordeStorage/Horde.Storage/Controllers/ReferencesController.cs
@@ -679,7 +679,6 @@ namespace Horde.Storage.Controllers
         [HttpPut("{ns}/{bucket}/{key}", Order = 300)]
         [DisableRequestSizeLimit]
         [RequiredContentType(CustomMediaTypeNames.UnrealCompactBinaryPackage)]
-        [Authorize("Object.write")]
         public async Task<IActionResult> PutPackage(
             [FromRoute][Required] NamespaceId ns,
             [FromRoute][Required] BucketId bucket,
