Horde Storage - Improvements to ref cleanup. Made it run ref cleanup for namespaces in parallel. Skipping cleanup of legacy namespaces as that can take a very long time.
Also fixed issue with the logging of number of gc records considered as it can grow quite large.

diff --git a/Engine/Source/Programs/Horde/HordeStorage/Horde.Storage/Implementation/GC/CacheCleanup.cs b/Engine/Source/Programs/Horde/HordeStorage/Horde.Storage/Implementation/GC/CacheCleanup.cs
--- a/Engine/Source/Programs/Horde/HordeStorage/Horde.Storage/Implementation/GC/CacheCleanup.cs
+++ b/Engine/Source/Programs/Horde/HordeStorage/Horde.Storage/Implementation/GC/CacheCleanup.cs
@@ -79,7 +79,7 @@ namespace Horde.Storage.Implementation
         {
             int countOfDeletedRecords = 0;
             DateTime cutoffTime = DateTime.Now.AddSeconds(-1 * _settings.CurrentValue.LastAccessCutoff.TotalSeconds);
-            int consideredCount = 0;
+            ulong consideredCount = 0;
             await Parallel.ForEachAsync(_referencesStore.GetRecords(ns), new ParallelOptions {MaxDegreeOfParallelism = _settings.CurrentValue.OrphanRefMaxParallelOperations, CancellationToken = cancellationToken}, async (tuple, token) =>
             {
                 (BucketId bucket, IoHashKey name, DateTime lastAccessTime) = tuple;
diff --git a/Engine/Source/Programs/Horde/HordeStorage/Horde.Storage/Implementation/GC/RefCleanupService.cs b/Engine/Source/Programs/Horde/HordeStorage/Horde.Storage/Implementation/GC/RefCleanupService.cs
--- a/Engine/Source/Programs/Horde/HordeStorage/Horde.Storage/Implementation/GC/RefCleanupService.cs
+++ b/Engine/Source/Programs/Horde/HordeStorage/Horde.Storage/Implementation/GC/RefCleanupService.cs
@@ -59,22 +59,25 @@ namespace Horde.Storage.Implementation
                     _logger.Information("Skipped ref cleanup run as this instance was not the leader");
                     return false;
                 }
-                await foreach (NamespaceId ns in state.Refs.GetNamespaces().WithCancellation(cancellationToken))
+                List<NamespaceId>? namespaces = await _referencesStore.GetNamespaces().ToListAsync(cancellationToken);
+                await Parallel.ForEachAsync(namespaces, cancellationToken, async (ns, token) =>
                 {
                     _logger.Information("Attempting to run Refs Cleanup of {Namespace}. ", ns);
                     try
                     {
-                        int countOfRemovedRecords = await state.RefCleanup.Cleanup(ns, cancellationToken);
-                        _logger.Information("Ran Refs Cleanup of {Namespace}. Deleted {CountRefRecords}", ns, countOfRemovedRecords);
+                        int countOfRemovedRecords = await state.RefCleanup.Cleanup(ns, token);
+                        _logger.Information("Ran Refs Cleanup of {Namespace}. Deleted {CountRefRecords}", ns,
+                            countOfRemovedRecords);
                     }
                     catch (Exception e)
                     {
                         _logger.Error("Error running Refs Cleanup of {Namespace}. {Exception}", ns, e);
                     }
-                }
+                });
 
-                List<NamespaceId>? namespaces = await _referencesStore.GetNamespaces().ToListAsync(cancellationToken);
-                await foreach (NamespaceId ns in namespaces)
+                // Do not run the cleanup of legacy namespaces as this can take a very long time and we do not care about these anymore
+                /*
+                 await foreach (NamespaceId ns in state.Refs.GetNamespaces().WithCancellation(cancellationToken))
                 {
                     _logger.Information("Attempting to run Refs Cleanup of {Namespace}. ", ns);
                     try
@@ -86,8 +89,7 @@ namespace Horde.Storage.Implementation
                     {
                         _logger.Error("Error running Refs Cleanup of {Namespace}. {Exception}", ns, e);
                     }
-                }
-
+                }*/
                 return true;
 
             }
