Horde: Add an implementation of IBlobStore built on the IStorageBackend interface in Horde.Build.

diff --git a/Engine/Source/Programs/Horde/Horde.Agent/Commands/Bundles/CreateCommand.cs b/Engine/Source/Programs/Horde/Horde.Agent/Commands/Bundles/CreateCommand.cs
--- a/Engine/Source/Programs/Horde/Horde.Agent/Commands/Bundles/CreateCommand.cs
+++ b/Engine/Source/Programs/Horde/Horde.Agent/Commands/Bundles/CreateCommand.cs
@@ -3,6 +3,7 @@
 using System;
 using System.Buffers;
 using System.Collections.Generic;
+using System.IO;
 using System.Threading;
 using System.Threading.Tasks;
 using EpicGames.Core;
@@ -16,41 +17,6 @@ namespace Horde.Agent.Commands.Bundles
 {
 	abstract class BundleCommandBase : Command
 	{
-		class FileBlob : IBlob
-		{
-			readonly IReadOnlyList<BlobId> _references;
-			readonly ReadOnlyMemory<byte> _data;
-
-			public FileBlob(IReadOnlyList<BlobId> references, ReadOnlyMemory<byte> data)
-			{
-				_references = references;
-				_data = data;
-			}
-
-			public static async Task<FileBlob> ReadAsync(FileReference file, CancellationToken cancellationToken)
-			{
-				byte[] bytes = await FileReference.ReadAllBytesAsync(file, cancellationToken);
-				MemoryReader reader = new MemoryReader(bytes);
-				IReadOnlyList<BlobId> references = reader.ReadVariableLengthArray(() => new BlobId(reader.ReadUtf8String()));
-				ReadOnlyMemory<byte> data = reader.ReadVariableLengthBytes();
-				return new FileBlob(references, data);
-			}
-
-			public async Task WriteAsync(FileReference file, CancellationToken cancellationToken)
-			{
-				ByteArrayBuilder writer = new ByteArrayBuilder();
-				writer.WriteVariableLengthArray(_references, x => writer.WriteUtf8String(x.Inner));
-				writer.WriteVariableLengthBytes(_data.Span);
-				await FileReference.WriteAllBytesAsync(file, writer.ToByteArray(), cancellationToken);
-			}
-
-			/// <inheritdoc/>
-			public ValueTask<ReadOnlyMemory<byte>> GetDataAsync() => new ValueTask<ReadOnlyMemory<byte>>(_data);
-
-			/// <inheritdoc/>
-			public ValueTask<IReadOnlyList<BlobId>> GetReferencesAsync() => new ValueTask<IReadOnlyList<BlobId>>(_references);
-		}
-
 		class FileBlobStore : IBlobStore
 		{
 			readonly DirectoryReference _rootDir;
@@ -67,14 +33,41 @@ namespace Horde.Agent.Commands.Bundles
 			FileReference GetRefFile(RefId id) => FileReference.Combine(_rootDir, id.Hash.ToString() + ".ref");
 			FileReference GetBlobFile(BlobId id) => FileReference.Combine(_rootDir, id.Inner.ToString() + ".blob");
 
-			public async Task<IBlob> ReadBlobAsync(BlobId id, CancellationToken cancellationToken = default)
+			static async Task WriteAsync(FileReference file, ReadOnlySequence<byte> sequence, CancellationToken cancellationToken)
+			{
+				using (FileStream stream = FileReference.Open(file, FileMode.Create, FileAccess.Write, FileShare.ReadWrite))
+				{
+					foreach (ReadOnlyMemory<byte> memory in sequence)
+					{
+						await stream.WriteAsync(memory, cancellationToken);
+					}
+				}
+			}
+
+			#region Blobs
+
+			public async Task<IBlob?> TryReadBlobAsync(BlobId id, CancellationToken cancellationToken = default)
 			{
 				FileReference file = GetBlobFile(id);
 				_logger.LogInformation("Reading {File}", file);
-				return await FileBlob.ReadAsync(file, cancellationToken);
+				byte[] bytes = await FileReference.ReadAllBytesAsync(file, cancellationToken);
+				return BlobUtils.Deserialize(bytes);
 			}
 
-			public Task<bool> DeleteRefAsync(RefId id, CancellationToken cancellationToken = default)
+			public async Task<BlobId> WriteBlobAsync(ReadOnlySequence<byte> data, IReadOnlyList<BlobId> references, CancellationToken cancellationToken = default)
+			{
+				BlobId id = new BlobId(Guid.NewGuid().ToString());
+				FileReference file = GetBlobFile(id);
+				_logger.LogInformation("Writing {File}", file);
+				await WriteAsync(file, BlobUtils.Serialize(data, references), cancellationToken);
+				return id;
+			}
+
+			#endregion
+
+			#region Refs
+
+			public Task DeleteRefAsync(RefId id, CancellationToken cancellationToken = default)
 			{
 				throw new NotImplementedException();
 			}
@@ -84,28 +77,22 @@ namespace Horde.Agent.Commands.Bundles
 				throw new NotImplementedException();
 			}
 
-			public async Task<IBlob> ReadRefAsync(RefId id, CancellationToken cancellationToken = default)
+			public async Task<IBlob?> TryReadRefAsync(RefId id, CancellationToken cancellationToken = default)
 			{
 				FileReference file = GetRefFile(id);
 				_logger.LogInformation("Reading {File}", file);
-				return await FileBlob.ReadAsync(file, cancellationToken);
-			}
-
-			public async Task<BlobId> WriteBlobAsync(ReadOnlySequence<byte> data, IReadOnlyList<BlobId> references, CancellationToken cancellationToken = default)
-			{
-				BlobId id = new BlobId(Guid.NewGuid().ToString());
-				FileReference file = GetBlobFile(id);
-				_logger.LogInformation("Writing {File}", file);
-				await new FileBlob(references, data.AsSingleSegment()).WriteAsync(file, cancellationToken);
-				return id;
+				byte[] bytes = await FileReference.ReadAllBytesAsync(file, cancellationToken);
+				return BlobUtils.Deserialize(bytes);
 			}
 
 			public async Task WriteRefAsync(RefId id, ReadOnlySequence<byte> data, IReadOnlyList<BlobId> references, CancellationToken cancellationToken = default)
 			{
 				FileReference file = GetRefFile(id);
 				_logger.LogInformation("Writing {File}", file);
-				await new FileBlob(references, data.AsSingleSegment()).WriteAsync(file, cancellationToken);
+				await WriteAsync(file, BlobUtils.Serialize(data, references), cancellationToken);
 			}
+
+			#endregion
 		}
 
 		public static RefId DefaultRefId { get; } = new RefId("default-ref");
diff --git a/Engine/Source/Programs/Horde/Horde.Agent/Commands/Bundles/DumpCommand.cs b/Engine/Source/Programs/Horde/Horde.Agent/Commands/Bundles/DumpCommand.cs
--- a/Engine/Source/Programs/Horde/Horde.Agent/Commands/Bundles/DumpCommand.cs
+++ b/Engine/Source/Programs/Horde/Horde.Agent/Commands/Bundles/DumpCommand.cs
@@ -37,7 +37,7 @@ namespace Horde.Agent.Commands.Bundles
 				logger.LogInformation("Summary for ref {RefId}", RefId);
 			}
 
-			IReadOnlyList<BlobId> references = await blob.GetReferencesAsync();
+			IReadOnlyList<BlobId> references = blob.References;
 			logger.LogInformation("");
 			logger.LogInformation("BlobRefs: {NumRefs}", references.Count);
 			foreach (BlobId reference in references)
@@ -45,7 +45,7 @@ namespace Horde.Agent.Commands.Bundles
 				logger.LogInformation("  {BlobId}", reference);
 			}
 
-			ReadOnlyMemory<byte> data = await blob.GetDataAsync();
+			ReadOnlyMemory<byte> data = blob.Data;
 			MemoryReader reader = new MemoryReader(data);
 			BundleHeader header = new BundleHeader(reader);
 			int packetStart = data.Length - reader.Memory.Length;
diff --git a/Engine/Source/Programs/Horde/Horde.Build/Storage/Backends/AwsStorageBackend.cs b/Engine/Source/Programs/Horde/Horde.Build/Storage/Backends/AwsStorageBackend.cs
--- a/Engine/Source/Programs/Horde/Horde.Build/Storage/Backends/AwsStorageBackend.cs
+++ b/Engine/Source/Programs/Horde/Horde.Build/Storage/Backends/AwsStorageBackend.cs
@@ -161,7 +161,7 @@ namespace Horde.Build.Storage.Backends
 		}
 
 		/// <inheritdoc/>
-		public async Task<Stream?> ReadAsync(string path)
+		public async Task<Stream?> ReadAsync(string path, CancellationToken cancellationToken)
 		{
 			string fullPath = GetFullPath(path);
 
@@ -169,13 +169,13 @@ namespace Horde.Build.Storage.Backends
 			GetObjectResponse? response = null;
 			try
 			{
-				semaLock = await _semaphore.UseWaitAsync();
+				semaLock = await _semaphore.UseWaitAsync(cancellationToken);
 
 				GetObjectRequest newGetRequest = new GetObjectRequest();
 				newGetRequest.BucketName = _options.BucketName;
 				newGetRequest.Key = fullPath;
 
-				response = await _client.GetObjectAsync(newGetRequest);
+				response = await _client.GetObjectAsync(newGetRequest, cancellationToken);
 
 				return new WrappedResponseStream(semaLock, response);
 			}
@@ -191,7 +191,7 @@ namespace Horde.Build.Storage.Backends
 		}
 
 		/// <inheritdoc/>
-		public async Task WriteAsync(string path, Stream inputStream)
+		public async Task WriteAsync(string path, Stream inputStream, CancellationToken cancellationToken)
 		{
 			TimeSpan[] retryTimes =
 			{
@@ -205,8 +205,8 @@ namespace Horde.Build.Storage.Backends
 			{
 				try
 				{
-					using IDisposable semaLock = await _semaphore.UseWaitAsync();
-					await WriteInternalAsync(fullPath, inputStream);
+					using IDisposable semaLock = await _semaphore.UseWaitAsync(cancellationToken);
+					await WriteInternalAsync(fullPath, inputStream, cancellationToken);
 					_logger.LogDebug("Written data to {Path}", path);
 					break;
 				}
@@ -224,7 +224,7 @@ namespace Horde.Build.Storage.Backends
 		}
 
 		/// <inheritdoc/>
-		public async Task WriteInternalAsync(string fullPath, Stream stream)
+		public async Task WriteInternalAsync(string fullPath, Stream stream, CancellationToken cancellationToken)
 		{
 			const int MinPartSize = 5 * 1024 * 1024;
 
@@ -233,7 +233,7 @@ namespace Horde.Build.Storage.Backends
 			{
 				// Read the data into memory. Errors with hash value not matching if we don't do this (?)
 				byte[] buffer = new byte[streamLen];
-				await ReadExactLengthAsync(stream, buffer, (int)streamLen);
+				await ReadExactLengthAsync(stream, buffer, (int)streamLen, cancellationToken);
 
 				// Upload it to S3
 				using (MemoryStream inputStream = new MemoryStream(buffer))
@@ -243,7 +243,7 @@ namespace Horde.Build.Storage.Backends
 					uploadRequest.Key = fullPath;
 					uploadRequest.InputStream = inputStream;
 					uploadRequest.Metadata.Add("bytes-written", streamLen.ToString(CultureInfo.InvariantCulture));
-					await _client.PutObjectAsync(uploadRequest);
+					await _client.PutObjectAsync(uploadRequest, cancellationToken);
 				}
 			}
 			else
@@ -253,7 +253,7 @@ namespace Horde.Build.Storage.Backends
 				initiateRequest.BucketName = _options.BucketName;
 				initiateRequest.Key = fullPath;
 
-				InitiateMultipartUploadResponse initiateResponse = await _client.InitiateMultipartUploadAsync(initiateRequest);
+				InitiateMultipartUploadResponse initiateResponse = await _client.InitiateMultipartUploadAsync(initiateRequest, cancellationToken);
 				try
 				{
 					// Buffer for reading the data
@@ -265,7 +265,7 @@ namespace Horde.Build.Storage.Backends
 					{
 						// Read the next chunk of data into the buffer
 						int bufferLen = (int)Math.Min((long)MinPartSize, streamLen - streamPos);
-						await ReadExactLengthAsync(stream, buffer, bufferLen);
+						await ReadExactLengthAsync(stream, buffer, bufferLen, cancellationToken);
 						streamPos += bufferLen;
 
 						// Upload the part
@@ -280,7 +280,7 @@ namespace Horde.Build.Storage.Backends
 							partRequest.PartNumber = partTags.Count + 1;
 							partRequest.IsLastPart = (streamPos == streamLen);
 
-							UploadPartResponse partResponse = await _client.UploadPartAsync(partRequest);
+							UploadPartResponse partResponse = await _client.UploadPartAsync(partRequest, cancellationToken);
 							partTags.Add(new PartETag(partResponse.PartNumber, partResponse.ETag));
 						}
 					}
@@ -291,7 +291,7 @@ namespace Horde.Build.Storage.Backends
 					completeRequest.Key = fullPath;
 					completeRequest.UploadId = initiateResponse.UploadId;
 					completeRequest.PartETags = partTags;
-					await _client.CompleteMultipartUploadAsync(completeRequest);
+					await _client.CompleteMultipartUploadAsync(completeRequest, cancellationToken);
 				}
 				catch
 				{
@@ -300,7 +300,7 @@ namespace Horde.Build.Storage.Backends
 					abortRequest.BucketName = _options.BucketName;
 					abortRequest.Key = fullPath;
 					abortRequest.UploadId = initiateResponse.UploadId;
-					await _client.AbortMultipartUploadAsync(abortRequest);
+					await _client.AbortMultipartUploadAsync(abortRequest, cancellationToken);
 
 					throw;
 				}
@@ -313,13 +313,14 @@ namespace Horde.Build.Storage.Backends
 		/// <param name="stream">The stream to read from</param>
 		/// <param name="buffer">The buffer to read into</param>
 		/// <param name="length">Length of the data to read</param>
+		/// <param name="cancellationToken">Cancellation token for the operation</param>
 		/// <returns>Async task</returns>
-		static async Task ReadExactLengthAsync(System.IO.Stream stream, byte[] buffer, int length)
+		static async Task ReadExactLengthAsync(System.IO.Stream stream, byte[] buffer, int length, CancellationToken cancellationToken)
 		{
 			int bufferPos = 0;
 			while (bufferPos < length)
 			{
-				int bytesRead = await stream.ReadAsync(buffer, bufferPos, length - bufferPos);
+				int bytesRead = await stream.ReadAsync(buffer, bufferPos, length - bufferPos, cancellationToken);
 				if (bytesRead == 0)
 				{
 					throw new InvalidOperationException("Unexpected end of stream");
@@ -329,23 +330,23 @@ namespace Horde.Build.Storage.Backends
 		}
 
 		/// <inheritdoc/>
-		public async Task DeleteAsync(string path)
+		public async Task DeleteAsync(string path, CancellationToken cancellationToken)
 		{
 			DeleteObjectRequest newDeleteRequest = new DeleteObjectRequest();
 			newDeleteRequest.BucketName = _options.BucketName;
 			newDeleteRequest.Key = GetFullPath(path);
-			await _client.DeleteObjectAsync(newDeleteRequest);
+			await _client.DeleteObjectAsync(newDeleteRequest, cancellationToken);
 		}
 
 		/// <inheritdoc/>
-		public async Task<bool> ExistsAsync(string path)
+		public async Task<bool> ExistsAsync(string path, CancellationToken cancellationToken)
 		{
 			try
 			{
 				GetObjectMetadataRequest request = new GetObjectMetadataRequest();
 				request.BucketName = _options.BucketName;
 				request.Key = GetFullPath(path);
-				await _client.GetObjectMetadataAsync(request);
+				await _client.GetObjectMetadataAsync(request, cancellationToken);
 				return true;
 			}
 			catch (AmazonS3Exception ex) when (ex.StatusCode == System.Net.HttpStatusCode.NotFound)
diff --git a/Engine/Source/Programs/Horde/Horde.Build/Storage/Backends/FileSystemStorageBackend.cs b/Engine/Source/Programs/Horde/Horde.Build/Storage/Backends/FileSystemStorageBackend.cs
--- a/Engine/Source/Programs/Horde/Horde.Build/Storage/Backends/FileSystemStorageBackend.cs
+++ b/Engine/Source/Programs/Horde/Horde.Build/Storage/Backends/FileSystemStorageBackend.cs
@@ -51,7 +51,7 @@ namespace Horde.Build.Storage.Backends
 		}
 
 		/// <inheritdoc/>
-		public Task<Stream?> ReadAsync(string path)
+		public Task<Stream?> ReadAsync(string path, CancellationToken cancellationToken)
 		{
 			FileReference location = FileReference.Combine(_baseDir, path);
 			if (!FileReference.Exists(location))
@@ -74,7 +74,7 @@ namespace Horde.Build.Storage.Backends
 		}
 
 		/// <inheritdoc/>
-		public async Task WriteAsync(string path, Stream stream)
+		public async Task WriteAsync(string path, Stream stream, CancellationToken cancellationToken)
 		{
 			FileReference finalLocation = FileReference.Combine(_baseDir, path);
 			if (!FileReference.Exists(finalLocation))
@@ -87,7 +87,7 @@ namespace Horde.Build.Storage.Backends
 
 				using (Stream outputStream = FileReference.Open(tempLocation, FileMode.Create, FileAccess.Write, FileShare.Read))
 				{
-					await stream.CopyToAsync(outputStream);
+					await stream.CopyToAsync(outputStream, cancellationToken);
 				}
 
 				// Move the temp file into place
@@ -110,14 +110,14 @@ namespace Horde.Build.Storage.Backends
 		}
 
 		/// <inheritdoc/>
-		public Task<bool> ExistsAsync(string path)
+		public Task<bool> ExistsAsync(string path, CancellationToken cancellationToken)
 		{
 			FileReference location = FileReference.Combine(_baseDir, path);
 			return Task.FromResult(FileReference.Exists(location));
 		}
 
 		/// <inheritdoc/>
-		public Task DeleteAsync(string path)
+		public Task DeleteAsync(string path, CancellationToken cancellationToken)
 		{
 			FileReference location = FileReference.Combine(_baseDir, path);
 			FileReference.Delete(location);
diff --git a/Engine/Source/Programs/Horde/Horde.Build/Storage/Backends/RelayStorageBackend.cs b/Engine/Source/Programs/Horde/Horde.Build/Storage/Backends/RelayStorageBackend.cs
--- a/Engine/Source/Programs/Horde/Horde.Build/Storage/Backends/RelayStorageBackend.cs
+++ b/Engine/Source/Programs/Horde/Horde.Build/Storage/Backends/RelayStorageBackend.cs
@@ -4,6 +4,7 @@ using System;
 using System.IO;
 using System.Net.Http;
 using System.Net.Http.Headers;
+using System.Threading;
 using System.Threading.Tasks;
 using Microsoft.Extensions.Options;
 
@@ -65,21 +66,21 @@ namespace Horde.Build.Storage.Backends
 		}
 
 		/// <inheritdoc/>
-		public async Task<Stream?> ReadAsync(string path)
+		public async Task<Stream?> ReadAsync(string path, CancellationToken cancellationToken)
 		{
-			Stream? localResult = await _localStorage.ReadAsync(path);
+			Stream? localResult = await _localStorage.ReadAsync(path, cancellationToken);
 			if (localResult != null)
 			{
 				return localResult;
 			}
 
 			Uri url = new Uri(_serverUrl, $"api/v1/debug/storage?path={path}");
-			using (HttpResponseMessage response = await _client.GetAsync(url))
+			using (HttpResponseMessage response = await _client.GetAsync(url, cancellationToken))
 			{
 				if (response.IsSuccessStatusCode)
 				{
-					byte[] responseData = await response.Content.ReadAsByteArrayAsync();
-					await _localStorage.WriteBytesAsync(path, responseData);
+					byte[] responseData = await response.Content.ReadAsByteArrayAsync(cancellationToken);
+					await _localStorage.WriteBytesAsync(path, responseData, cancellationToken);
 					return new MemoryStream(responseData);
 				}
 			}
@@ -88,19 +89,19 @@ namespace Horde.Build.Storage.Backends
 		}
 
 		/// <inheritdoc/>
-		public Task WriteAsync(string path, Stream stream)
+		public Task WriteAsync(string path, Stream stream, CancellationToken cancellationToken)
 		{
-			return _localStorage.WriteAsync(path, stream);
+			return _localStorage.WriteAsync(path, stream, cancellationToken);
 		}
 
 		/// <inheritdoc/>
-		public Task<bool> ExistsAsync(string path)
+		public Task<bool> ExistsAsync(string path, CancellationToken cancellationToken)
 		{
 			throw new NotSupportedException();
 		}
 
 		/// <inheritdoc/>
-		public Task DeleteAsync(string path)
+		public Task DeleteAsync(string path, CancellationToken cancellationToken)
 		{
 			throw new NotImplementedException();
 		}
diff --git a/Engine/Source/Programs/Horde/Horde.Build/Storage/Backends/TransientStorageBackend.cs b/Engine/Source/Programs/Horde/Horde.Build/Storage/Backends/TransientStorageBackend.cs
--- a/Engine/Source/Programs/Horde/Horde.Build/Storage/Backends/TransientStorageBackend.cs
+++ b/Engine/Source/Programs/Horde/Horde.Build/Storage/Backends/TransientStorageBackend.cs
@@ -2,6 +2,7 @@
 
 using System.Collections.Concurrent;
 using System.IO;
+using System.Threading;
 using System.Threading.Tasks;
 
 namespace Horde.Build.Storage.Backends
@@ -17,7 +18,7 @@ namespace Horde.Build.Storage.Backends
 		readonly ConcurrentDictionary<string, byte[]> _pathToData = new ConcurrentDictionary<string, byte[]>();
 
 		/// <inheritdoc/>
-		public Task<Stream?> ReadAsync(string path)
+		public Task<Stream?> ReadAsync(string path, CancellationToken cancellationToken)
 		{
 			byte[]? data;
 			if (_pathToData.TryGetValue(path, out data))
@@ -31,7 +32,7 @@ namespace Horde.Build.Storage.Backends
 		}
 
 		/// <inheritdoc/>
-		public async Task WriteAsync(string path, Stream stream)
+		public async Task WriteAsync(string path, Stream stream, CancellationToken cancellationToken)
 		{
 			using (MemoryStream buffer = new MemoryStream())
 			{
@@ -41,13 +42,13 @@ namespace Horde.Build.Storage.Backends
 		}
 
 		/// <inheritdoc/>
-		public Task<bool> ExistsAsync(string path)
+		public Task<bool> ExistsAsync(string path, CancellationToken cancellationToken)
 		{
 			return Task.FromResult(_pathToData.ContainsKey(path));
 		}
 
 		/// <inheritdoc/>
-		public Task DeleteAsync(string path)
+		public Task DeleteAsync(string path, CancellationToken cancellationToken)
 		{
 			_pathToData.TryRemove(path, out _);
 			return Task.CompletedTask;
diff --git a/Engine/Source/Programs/Horde/Horde.Build/Storage/Blobs/BlobStore.cs b/Engine/Source/Programs/Horde/Horde.Build/Storage/Blobs/BlobStore.cs
new file mode 100644
--- /dev/null
+++ b/Engine/Source/Programs/Horde/Horde.Build/Storage/Blobs/BlobStore.cs
@@ -0,0 +1,98 @@
+// Copyright Epic Games, Inc. All Rights Reserved.
+
+using System;
+using System.Buffers;
+using System.Collections.Generic;
+using System.Threading;
+using System.Threading.Tasks;
+using EpicGames.Core;
+using EpicGames.Horde.Storage;
+
+namespace Horde.Build.Storage
+{
+	/// <summary>
+	/// Implementation of <see cref="IBlobStore"/> for AWS
+	/// </summary>
+	public class BlobStore : IBlobStore
+	{
+		readonly IStorageBackend _backend;
+		readonly string _blobPrefix;
+		readonly string _refPrefix;
+
+		/// <summary>
+		/// Constructor
+		/// </summary>
+		/// <param name="backend">Backend to use for storing data</param>
+		/// <param name="blobPrefix">Prefix to append to any blob ids</param>
+		/// <param name="refPrefix">Prefix to append to any ref ids</param>
+		public BlobStore(IStorageBackend backend, string blobPrefix, string refPrefix)
+		{
+			_backend = backend;
+			_blobPrefix = blobPrefix;
+			_refPrefix = refPrefix;
+		}
+
+		async Task<IBlob?> TryReadAsync(string path, CancellationToken cancellationToken = default)
+		{
+			ReadOnlyMemory<byte>? data = await _backend.ReadBytesAsync(path, cancellationToken);
+			if (data == null)
+			{
+				return null;
+			}
+			return BlobUtils.Deserialize(data.Value);
+		}
+
+		async Task WriteAsync(string path, ReadOnlySequence<byte> data, IReadOnlyList<BlobId> references, CancellationToken cancellationToken = default)
+		{
+			ReadOnlySequence<byte> sequence = BlobUtils.Serialize(data, references);
+			await _backend.WriteBytesAsync(path, sequence.AsSingleSegment(), cancellationToken);
+		}
+
+		#region Blobs
+
+		string GetBlobPath(BlobId id) => $"{_blobPrefix}{id}";
+
+		/// <inheritdoc/>
+		public Task<IBlob?> TryReadBlobAsync(BlobId id, CancellationToken cancellationToken = default) => TryReadAsync(GetBlobPath(id), cancellationToken);
+
+		/// <inheritdoc/>
+		public async Task<BlobId> WriteBlobAsync(ReadOnlySequence<byte> data, IReadOnlyList<BlobId> references, CancellationToken cancellationToken = default)
+		{
+			BlobId id = new BlobId(Guid.NewGuid().ToString());
+			await WriteAsync(GetBlobPath(id), data, references, cancellationToken);
+			return id;
+		}
+
+		#endregion
+
+		#region Refs
+
+		string GetRefPath(RefId id) => $"{_refPrefix}{id}";
+
+		/// <inheritdoc/>
+		public async Task DeleteRefAsync(RefId id, CancellationToken cancellationToken = default)
+		{
+			await _backend.DeleteAsync(GetRefPath(id), cancellationToken);
+		}
+
+		/// <inheritdoc/>
+		public Task<bool> HasRefAsync(RefId id, CancellationToken cancellationToken = default)
+		{
+			throw new System.NotImplementedException();
+		}
+
+		/// <inheritdoc/>
+		public Task<IBlob?> TryReadRefAsync(RefId id, CancellationToken cancellationToken = default)
+		{
+			throw new System.NotImplementedException();
+		}
+
+		/// <inheritdoc/>
+		public Task WriteRefAsync(RefId id, ReadOnlySequence<byte> data, IReadOnlyList<BlobId> references, CancellationToken cancellationToken = default)
+		{
+			throw new System.NotImplementedException();
+		}
+
+		#endregion
+	}
+}
diff --git a/Engine/Source/Programs/Horde/Horde.Build/Storage/IStorageBackend.cs b/Engine/Source/Programs/Horde/Horde.Build/Storage/IStorageBackend.cs
--- a/Engine/Source/Programs/Horde/Horde.Build/Storage/IStorageBackend.cs
+++ b/Engine/Source/Programs/Horde/Horde.Build/Storage/IStorageBackend.cs
@@ -1,7 +1,9 @@
 // Copyright Epic Games, Inc. All Rights Reserved.
 
 using System;
+using System.Buffers;
 using System.IO;
+using System.Threading;
 using System.Threading.Tasks;
 using EpicGames.Core;
 
@@ -16,29 +18,33 @@ namespace Horde.Build.Storage
 		/// Opens a read stream for the given path.
 		/// </summary>
 		/// <param name="path">Relative path within the bucket</param>
+		/// <param name="cancellationToken">Cancellation token for the operation</param>
 		/// <returns></returns>
-		Task<Stream?> ReadAsync(string path);
+		Task<Stream?> ReadAsync(string path, CancellationToken cancellationToken = default);
 
 		/// <summary>
 		/// Writes a stream to the given path. If the stream throws an exception during read, the write will be aborted.
 		/// </summary>
 		/// <param name="path">Relative path within the bucket</param>
 		/// <param name="stream">Stream to write</param>
-		Task WriteAsync(string path, Stream stream);
+		/// <param name="cancellationToken">Cancellation token for the operation</param>
+		Task WriteAsync(string path, Stream stream, CancellationToken cancellationToken = default);
 
 		/// <summary>
 		/// Tests whether the given path exists
 		/// </summary>
 		/// <param name="path">Relative path within the bucket</param>
+		/// <param name="cancellationToken">Cancellation token for the operation</param>
 		/// <returns></returns>
-		Task<bool> ExistsAsync(string path);
+		Task<bool> ExistsAsync(string path, CancellationToken cancellationToken = default);
 
 		/// <summary>
 		/// Deletes a file with the given path
 		/// </summary>
 		/// <param name="path">Relative path within the bucket</param>
+		/// <param name="cancellationToken">Cancellation token for the operation</param>
 		/// <returns>Async task</returns>
-		Task DeleteAsync(string path);
+		Task DeleteAsync(string path, CancellationToken cancellationToken = default);
 	}
 
 	/// <summary>
@@ -69,16 +75,16 @@ namespace Horde.Build.Storage
 			public StorageBackend(IStorageBackend inner) => _inner = inner;
 
 			/// <inheritdoc/>
-			public Task<Stream?> ReadAsync(string path) => _inner.ReadAsync(path);
+			public Task<Stream?> ReadAsync(string path, CancellationToken cancellationToken) => _inner.ReadAsync(path, cancellationToken);
 
 			/// <inheritdoc/>
-			public Task WriteAsync(string path, Stream stream) => _inner.WriteAsync(path, stream);
+			public Task WriteAsync(string path, Stream stream, CancellationToken cancellationToken) => _inner.WriteAsync(path, stream, cancellationToken);
 
 			/// <inheritdoc/>
-			public Task DeleteAsync(string path) => _inner.DeleteAsync(path);
+			public Task DeleteAsync(string path, CancellationToken cancellationToken) => _inner.DeleteAsync(path, cancellationToken);
 
 			/// <inheritdoc/>
-			public Task<bool> ExistsAsync(string path) => _inner.ExistsAsync(path);
+			public Task<bool> ExistsAsync(string path, CancellationToken cancellationToken) => _inner.ExistsAsync(path, cancellationToken);
 		}
 
 		/// <summary>
@@ -97,10 +103,11 @@ namespace Horde.Build.Storage
 		/// </summary>
 		/// <param name="storageBackend"></param>
 		/// <param name="path"></param>
+		/// <param name="cancellationToken"></param>
 		/// <returns></returns>
-		public static async Task<ReadOnlyMemory<byte>?> ReadBytesAsync(this IStorageBackend storageBackend, string path)
+		public static async Task<ReadOnlyMemory<byte>?> ReadBytesAsync(this IStorageBackend storageBackend, string path, CancellationToken cancellationToken = default)
 		{
-			using (Stream? inputStream = await storageBackend.ReadAsync(path))
+			using (Stream? inputStream = await storageBackend.ReadAsync(path, cancellationToken))
 			{
 				if (inputStream == null)
 				{
@@ -109,7 +116,7 @@ namespace Horde.Build.Storage
 
 				using (MemoryStream outputStream = new MemoryStream())
 				{
-					await inputStream.CopyToAsync(outputStream);
+					await inputStream.CopyToAsync(outputStream, cancellationToken);
 					return outputStream.ToArray();
 				}
 			}
@@ -121,12 +128,13 @@ namespace Horde.Build.Storage
 		/// <param name="storageBackend"></param>
 		/// <param name="path"></param>
 		/// <param name="data"></param>
+		/// <param name="cancellationToken"></param>
 		/// <returns></returns>
-		public static async Task WriteBytesAsync(this IStorageBackend storageBackend, string path, ReadOnlyMemory<byte> data)
+		public static async Task WriteBytesAsync(this IStorageBackend storageBackend, string path, ReadOnlyMemory<byte> data, CancellationToken cancellationToken = default)
 		{
 			using (ReadOnlyMemoryStream stream = new ReadOnlyMemoryStream(data))
 			{
-				await storageBackend.WriteAsync(path, stream);
+				await storageBackend.WriteAsync(path, stream, cancellationToken);
 			}
 		}
 	}
diff --git a/Engine/Source/Programs/Shared/EpicGames.Horde/Storage/Backends/InMemoryBlobStore.cs b/Engine/Source/Programs/Shared/EpicGames.Horde/Storage/Backends/InMemoryBlobStore.cs
--- a/Engine/Source/Programs/Shared/EpicGames.Horde/Storage/Backends/InMemoryBlobStore.cs
+++ b/Engine/Source/Programs/Shared/EpicGames.Horde/Storage/Backends/InMemoryBlobStore.cs
@@ -52,7 +52,7 @@ namespace EpicGames.Horde.Storage.Backends
 		#region Refs
 
 		/// <inheritdoc/>
-		public Task<bool> DeleteRefAsync(RefId id, CancellationToken cancellationToken) => Task.FromResult(_refs.TryRemove(id, out _));
+		public Task DeleteRefAsync(RefId id, CancellationToken cancellationToken) => Task.FromResult(_refs.TryRemove(id, out _));
 
 		/// <inheritdoc/>
 		public Task<bool> HasRefAsync(RefId id, CancellationToken cancellationToken) => Task.FromResult(_refs.ContainsKey(id));
diff --git a/Engine/Source/Programs/Shared/EpicGames.Horde/Storage/Bundles/BundleStore.cs b/Engine/Source/Programs/Shared/EpicGames.Horde/Storage/Bundles/BundleStore.cs
--- a/Engine/Source/Programs/Shared/EpicGames.Horde/Storage/Bundles/BundleStore.cs
+++ b/Engine/Source/Programs/Shared/EpicGames.Horde/Storage/Bundles/BundleStore.cs
@@ -307,7 +307,7 @@ namespace EpicGames.Horde.Storage.Bundles
 		}
 
 		/// <inheritdoc/>
-		public Task<bool> DeleteTreeAsync(RefId id, CancellationToken cancellationToken) => _blobStore.DeleteRefAsync(id, cancellationToken);
+		public Task DeleteTreeAsync(RefId id, CancellationToken cancellationToken) => _blobStore.DeleteRefAsync(id, cancellationToken);
 
 		/// <summary>
 		/// Gets the data for a given node
diff --git a/Engine/Source/Programs/Shared/EpicGames.Horde/Storage/IBlobStore.cs b/Engine/Source/Programs/Shared/EpicGames.Horde/Storage/IBlobStore.cs
--- a/Engine/Source/Programs/Shared/EpicGames.Horde/Storage/IBlobStore.cs
+++ b/Engine/Source/Programs/Shared/EpicGames.Horde/Storage/IBlobStore.cs
@@ -86,8 +86,7 @@ namespace EpicGames.Horde.Storage
 		/// </summary>
 		/// <param name="id">The ref identifier</param>
 		/// <param name="cancellationToken">Cancellation token for the operation</param>
-		/// <returns></returns>
-		Task<bool> DeleteRefAsync(RefId id, CancellationToken cancellationToken = default);
+		Task DeleteRefAsync(RefId id, CancellationToken cancellationToken = default);
 
 		#endregion
 	}
@@ -95,7 +94,7 @@ namespace EpicGames.Horde.Storage
 	/// <summary>
 	/// Extension methods for <see cref="IBlobStore"/>
 	/// </summary>
-	static class BlobStoreExtensions
+	public static class BlobStoreExtensions
 	{
 		/// <summary>
 		/// Reads a blob from the store, throwing an exception if it does not exist
@@ -135,7 +134,7 @@ namespace EpicGames.Horde.Storage
 	/// <summary>
 	/// Utility methods for blobs
 	/// </summary>
-	static class BlobUtils
+	public static class BlobUtils
 	{
 		class InMemoryBlob : IBlob
 		{
diff --git a/Engine/Source/Programs/Shared/EpicGames.Horde/Storage/ITreeStore.cs b/Engine/Source/Programs/Shared/EpicGames.Horde/Storage/ITreeStore.cs
--- a/Engine/Source/Programs/Shared/EpicGames.Horde/Storage/ITreeStore.cs
+++ b/Engine/Source/Programs/Shared/EpicGames.Horde/Storage/ITreeStore.cs
@@ -62,8 +62,7 @@ namespace EpicGames.Horde.Storage
 		/// </summary>
 		/// <param name="id">Name of the ref used to store the tree</param>
 		/// <param name="cancellationToken">Cancellation token for the operation</param>
-		/// <returns>Tree writer instance</returns>
-		Task<bool> DeleteTreeAsync(RefId id, CancellationToken cancellationToken = default);
+		Task DeleteTreeAsync(RefId id, CancellationToken cancellationToken = default);
 
 		/// <summary>
 		/// Reads the root blob from a tree in storage
diff --git a/Engine/Source/Programs/Shared/EpicGames.Horde/Storage/Impl/MemoryStorageClient.cs b/Engine/Source/Programs/Shared/EpicGames.Horde/Storage/Impl/MemoryStorageClient.cs
--- a/Engine/Source/Programs/Shared/EpicGames.Horde/Storage/Impl/MemoryStorageClient.cs
+++ b/Engine/Source/Programs/Shared/EpicGames.Horde/Storage/Impl/MemoryStorageClient.cs
@@ -44,12 +44,16 @@
 		/// <inheritdoc/>
 		public Dictionary<(NamespaceId, BucketId, RefId), IRef> Refs { get; } = new Dictionary<(NamespaceId, BucketId, RefId), IRef>();
 
-		private FakeCompressor _compressor = new FakeCompressor();
+		private readonly FakeCompressor _compressor = new FakeCompressor();
 
 		/// <inheritdoc/>
 		public Task<Stream> ReadBlobAsync(NamespaceId namespaceId, IoHash hash, CancellationToken cancellationToken = default)
 		{
-			ReadOnlyMemory<byte> data = Blobs[(namespaceId, hash)];
+			ReadOnlyMemory<byte> data;
+			if(!Blobs.TryGetValue((namespaceId, hash), out data))
+			{
+				throw new BlobNotFoundException(namespaceId, hash);
+			}
 			return Task.FromResult<Stream>(new ReadOnlyMemoryStream(data));
 		}
 		
@@ -145,7 +149,12 @@
 		/// <inheritdoc/>
 		public Task<IRef> GetRefAsync(NamespaceId namespaceId, BucketId bucketId, RefId refId, CancellationToken cancellationToken = default)
 		{
-			return Task.FromResult(Refs[(namespaceId, bucketId, refId)]);
+			IRef? result;
+			if(!Refs.TryGetValue((namespaceId, bucketId, refId), out result))
+			{
+				throw new RefNotFoundException(namespaceId, bucketId, refId);
+			}
+			return Task.FromResult(result);
 		}
 
 		/// <inheritdoc/>
