Horde: Missing file.

diff --git a/Engine/Source/Programs/Horde/Horde.Build/Startup.cs b/Engine/Source/Programs/Horde/Horde.Build/Startup.cs
--- a/Engine/Source/Programs/Horde/Horde.Build/Startup.cs
+++ b/Engine/Source/Programs/Horde/Horde.Build/Startup.cs
@@ -765,22 +765,22 @@ namespace Horde.Build
 			}
 
 			/// <inheritdoc/>
-			public Task<Stream?> ReadAsync(string path) => _inner.ReadAsync(path);
+			public Task<Stream?> ReadAsync(string path, CancellationToken cancellationToken) => _inner.ReadAsync(path, cancellationToken);
 
 			/// <inheritdoc/>
-			public Task<ReadOnlyMemory<byte>?> ReadBytesAsync(string path) => _inner.ReadBytesAsync(path);
+			public Task<ReadOnlyMemory<byte>?> ReadBytesAsync(string path, CancellationToken cancellationToken) => _inner.ReadBytesAsync(path, cancellationToken);
 
 			/// <inheritdoc/>
-			public Task WriteAsync(string path, Stream stream) => _inner.WriteAsync(path, stream);
+			public Task WriteAsync(string path, Stream stream, CancellationToken cancellationToken) => _inner.WriteAsync(path, stream, cancellationToken);
 
 			/// <inheritdoc/>
-			public Task WriteBytesAsync(string path, ReadOnlyMemory<byte> data) => _inner.WriteBytesAsync(path, data);
+			public Task WriteBytesAsync(string path, ReadOnlyMemory<byte> data, CancellationToken cancellationToken) => _inner.WriteBytesAsync(path, data, cancellationToken);
 
 			/// <inheritdoc/>
-			public Task DeleteAsync(string path) => _inner.DeleteAsync(path);
+			public Task DeleteAsync(string path, CancellationToken cancellationToken) => _inner.DeleteAsync(path, cancellationToken);
 
 			/// <inheritdoc/>
-			public Task<bool> ExistsAsync(string path) => _inner.ExistsAsync(path);
+			public Task<bool> ExistsAsync(string path, CancellationToken cancellationToken) => _inner.ExistsAsync(path, cancellationToken);
 		}
 
 		private static void ConfigureLogStorage(IServiceCollection services)
