Horde: Refactor configuration of storage backends. Each use of storage now has its own configuration object rather than having a mixture of shared and unique settings.

diff --git a/Engine/Source/Programs/Horde/Horde.Build/ServerSettings.cs b/Engine/Source/Programs/Horde/Horde.Build/ServerSettings.cs
--- a/Engine/Source/Programs/Horde/Horde.Build/ServerSettings.cs
+++ b/Engine/Source/Programs/Horde/Horde.Build/ServerSettings.cs
@@ -3,26 +3,30 @@
 using System;
 using System.Linq;
 using EpicGames.Horde.Storage;
+using EpicGames.Horde.Storage.Bundles;
+using EpicGames.Horde.Storage.Nodes;
 using Horde.Build.Agents.Fleet;
+using Horde.Build.Storage;
+using Horde.Build.Storage.Backends;
 using Horde.Build.Utilities;
 using TimeZoneConverter;
 
 namespace Horde.Build
 {
 	/// <summary>
-	/// Types of storage to use for log data
+	/// Types of storage backend to use
 	/// </summary>
-	public enum StorageProviderType
+	public enum StorageBackendType
 	{
 		/// <summary>
-		/// AWS S3
-		/// </summary>
-		S3,
-
-		/// <summary>
 		/// Local filesystem
 		/// </summary>
 		FileSystem,
 
 		/// <summary>
+		/// AWS S3
+		/// </summary>
+		Aws,
+
+		/// <summary>
 		/// In-memory only (for testing)
@@ -36,5 +40,49 @@
 	};
 
 	/// <summary>
+	/// Common settings for different storage backends
+	/// </summary>
+	public interface IStorageBackendOptions : IFileSystemStorageOptions, IAwsStorageOptions, IRelayStorageOptions
+	{
+		/// <summary>
+		/// The type of storage backend to use
+		/// </summary>
+		StorageBackendType? Type { get; }
+	}
+
+	/// <summary>
+	/// Common settings object for different providers
+	/// </summary>
+	public class StorageBackendOptions : IStorageBackendOptions
+	{
+		/// <inheritdoc/>
+		public StorageBackendType? Type { get; set; }
+
+		/// <inheritdoc/>
+		public string? BaseDir { get; set; }
+
+		/// <inheritdoc/>
+		public string? AwsBucketName { get; set; }
+
+		/// <inheritdoc/>
+		public string? AwsBucketPath { get; set; }
+
+		/// <inheritdoc/>
+		public AwsCredentialsType AwsCredentials { get; set; }
+
+		/// <inheritdoc/>
+		public string? AwsRole { get; set; }
+
+		/// <inheritdoc/>
+		public string? AwsProfile { get; set; }
+
+		/// <inheritdoc/>
+		public string? RelayServer { get; set; }
+
+		/// <inheritdoc/>
+		public string? RelayToken { get; set; }
+	}
+
+	/// <summary>
 	/// Specifies the service to use for controlling the size of the fleet
 	/// </summary>
@@ -289,80 +337,14 @@
 		public string LogServiceWriteCacheType { get; set; } = "InMemory";
 
 		/// <summary>
-		/// Provider Type
-		/// Currently Supported: "S3" or "FileSystem"
+		/// Settings for artifact storage
 		/// </summary>
-		public StorageProviderType ExternalStorageProviderType { get; set; } = StorageProviderType.FileSystem;
+		public StorageBackendOptions LogStorage { get; set; } = new StorageBackendOptions() { BaseDir = "Logs" };
 
 		/// <summary>
-		/// Local storage directory, if using type filesystem
+		/// Settings for artifact storage
 		/// </summary>
-		public string LocalLogsDir { get; set; } = "Logs";
+		public StorageBackendOptions ArtifactStorage { get; set; } = new StorageBackendOptions() { BaseDir = "Artifacts" };
 
 		/// <summary>
-		/// Local artifact storage directory, if using type filesystem
-		/// </summary>
-		public string LocalArtifactsDir { get; set; } = "Artifacts";
-
-		/// <summary>
-		/// Local artifact storage directory, if using type filesystem
-		/// </summary>
-		public string LocalBlobsDir { get; set; } = "Blobs";
-
-		/// <summary>
-		/// S3 bucket region for logfile storage
-		/// </summary>
-		public string S3BucketRegion { get; set; } = null!;
-		
-		/// <summary>
-		/// AWS profile name. Used when S3CredentialType is set to "SharedCredentials".
-		/// </summary>
-		public string? S3AwsProfile { get; set; }
-
-		/// <summary>
-		/// Arn to assume for s3.  "Basic", "AssumeRole", "AssumeRoleWebIdentity", "SharedCredentials" only
-		/// </summary>
-		public string? S3CredentialType { get; set; }
-
-		/// <summary>
-		/// S3 Client username (used in Basic auth type only)
-		/// </summary>
-		public string S3ClientKeyId { get; set; } = null!;
-
-		/// <summary>
-		/// S3 client password (used in Basic auth type only)
-		/// </summary>
-		public string S3ClientSecret { get; set; } = null!;
-
-		/// <summary>
-		/// Arn to assume for s3
-		/// </summary>
-		public string S3AssumeArn { get; set; } = null!;
-
-		/// <summary>
-		/// S3 log bucket name
-		/// </summary>
-		public string S3LogBucketName { get; set; } = null!;
-
-		/// <summary>
-		/// S3 artifact bucket name
-		/// </summary>
-		public string S3ArtifactBucketName { get; set; } = null!;
-
-		/// <summary>
-		/// S3 blob bucket name
-		/// </summary>
-		public string S3BlobBucketName { get; set; } = null!;
-
-		/// <summary>
-		/// When using a relay storage provider, specifies the remote server to use
-		/// </summary>
-		public string? LogRelayServer { get; set; }
-
-		/// <summary>
-		/// Authentication token for using a relay server
-		/// </summary>
-		public string? LogRelayBearerToken { get; set; }
-
-		/// <summary>
 		/// Whether to log json to stdout
diff --git a/Engine/Source/Programs/Horde/Horde.Build/Startup.cs b/Engine/Source/Programs/Horde/Horde.Build/Startup.cs
--- a/Engine/Source/Programs/Horde/Horde.Build/Startup.cs
+++ b/Engine/Source/Programs/Horde/Horde.Build/Startup.cs
@@ -93,6 +93,8 @@ using Horde.Build.Auditing;
 using Horde.Build.Agents.Fleet.Providers;
 using Horde.Build.Server.Notices;
 using Horde.Build.Notifications.Sinks;
+using EpicGames.Horde.Storage.Bundles;
+using Microsoft.Extensions.Caching.Memory;
 
 namespace Horde.Build
 {
@@ -257,6 +259,23 @@ namespace Horde.Build
 			}
 		}
 
+		static IStorageBackend CreateStorageBackend(IServiceProvider sp, StorageBackendOptions options)
+		{
+			switch (options.Type ?? StorageBackendType.FileSystem)
+			{
+				case StorageBackendType.FileSystem:
+					return new FileSystemStorageBackend(options);
+				case StorageBackendType.Aws:
+					return new AwsStorageBackend(sp.GetRequiredService<IConfiguration>(), options, sp.GetRequiredService<ILogger<AwsStorageBackend>>());
+				case StorageBackendType.Transient:
+					return new TransientStorageBackend();
+				case StorageBackendType.Relay:
+					return new RelayStorageBackend(options);
+				default:
+					throw new NotImplementedException();
+			}
+		}
+
 		public Startup(IConfiguration configuration)
 		{
 			Configuration = configuration;
@@ -427,33 +446,9 @@ namespace Horde.Build
 				services.AddSingleton<IExternalIssueService, DefaultExternalIssueService>();
 			}
 
-
-			AWSOptions awsOptions = Configuration.GetAWSOptions();
-			if (settings.S3CredentialType == "AssumeRole" && settings.S3AssumeArn != null)
-			{
-				awsOptions.Credentials = new AssumeRoleAWSCredentials(FallbackCredentialsFactory.GetCredentials(), settings.S3AssumeArn, "Horde");
-			}
-			else if(settings.S3CredentialType == "AssumeRoleWebIdentity")
-			{
-				awsOptions.Credentials = AssumeRoleWithWebIdentityCredentials.FromEnvironmentVariables();
-			}
-			else if(settings.S3CredentialType == "Fallback")
-			{
-				// Using the fallback credentials from the AWS SDK, it will pick up credentials through a number of default mechanisms.
-				awsOptions.Credentials = FallbackCredentialsFactory.GetCredentials();
-			}
-			else if(settings.S3CredentialType == "SharedCredentials" && settings.S3AwsProfile != null)
-			{
-				// This SharedCredentials option is primarily for development purposes.
-				(string accessKey, string secretAccessKey, string secretToken) = AwsHelper.ReadAwsCredentials(settings.S3AwsProfile);
-				awsOptions.Credentials = new Amazon.SecurityToken.Model.Credentials(accessKey, secretAccessKey, secretToken, DateTime.Now + TimeSpan.FromHours(12));
-			}
-			
-			services.AddSingleton(awsOptions);
-
-			services.AddSingleton(new StorageBackendSettings<PersistentLogStorage> { Type = settings.ExternalStorageProviderType, BaseDir = settings.LocalLogsDir, BucketName = settings.S3LogBucketName });
-			services.AddSingleton(new StorageBackendSettings<ArtifactCollection> { Type = settings.ExternalStorageProviderType, BaseDir = settings.LocalArtifactsDir, BucketName = settings.S3ArtifactBucketName });
-			services.AddSingleton(typeof(IStorageBackend<>), typeof(StorageBackendFactory<>));
+			// Storage providers
+			services.AddSingleton<IStorageBackend<PersistentLogStorage>>(sp => CreateStorageBackend(sp, settings.LogStorage).ForType<PersistentLogStorage>());
+			services.AddSingleton<IStorageBackend<ArtifactCollection>>(sp => CreateStorageBackend(sp, settings.ArtifactStorage).ForType<ArtifactCollection>());
 
 			services.AddHordeStorage(settings => configSection.GetSection("Storage").Bind(settings));
 
@@ -722,67 +717,6 @@ namespace Horde.Build
 			options.Converters.Add(new JsonTimeSpanConverter());
 		}
 
-		public class StorageBackendSettings<T> : IFileSystemStorageOptions, IAwsStorageOptions
-		{
-			/// <summary>
-			/// The type of storage backend to use
-			/// </summary>
-			public StorageProviderType Type { get; set; }
-
-			/// <inheritdoc/>
-			public string? BaseDir { get; set; }
-
-			/// <inheritdoc/>
-			public string? BucketName { get; set; }
-
-			/// <inheritdoc/>
-			public string? BucketPath { get; set; }
-		}
-
-		public class StorageBackendFactory<T> : IStorageBackend<T>
-		{
-			readonly IStorageBackend _inner;
-
-			public StorageBackendFactory(IServiceProvider serviceProvider, StorageBackendSettings<T> options)
-			{
-				switch (options.Type)
-				{
-					case StorageProviderType.S3:
-						_inner = new AwsStorageBackend(serviceProvider.GetRequiredService<AWSOptions>(), options, serviceProvider.GetRequiredService<ILogger<AwsStorageBackend>>());
-						break;
-					case StorageProviderType.FileSystem:
-						_inner = new FileSystemStorageBackend(options);
-						break;
-					case StorageProviderType.Transient:
-						_inner = new TransientStorageBackend();
-						break;
-					case StorageProviderType.Relay:
-						_inner = new RelayStorageBackend(serviceProvider.GetRequiredService<IOptions<ServerSettings>>());
-						break;
-					default:
-						throw new NotImplementedException();
-				}
-			}
-
-			/// <inheritdoc/>
-			public Task<Stream?> ReadAsync(string path, CancellationToken cancellationToken) => _inner.ReadAsync(path, cancellationToken);
-
-			/// <inheritdoc/>
-			public Task<ReadOnlyMemory<byte>?> ReadBytesAsync(string path, CancellationToken cancellationToken) => _inner.ReadBytesAsync(path, cancellationToken);
-
-			/// <inheritdoc/>
-			public Task WriteAsync(string path, Stream stream, CancellationToken cancellationToken) => _inner.WriteAsync(path, stream, cancellationToken);
-
-			/// <inheritdoc/>
-			public Task WriteBytesAsync(string path, ReadOnlyMemory<byte> data, CancellationToken cancellationToken) => _inner.WriteBytesAsync(path, data, cancellationToken);
-
-			/// <inheritdoc/>
-			public Task DeleteAsync(string path, CancellationToken cancellationToken) => _inner.DeleteAsync(path, cancellationToken);
-
-			/// <inheritdoc/>
-			public Task<bool> ExistsAsync(string path, CancellationToken cancellationToken) => _inner.ExistsAsync(path, cancellationToken);
-		}
-
 		private static void ConfigureLogStorage(IServiceCollection services)
 		{
 			services.AddSingleton<ILogBuilder>(provider =>
@@ -1014,5 +948,13 @@ namespace Horde.Build
 			AddServices(serviceCollection, configuration);
 			return serviceCollection.BuildServiceProvider();
 		}
+
+		public static ServiceProvider CreateServiceProvider(IConfiguration configuration, ILoggerProvider loggerProvider)
+		{
+			IServiceCollection services = new ServiceCollection();
+			services.AddSingleton(loggerProvider);
+			AddServices(services, configuration);
+			return services.BuildServiceProvider();
+		}
 	}
 }
diff --git a/Engine/Source/Programs/Horde/Horde.Build/Storage/Backends/AwsStorageBackend.cs b/Engine/Source/Programs/Horde/Horde.Build/Storage/Backends/AwsStorageBackend.cs
--- a/Engine/Source/Programs/Horde/Horde.Build/Storage/Backends/AwsStorageBackend.cs
+++ b/Engine/Source/Programs/Horde/Horde.Build/Storage/Backends/AwsStorageBackend.cs
@@ -8,9 +8,11 @@ using System.Text;
 using System.Threading;
 using System.Threading.Tasks;
 using Amazon.Extensions.NETCore.Setup;
+using Amazon.Runtime;
 using Amazon.S3;
 using Amazon.S3.Model;
 using Horde.Build.Utilities;
+using Microsoft.Extensions.Configuration;
 using Microsoft.Extensions.Logging;
 
 namespace Horde.Build.Storage.Backends
@@ -31,20 +33,61 @@ namespace Horde.Build.Storage.Backends
 		}
 	}
 
+	/// <summary>
+	/// Credentials to use for AWS
+	/// </summary>
+	public enum AwsCredentialsType
+	{
+		/// <summary>
+		/// Assume a particular role. Should specify ARN in <see cref="IAwsStorageOptions.AwsRole"/>
+		/// </summary>
+		AssumeRole,
+
+		/// <summary>
+		/// Assume a particular role using the current environment variables.
+		/// </summary>
+		AssumeRoleWebIdentity,
+
+		/// <summary>
+		/// Use default credentials from the AWS SDK
+		/// </summary>
+		Fallback,
+
+		/// <summary>
+		/// Read credentials from a profile in the AWS config file
+		/// </summary>
+		SharedCredentials,
+	}
+
 	/// <summary>
 	/// Options for AWS
 	/// </summary>
 	public interface IAwsStorageOptions
 	{
+		/// <summary>
+		/// Type of credentials to use
+		/// </summary>
+		public AwsCredentialsType AwsCredentials { get; }
+
 		/// <summary>
 		/// Name of the bucket to use
 		/// </summary>
-		public string? BucketName { get; }
+		public string? AwsBucketName { get; }
 
 		/// <summary>
 		/// Base path within the bucket 
 		/// </summary>
-		public string? BucketPath { get; }
+		public string? AwsBucketPath { get; }
+
+		/// <summary>
+		/// ARN of a role to assume
+		/// </summary>
+		public string? AwsRole { get; }
+
+		/// <summary>
+		/// The AWS profile to read credentials form
+		/// </summary>
+		public string? AwsProfile { get; }
 	}
 
 	/// <summary>
@@ -75,17 +118,57 @@ namespace Horde.Build.Storage.Backends
 		/// <summary>
 		/// Constructor
 		/// </summary>
-		/// <param name="awsOptions">AWS options</param>
+		/// <param name="configuration">Global configuration object</param>
 		/// <param name="options">Storage options</param>
 		/// <param name="logger">Logger interface</param>
-		public AwsStorageBackend(AWSOptions awsOptions, IAwsStorageOptions options, ILogger<AwsStorageBackend> logger)
+		public AwsStorageBackend(IConfiguration configuration, IAwsStorageOptions options, ILogger<AwsStorageBackend> logger)
 		{
+			AWSOptions awsOptions = GetAwsOptions(configuration, options);
+
 			_client = awsOptions.CreateServiceClient<IAmazonS3>();
 			_options = options;
 			_semaphore = new SemaphoreSlim(16);
 			_logger = logger;
 
-			logger.LogInformation("Created AWS storage backend for bucket {BucketName} using credentials {Credentials} {CredentialsStr}", options.BucketName, awsOptions.Credentials.GetType(), awsOptions.Credentials.ToString());
+			logger.LogInformation("Created AWS storage backend for bucket {BucketName} using credentials {Credentials} {CredentialsStr}", options.AwsBucketName, awsOptions.Credentials.GetType(), awsOptions.Credentials.ToString());
+		}
+
+		/// <summary>
+		/// Gets the AWS options 
+		/// </summary>
+		/// <param name="configuration">Global configuration object</param>
+		/// <param name="options">AWS storage options</param>
+		/// <returns></returns>
+		static AWSOptions GetAwsOptions(IConfiguration configuration, IAwsStorageOptions options)
+		{
+			AWSOptions awsOptions = configuration.GetAWSOptions();
+			switch (options.AwsCredentials)
+			{
+				case AwsCredentialsType.AssumeRole:
+					if(options.AwsRole == null)
+					{
+						throw new AwsException($"Missing {nameof(IAwsStorageOptions.AwsRole)} setting for configuring {nameof(AwsStorageBackend)}", null);
+					}
+					awsOptions.Credentials = new AssumeRoleAWSCredentials(FallbackCredentialsFactory.GetCredentials(), options.AwsRole, "Horde");
+					break;
+				case AwsCredentialsType.AssumeRoleWebIdentity:
+					awsOptions.Credentials = AssumeRoleWithWebIdentityCredentials.FromEnvironmentVariables();
+					break;
+				case AwsCredentialsType.Fallback:
+					// Using the fallback credentials from the AWS SDK, it will pick up credentials through a number of default mechanisms.
+					awsOptions.Credentials = FallbackCredentialsFactory.GetCredentials();
+					break;
+				case AwsCredentialsType.SharedCredentials:
+					if (options.AwsProfile == null)
+					{
+						throw new AwsException($"Missing {nameof(IAwsStorageOptions.AwsProfile)} setting for configuring {nameof(AwsStorageBackend)}", null);
+					}
+
+					(string accessKey, string secretAccessKey, string secretToken) = AwsHelper.ReadAwsCredentials(options.AwsProfile);
+					awsOptions.Credentials = new Amazon.SecurityToken.Model.Credentials(accessKey, secretAccessKey, secretToken, DateTime.Now + TimeSpan.FromHours(12));
+					break;
+			}
+			return awsOptions;
 		}
 
 		/// <inheritdoc/>
@@ -145,13 +228,13 @@ namespace Horde.Build.Storage.Backends
 
 		string GetFullPath(string path)
 		{
-			if (_options.BucketPath == null)
+			if (_options.AwsBucketPath == null)
 			{
 				return path;
 			}
 
 			StringBuilder result = new StringBuilder();
-			result.Append(_options.BucketPath);
+			result.Append(_options.AwsBucketPath);
 			if (result.Length > 0 && result[^1] != '/')
 			{
 				result.Append('/');
@@ -172,7 +255,7 @@ namespace Horde.Build.Storage.Backends
 				semaLock = await _semaphore.UseWaitAsync(cancellationToken);
 
 				GetObjectRequest newGetRequest = new GetObjectRequest();
-				newGetRequest.BucketName = _options.BucketName;
+				newGetRequest.BucketName = _options.AwsBucketName;
 				newGetRequest.Key = fullPath;
 
 				response = await _client.GetObjectAsync(newGetRequest, cancellationToken);
@@ -215,7 +298,7 @@ namespace Horde.Build.Storage.Backends
 					_logger.LogError(ex, "Unable to write data to {Path} ({Attempt}/{AttemptCount})", fullPath, attempt + 1, retryTimes.Length + 1);
 					if (attempt >= retryTimes.Length)
 					{
-						throw new AwsException($"Unable to write to bucket {_options.BucketName} path {fullPath}", ex);
+						throw new AwsException($"Unable to write to bucket {_options.AwsBucketName} path {fullPath}", ex);
 					}
 				}
 
@@ -239,7 +322,7 @@ namespace Horde.Build.Storage.Backends
 				using (MemoryStream inputStream = new MemoryStream(buffer))
 				{
 					PutObjectRequest uploadRequest = new PutObjectRequest();
-					uploadRequest.BucketName = _options.BucketName;
+					uploadRequest.BucketName = _options.AwsBucketName;
 					uploadRequest.Key = fullPath;
 					uploadRequest.InputStream = inputStream;
 					uploadRequest.Metadata.Add("bytes-written", streamLen.ToString(CultureInfo.InvariantCulture));
@@ -250,7 +333,7 @@ namespace Horde.Build.Storage.Backends
 			{
 				// Initiate a multi-part upload
 				InitiateMultipartUploadRequest initiateRequest = new InitiateMultipartUploadRequest();
-				initiateRequest.BucketName = _options.BucketName;
+				initiateRequest.BucketName = _options.AwsBucketName;
 				initiateRequest.Key = fullPath;
 
 				InitiateMultipartUploadResponse initiateResponse = await _client.InitiateMultipartUploadAsync(initiateRequest, cancellationToken);
@@ -272,7 +355,7 @@ namespace Horde.Build.Storage.Backends
 						using (MemoryStream inputStream = new MemoryStream(buffer, 0, bufferLen))
 						{
 							UploadPartRequest partRequest = new UploadPartRequest();
-							partRequest.BucketName = _options.BucketName;
+							partRequest.BucketName = _options.AwsBucketName;
 							partRequest.Key = fullPath;
 							partRequest.UploadId = initiateResponse.UploadId;
 							partRequest.InputStream = inputStream;
@@ -287,7 +370,7 @@ namespace Horde.Build.Storage.Backends
 
 					// Mark the upload as complete
 					CompleteMultipartUploadRequest completeRequest = new CompleteMultipartUploadRequest();
-					completeRequest.BucketName = _options.BucketName;
+					completeRequest.BucketName = _options.AwsBucketName;
 					completeRequest.Key = fullPath;
 					completeRequest.UploadId = initiateResponse.UploadId;
 					completeRequest.PartETags = partTags;
@@ -297,7 +380,7 @@ namespace Horde.Build.Storage.Backends
 				{
 					// Abort the upload
 					AbortMultipartUploadRequest abortRequest = new AbortMultipartUploadRequest();
-					abortRequest.BucketName = _options.BucketName;
+					abortRequest.BucketName = _options.AwsBucketName;
 					abortRequest.Key = fullPath;
 					abortRequest.UploadId = initiateResponse.UploadId;
 					await _client.AbortMultipartUploadAsync(abortRequest, cancellationToken);
@@ -333,7 +416,7 @@ namespace Horde.Build.Storage.Backends
 		public async Task DeleteAsync(string path, CancellationToken cancellationToken)
 		{
 			DeleteObjectRequest newDeleteRequest = new DeleteObjectRequest();
-			newDeleteRequest.BucketName = _options.BucketName;
+			newDeleteRequest.BucketName = _options.AwsBucketName;
 			newDeleteRequest.Key = GetFullPath(path);
 			await _client.DeleteObjectAsync(newDeleteRequest, cancellationToken);
 		}
@@ -344,7 +427,7 @@ namespace Horde.Build.Storage.Backends
 			try
 			{
 				GetObjectMetadataRequest request = new GetObjectMetadataRequest();
-				request.BucketName = _options.BucketName;
+				request.BucketName = _options.AwsBucketName;
 				request.Key = GetFullPath(path);
 				await _client.GetObjectMetadataAsync(request, cancellationToken);
 				return true;
