Horde: Missing files.

diff --git a/Engine/Source/Programs/Horde/Horde.Build.Tests/TestSetup.cs b/Engine/Source/Programs/Horde/Horde.Build.Tests/TestSetup.cs
--- a/Engine/Source/Programs/Horde/Horde.Build.Tests/TestSetup.cs
+++ b/Engine/Source/Programs/Horde/Horde.Build.Tests/TestSetup.cs
@@ -56,6 +56,7 @@ using Horde.Build.Agents.Telemetry;
 using Horde.Build.Logs.Storage;
 using Horde.Build.Tasks;
 using Horde.Build.Auditing;
+using Horde.Build.Storage.Backends;
 
 namespace Horde.Build.Tests
 {
@@ -218,10 +219,9 @@ namespace Horde.Build.Tests
 			services.AddSingleton<ConformTaskSource>();
 			services.AddSingleton<ICommitService, CommitService>();
 
-			services.AddSingleton(new Startup.StorageBackendSettings<PersistentLogStorage> { Type = StorageProviderType.Transient });
-			services.AddSingleton(new Startup.StorageBackendSettings<ArtifactCollection> { Type = StorageProviderType.Transient });
-			services.AddSingleton(new Startup.StorageBackendSettings<BasicStorageClient> { Type = StorageProviderType.Transient });
-			services.AddSingleton(typeof(IStorageBackend<>), typeof(Startup.StorageBackendFactory<>));
+			services.AddSingleton<IStorageBackend<PersistentLogStorage>>(sp => new TransientStorageBackend().ForType<PersistentLogStorage>());
+			services.AddSingleton<IStorageBackend<ArtifactCollection>>(sp => new TransientStorageBackend().ForType<ArtifactCollection>());
+			services.AddSingleton<IStorageBackend<BasicStorageClient>>(sp => new TransientStorageBackend().ForType<BasicStorageClient>());
 
 			services.AddSingleton<IStorageClient, BasicStorageClient>();
 
diff --git a/Engine/Source/Programs/Horde/Horde.Build/Storage/Backends/RelayStorageBackend.cs b/Engine/Source/Programs/Horde/Horde.Build/Storage/Backends/RelayStorageBackend.cs
--- a/Engine/Source/Programs/Horde/Horde.Build/Storage/Backends/RelayStorageBackend.cs
+++ b/Engine/Source/Programs/Horde/Horde.Build/Storage/Backends/RelayStorageBackend.cs
@@ -6,20 +6,30 @@ using System.Net.Http;
 using System.Net.Http.Headers;
 using System.Threading;
 using System.Threading.Tasks;
-using Microsoft.Extensions.Options;
 
 namespace Horde.Build.Storage.Backends
 {
+	/// <summary>
+	/// Options for the relay storage backend
+	/// </summary>
+	public interface IRelayStorageOptions : IFileSystemStorageOptions
+	{
+		/// <summary>
+		/// Remote Horde server to use for storage using the relay storage backend
+		/// </summary>
+		public string? RelayServer { get; }
+
+		/// <summary>
+		/// Authentication token for using a relay server
+		/// </summary>
+		public string? RelayToken { get; }
+	}
+
 	/// <summary>
 	/// Implementation of ILogFileStorage which forwards requests to another server
 	/// </summary>
 	public sealed class RelayStorageBackend : IStorageBackend, IDisposable
 	{
-		class FileSystemSettings : IFileSystemStorageOptions
-		{
-			public string? BaseDir { get; set; }
-		}
-
 		/// <summary>
 		/// The client to connect with
 		/// </summary>
@@ -38,25 +48,24 @@ namespace Horde.Build.Storage.Backends
 		/// <summary>
 		/// Constructor
 		/// </summary>
-		/// <param name="settings">Settings for the server instance</param>
-		public RelayStorageBackend(IOptions<ServerSettings> settings)
+		/// <param name="options">Settings for the server instance</param>
+		public RelayStorageBackend(IRelayStorageOptions options)
 		{
-			ServerSettings currentSettings = settings.Value;
-			if (currentSettings.LogRelayServer == null)
+			if (options.RelayServer == null)
 			{
-				throw new InvalidDataException("Missing LogRelayServer in server configuration");
+				throw new InvalidDataException($"Missing {nameof(IRelayStorageOptions.RelayServer)} in server configuration");
 			}
-			if (currentSettings.LogRelayBearerToken == null)
+			if (options.RelayToken == null)
 			{
-				throw new InvalidDataException("Missing LogRelayBearerToken in server configuration");
+				throw new InvalidDataException($"Missing {nameof(IRelayStorageOptions.RelayToken)} in server configuration");
 			}
 
-			_serverUrl = new Uri(currentSettings.LogRelayServer);
+			_serverUrl = new Uri(options.RelayServer);
 
 			_client = new HttpClient();
-			_client.DefaultRequestHeaders.Authorization = new AuthenticationHeaderValue("Bearer", currentSettings.LogRelayBearerToken);
+			_client.DefaultRequestHeaders.Authorization = new AuthenticationHeaderValue("Bearer", options.RelayToken);
 
-			_localStorage = new FileSystemStorageBackend(new FileSystemSettings { BaseDir = "RelayCache" });
+			_localStorage = new FileSystemStorageBackend(options);
 		}
 
 		/// <inheritdoc/>
