Horde: Rework commit mirroring to use latest bundle implementation.

diff --git a/Engine/Source/Programs/Horde/Horde.Build/Commands/Bundles/CreateCommand.cs b/Engine/Source/Programs/Horde/Horde.Build/Commands/Bundles/CreateCommand.cs
new file mode 100644
--- /dev/null
+++ b/Engine/Source/Programs/Horde/Horde.Build/Commands/Bundles/CreateCommand.cs
@@ -0,0 +1,52 @@
+// Copyright Epic Games, Inc. All Rights Reserved.
+
+using System;
+using System.Buffers;
+using System.Collections.Generic;
+using System.IO;
+using System.Threading;
+using System.Threading.Tasks;
+using EpicGames.Core;
+using EpicGames.Horde.Storage;
+using EpicGames.Horde.Storage.Bundles;
+using EpicGames.Horde.Storage.Nodes;
+using Horde.Build.Perforce;
+using Microsoft.Extensions.Caching.Memory;
+using Microsoft.Extensions.Configuration;
+using Microsoft.Extensions.DependencyInjection;
+using Microsoft.Extensions.Logging;
+
+namespace Horde.Build.Commands.Bundles
+{
+	[Command("bundle", "create", "Creates a bundle from a folder on the local hard drive")]
+	class CreateCommand : Command
+	{
+		[CommandLine("-Ref=")]
+		public RefId RefId { get; set; } = new RefId("default-ref");
+
+		[CommandLine("-InputDir=", Required = true)]
+		public DirectoryReference InputDir { get; set; } = null!;
+
+		protected readonly IConfiguration _configuration;
+		protected readonly ILoggerProvider _loggerProvider;
+
+		public CreateCommand(IConfiguration configuration, ILoggerProvider loggerProvider)
+		{
+			_configuration = configuration;
+			_loggerProvider = loggerProvider;
+		}
+
+		public override async Task<int> ExecuteAsync(ILogger logger)
+		{
+			using ServiceProvider serviceProvider = Startup.CreateServiceProvider(_configuration, _loggerProvider);
+
+			ITreeStore store = serviceProvider.GetRequiredService<ITreeStore<CommitService>>();
+
+			DirectoryNode node = new DirectoryNode();
+			await node.CopyFromDirectoryAsync(InputDir.ToDirectoryInfo(), new ChunkingOptions(), logger, CancellationToken.None);
+			await store.WriteTreeAsync(RefId, node);
+
+			return 0;
+		}
+	}
+}
diff --git a/Engine/Source/Programs/Horde/Horde.Build/Commands/Bundles/ExtractCommand.cs b/Engine/Source/Programs/Horde/Horde.Build/Commands/Bundles/ExtractCommand.cs
new file mode 100644
--- /dev/null
+++ b/Engine/Source/Programs/Horde/Horde.Build/Commands/Bundles/ExtractCommand.cs
@@ -0,0 +1,46 @@
+// Copyright Epic Games, Inc. All Rights Reserved.
+
+using System.Threading;
+using System.Threading.Tasks;
+using EpicGames.Core;
+using EpicGames.Horde.Storage;
+using EpicGames.Horde.Storage.Nodes;
+using Horde.Build.Perforce;
+using Microsoft.Extensions.Caching.Memory;
+using Microsoft.Extensions.Configuration;
+using Microsoft.Extensions.DependencyInjection;
+using Microsoft.Extensions.Logging;
+
+namespace Horde.Build.Commands.Bundles
+{
+	[Command("bundle", "extract", "Extracts data from a bundle to the local hard drive")]
+	internal class ExtractCommand : Command
+	{
+		[CommandLine("-Ref=")]
+		public RefId RefId { get; set; } = new RefId("default-ref");
+
+		[CommandLine("-OutputDir=", Required = true)]
+		public DirectoryReference OutputDir { get; set; } = null!;
+
+		protected readonly IConfiguration _configuration;
+		protected readonly ILoggerProvider _loggerProvider;
+
+		public ExtractCommand(IConfiguration configuration, ILoggerProvider loggerProvider)
+		{
+			_configuration = configuration;
+			_loggerProvider = loggerProvider;
+		}
+
+		public override async Task<int> ExecuteAsync(ILogger logger)
+		{
+			using ServiceProvider serviceProvider = Startup.CreateServiceProvider(_configuration, _loggerProvider);
+
+			ITreeStore store = serviceProvider.GetRequiredService<ITreeStore<CommitService>>();
+
+			DirectoryNode node = await store.ReadTreeAsync<DirectoryNode>(RefId);
+			await node.CopyToDirectoryAsync(OutputDir.ToDirectoryInfo(), logger, CancellationToken.None);
+
+			return 0;
+		}
+	}
+}
diff --git a/Engine/Source/Programs/Horde/Horde.Build/Commands/Bundles/PerforceCommand.cs b/Engine/Source/Programs/Horde/Horde.Build/Commands/Bundles/PerforceCommand.cs
new file mode 100644
--- /dev/null
+++ b/Engine/Source/Programs/Horde/Horde.Build/Commands/Bundles/PerforceCommand.cs
@@ -0,0 +1,112 @@
+// Copyright Epic Games, Inc. All Rights Reserved.
+
+using System;
+using System.Buffers;
+using System.Collections.Generic;
+using System.IO;
+using System.Linq;
+using System.Threading;
+using System.Threading.Tasks;
+using EpicGames.Core;
+using EpicGames.Horde.Storage;
+using EpicGames.Horde.Storage.Bundles;
+using EpicGames.Horde.Storage.Nodes;
+using Horde.Build.Agents.Leases;
+using Horde.Build.Configuration;
+using Horde.Build.Jobs.Templates;
+using Horde.Build.Perforce;
+using Horde.Build.Server;
+using Horde.Build.Streams;
+using Horde.Build.Utilities;
+using HordeCommon;
+using Microsoft.Extensions.Caching.Memory;
+using Microsoft.Extensions.Configuration;
+using Microsoft.Extensions.DependencyInjection;
+using Microsoft.Extensions.Logging;
+using Microsoft.Extensions.Options;
+
+namespace Horde.Build.Commands.Bundles
+{
+	using StreamId = StringId<IStream>;
+
+	[Command("bundle", "perforce", "Replicates commits for a particular change of changes from Perforce")]
+	class PerforceCommand : Command
+	{
+		[CommandLine("-Stream=", Required = true)]
+		public string StreamId { get; set; } = String.Empty;
+
+		[CommandLine(Required = true)]
+		public int Change { get; set; }
+
+		[CommandLine]
+		public int? BaseChange { get; set; }
+
+		[CommandLine]
+		public int Count { get; set; } = 1;
+
+		[CommandLine]
+		public bool Content { get; set; }
+
+		[CommandLine]
+		public bool Compact { get; set; } = true;
+
+		[CommandLine]
+		public bool Clean { get; set; }
+
+		[CommandLine]
+		public string Filter { get; set; } = "...";
+
+		[CommandLine]
+		public bool Metadata { get; set; } = false;
+
+		[CommandLine]
+		public DirectoryReference? OutputDir { get; set; }
+
+		protected readonly IConfiguration _configuration;
+		protected readonly ILoggerProvider _loggerProvider;
+
+		public PerforceCommand(IConfiguration configuration, ILoggerProvider loggerProvider)
+		{
+			_configuration = configuration;
+			_loggerProvider = loggerProvider;
+		}
+
+		public override async Task<int> ExecuteAsync(ILogger logger)
+		{
+			using ServiceProvider serviceProvider = Startup.CreateServiceProvider(_configuration, _loggerProvider);
+
+			CommitService commitService = serviceProvider.GetRequiredService<CommitService>();
+			ICommitCollection commitCollection = serviceProvider.GetRequiredService<ICommitCollection>();
+			IStreamCollection streamCollection = serviceProvider.GetRequiredService<IStreamCollection>();
+			IStorageClient storageClient = serviceProvider.GetRequiredService<IStorageClient>();
+
+			IStream? stream = await streamCollection.GetAsync(new StreamId(StreamId));
+			if (stream == null)
+			{
+				throw new FatalErrorException($"Stream '{stream}' not found");
+			}
+
+			Dictionary<IStream, int> streamToFirstChange = new Dictionary<IStream, int>();
+			streamToFirstChange[stream] = Change;
+
+			await foreach (NewCommit newCommit in commitService.FindCommitsForClusterAsync(stream.ClusterName, streamToFirstChange).Take(Count))
+			{
+				string briefSummary = newCommit.Description.Replace('\n', ' ');
+				logger.LogInformation("");
+				logger.LogInformation("Commit {Change} by {AuthorId}: {Summary}", newCommit.Change, newCommit.AuthorId, briefSummary.Substring(0, Math.Min(50, briefSummary.Length)));
+				logger.LogInformation("Base path: {BasePath}", newCommit.BasePath);
+
+				if (Content)
+				{
+					if (Clean || BaseChange != null)
+					{
+						await commitService.WriteCommitTreeAsync(stream, newCommit.Change, BaseChange, Filter, Metadata, CancellationToken.None);
+					}
+					BaseChange = newCommit.Change;
+				}
+			}
+
+			return 0;
+		}
+	}
+}
diff --git a/Engine/Source/Programs/Horde/Horde.Build/Commands/Perforce/ReplicateCommand.cs b/Engine/Source/Programs/Horde/Horde.Build/Commands/Perforce/ReplicateCommand.cs
deleted file mode 100644
diff --git a/Engine/Source/Programs/Horde/Horde.Build/Perforce/CommitService.cs b/Engine/Source/Programs/Horde/Horde.Build/Perforce/CommitService.cs
--- a/Engine/Source/Programs/Horde/Horde.Build/Perforce/CommitService.cs
+++ b/Engine/Source/Programs/Horde/Horde.Build/Perforce/CommitService.cs
@@ -10,9 +10,9 @@ using System.Text.RegularExpressions;
 using System.Threading;
 using System.Threading.Tasks;
 using EpicGames.Core;
-using EpicGames.Horde.Bundles;
-using EpicGames.Horde.Bundles.Nodes;
 using EpicGames.Horde.Storage;
+using EpicGames.Horde.Storage.Bundles;
+using EpicGames.Horde.Storage.Nodes;
 using EpicGames.Perforce;
 using EpicGames.Redis;
 using EpicGames.Serialization;
@@ -181,11 +181,10 @@ namespace Horde.Build.Perforce
 
 		// Collections
 		readonly ICommitCollection _commitCollection;
-		readonly IStorageClient _storageClient;
 		readonly IStreamCollection _streamCollection;
 		readonly IPerforceService _perforceService;
 		readonly IUserCollection _userCollection;
-		readonly IMemoryCache _memoryCache;
+		readonly ITreeStore _treeStore;
 		readonly ILogger<CommitService> _logger;
 
 		const int MaxBackgroundTasks = 2;
@@ -193,7 +192,7 @@ namespace Horde.Build.Perforce
 		/// <summary>
 		/// Constructor
 		/// </summary>
-		public CommitService(RedisService redisService, ICommitCollection commitCollection, IStorageClient storageClient, IStreamCollection streamCollection, IPerforceService perforceService, IUserCollection userCollection, IClock clock, IMemoryCache memoryCache, IOptions<CommitServiceOptions> options, ILogger<CommitService> logger)
+		public CommitService(RedisService redisService, ICommitCollection commitCollection, IStreamCollection streamCollection, IPerforceService perforceService, IUserCollection userCollection, ITreeStore<CommitService> treeStore, IClock clock, IOptions<CommitServiceOptions> options, ILogger<CommitService> logger)
 		{
 			Options = options.Value;
 
@@ -202,11 +201,10 @@ namespace Horde.Build.Perforce
 			_redisReservations = new RedisSortedSet<StreamId>(_redis, RedisBaseKey.Append("reservations"));
 
 			_commitCollection = commitCollection;
-			_storageClient = storageClient;
 			_streamCollection = streamCollection;
 			_perforceService = perforceService;
 			_userCollection = userCollection;
-			_memoryCache = memoryCache;
+			_treeStore = treeStore;
 			_logger = logger;
 			_updateMetadataTicker = clock.AddSharedTicker<CommitService>(TimeSpan.FromSeconds(30.0), UpdateMetadataAsync, logger);
 		}
@@ -706,7 +704,7 @@ namespace Horde.Build.Perforce
 						{
 							if (await _redisReservations.AddAsync(checkStream, newScore, When.NotExists))
 							{
-								Task newTask = Task.Run(() => UpdateStreamContentAsync(checkStream));
+								Task newTask = Task.Run(() => UpdateStreamContentAsync(checkStream, CancellationToken.None));
 								backgroundTasks.Add((checkStream, newTask));
 
 								if (backgroundTasks.Count == MaxBackgroundTasks)
@@ -750,7 +748,7 @@ namespace Horde.Build.Perforce
 			}
 		}
 
-		async Task UpdateStreamContentAsync(StreamId streamId)
+		async Task UpdateStreamContentAsync(StreamId streamId, CancellationToken cancellationToken)
 		{
 			Stopwatch timer = Stopwatch.StartNew();
 
@@ -758,7 +756,7 @@ namespace Horde.Build.Perforce
 			for (; ; )
 			{
 				// Update the stream, updating the reservation every 30 seconds
-				Task internalTask = Task.Run(() => UpdateStreamContentInternalAsync(streamId, streamChanges));
+				Task internalTask = Task.Run(() => UpdateStreamContentInternalAsync(streamId, streamChanges, cancellationToken));
 				while (!internalTask.IsCompleted)
 				{
 					Task delayTask = Task.Delay(TimeSpan.FromSeconds(15.0));
@@ -792,7 +790,7 @@ namespace Horde.Build.Perforce
 			await _redisReservations.RemoveAsync(streamId);
 		}
 
-		async Task UpdateStreamContentInternalAsync(StreamId streamId, RedisList<int> changes)
+		async Task UpdateStreamContentInternalAsync(StreamId streamId, RedisList<int> changes, CancellationToken cancellationToken)
 		{
 			IStream? stream = await _streamCollection.GetAsync(streamId);
 			if (stream == null)
@@ -820,11 +818,11 @@ namespace Horde.Build.Perforce
 				// If we don't have a previous commit tree yet, perform a full snapshot of that change
 				if (prevChange == null)
 				{
-					await WriteCommitTreeAsync(stream, values[0], null);
+					await WriteCommitTreeAsync(stream, values[0], null, cancellationToken);
 				}
 
 				// Perform a snapshot of the new change, then remove it from the list
-				await WriteCommitTreeAsync(stream, values[1], prevChange);
+				await WriteCommitTreeAsync(stream, values[1], prevChange, cancellationToken);
 				prevChange = values[1];
 				await changes.LeftPopAsync();
 			}
@@ -836,28 +834,27 @@ namespace Horde.Build.Perforce
 		/// <param name="stream">The stream to replicate</param>
 		/// <param name="change">Commit to store the tree ref</param>
 		/// <param name="baseChange">The base change to update from</param>
+		/// <param name="cancellationToken">Cancellation token for the operation</param>
 		/// <returns>Root tree object</returns>
-		public async Task WriteCommitTreeAsync(IStream stream, int change, int? baseChange)
+		public async Task WriteCommitTreeAsync(IStream stream, int change, int? baseChange, CancellationToken cancellationToken)
 		{
 			if (stream.ReplicationMode != ContentReplicationMode.None)
 			{
-				await WriteCommitTreeAsync(stream, change, baseChange, stream.ReplicationFilter, stream.ReplicationMode == ContentReplicationMode.RevisionsOnly);
+				await WriteCommitTreeAsync(stream, change, baseChange, stream.ReplicationFilter, stream.ReplicationMode == ContentReplicationMode.RevisionsOnly, cancellationToken);
 			}
 		}
 
 		/// <summary>
-		/// Creates a new bundle or reads an existing one
+		/// Gets the ref for a particular change
 		/// </summary>
-		async ValueTask<Bundle<DirectoryNode>> CreateOrReadBundle(StreamId streamId, int? baseChange, string? filter, bool revisionsOnly)
+		/// <param name="streamId"></param>
+		/// <param name="change"></param>
+		/// <param name="filter"></param>
+		/// <param name="revisionsOnly"></param>
+		/// <returns></returns>
+		static RefId GetRefId(StreamId streamId, int change, string? filter, bool revisionsOnly)
 		{
-			if (baseChange == null)
-			{
-				return Bundle.Create<DirectoryNode>(_storageClient, Options.NamespaceId, Options.Bundle, _memoryCache);
-			}
-			else
-			{
-				return await Bundle.ReadAsync<DirectoryNode>(_storageClient, Options.NamespaceId, GetStreamBucketId(streamId), new CommitKey(streamId, baseChange.Value, filter, revisionsOnly).GetRefId(), Options.Bundle, _memoryCache);
-			}
+			return new CommitKey(streamId, change, filter, revisionsOnly).GetRefId();
 		}
 
 		/// <summary>
@@ -868,8 +865,9 @@ namespace Horde.Build.Perforce
 		/// <param name="baseChange">The base change to update from</param>
 		/// <param name="filter">Depot path to query for changes. Will default to the entire depot (but filtered by the workspace)</param>
 		/// <param name="revisionsOnly">Whether to replicate file revisions only</param>
+		/// <param name="cancellationToken">Cancellation token for the operation</param>
 		/// <returns>Root tree object</returns>
-		public async Task WriteCommitTreeAsync(IStream stream, int change, int? baseChange, string? filter, bool revisionsOnly)
+		public async Task WriteCommitTreeAsync(IStream stream, int change, int? baseChange, string? filter, bool revisionsOnly, CancellationToken cancellationToken)
 		{
 			// Create a client to replicate from this stream
 			ReplicationClient clientInfo = await FindOrAddReplicationClientAsync(stream);
@@ -879,7 +877,15 @@ namespace Horde.Build.Perforce
 			await FlushWorkspaceAsync(clientInfo, perforce, baseChange ?? 0);
 
 			// Get the initial bundle state
-			using Bundle<DirectoryNode> commitBundle = await CreateOrReadBundle(stream.Id, baseChange, filter, revisionsOnly);
+			DirectoryNode contents;
+			if (baseChange == null)
+			{
+				contents = new DirectoryNode();
+			}
+			else
+			{
+				contents = await _treeStore.ReadTreeAsync<DirectoryNode>(GetRefId(stream.Id, baseChange.Value, filter, revisionsOnly), cancellationToken);
+			}
 
 			// Apply all the updates
 			_logger.LogInformation("Updating client {Client} to changelist {Change}", clientInfo.Client.Name, change);
@@ -888,11 +894,11 @@ namespace Horde.Build.Perforce
 			Utf8String clientRoot = new Utf8String(clientInfo.Client.Root);
 			string queryPath = $"//{clientInfo.Client.Name}/{filter}";
 
-			DirectoryNode contents = commitBundle.Root;
+			RefId refId = GetRefId(stream.Id, change, filter, revisionsOnly);
 			if (revisionsOnly)
 			{
 				int count = 0;
-				await foreach (PerforceResponse<SyncRecord> record in perforce.StreamCommandAsync<SyncRecord>("sync", new[] { "-k", $"{queryPath}@{change}" }, null, default))
+				await foreach (PerforceResponse<SyncRecord> record in perforce.StreamCommandAsync<SyncRecord>("sync", new[] { "-k", $"{queryPath}@{change}" }, null, default, cancellationToken))
 				{
 					PerforceError? error = record.Error;
 					if (error != null && error.Generic == PerforceGenericCode.Empty)
@@ -907,14 +913,16 @@ namespace Horde.Build.Perforce
 					}
 
 					Utf8String path = syncRecord.Path.Substring(clientRoot.Length);
-					FileNode file = await contents.AddFileByPathAsync(commitBundle, path, FileEntryFlags.PerforceDepotPathAndRevision);
-
 					byte[] data = Encoding.UTF8.GetBytes($"{syncRecord.DepotFile}#{syncRecord.Revision}");
-					file.Append(data, Options.Chunking);
+
+					using ReadOnlyMemoryStream dataStream = new ReadOnlyMemoryStream(data);
+
+					FileEntry entry = await contents.AddFileByPathAsync(path, FileEntryFlags.PerforceDepotPathAndRevision, cancellationToken);
+					await entry.AppendAsync(dataStream, Options.Chunking, cancellationToken);
 
 					if (++count > 1000)
 					{
-						await commitBundle.TrimAsync();
+						await _treeStore.WriteTreeAsync(refId, contents, false, cancellationToken);
 						count = 0;
 					}
 				}
@@ -926,7 +934,7 @@ namespace Horde.Build.Perforce
 				long dataSize = 0;
 				long trimDataSize = TrimInterval;
 
-				Dictionary<int, FileNode> files = new Dictionary<int, FileNode>();
+				Dictionary<int, FileEntry> files = new Dictionary<int, FileEntry>();
 				await foreach (PerforceResponse response in perforce.StreamCommandAsync("sync", Array.Empty<string>(), new string[] { $"{queryPath}@{change}" }, null, typeof(SyncRecord), true, default))
 				{
 					PerforceError? error = response.Error;
@@ -942,12 +950,12 @@ namespace Horde.Build.Perforce
 						if (io.Command == PerforceIoCommand.Open)
 						{
 							Utf8String path = GetClientRelativePath(io.Payload, clientInfo.Client.Root);
-							files[io.File] = await contents.AddFileByPathAsync(commitBundle, path, 0);
+							files[io.File] = await contents.AddFileByPathAsync(path, FileEntryFlags.None, cancellationToken);
 						}
 						else if (io.Command == PerforceIoCommand.Write)
 						{
-							FileNode file = files[io.File];
-							file.Append(io.Payload, Options.Chunking);
+							FileEntry file = files[io.File];
+							await file.AppendAsync(io.Payload, Options.Chunking, cancellationToken);
 							dataSize += io.Payload.Length;
 						}
 						else if (io.Command == PerforceIoCommand.Close)
@@ -957,13 +965,18 @@ namespace Horde.Build.Perforce
 						else if (io.Command == PerforceIoCommand.Unlink)
 						{
 							Utf8String path = GetClientRelativePath(io.Payload, clientInfo.Client.Root);
-							await contents.DeleteFileByPathAsync(commitBundle, path);
+							await contents.DeleteFileByPathAsync(path, cancellationToken);
+						}
+						else
+						{
+							_logger.LogWarning("Unhandled command code {Code}", io.Command);
 						}
 
 						if (dataSize > trimDataSize)
 						{
 							_logger.LogInformation("Trimming working set after receiving {NumBytes:n0}mb...", dataSize / (1024 * 1024));
-							await commitBundle.TrimAsync();
+							await _treeStore.WriteTreeAsync(refId, contents, false, cancellationToken);
+							GC.Collect();
 							trimDataSize = dataSize + TrimInterval;
 							_logger.LogInformation("Trimming complete. Next trim at {NextNumBytes:n0}mb.", trimDataSize / (1024 * 1024));
 						}
@@ -973,18 +986,9 @@ namespace Horde.Build.Perforce
 
 			clientInfo.Change = change;
 
-			// Create the key for this commit
-			CommitKey key = new CommitKey(stream.Id, change, filter, revisionsOnly);
-			CbObject keyObject = CbSerializer.Serialize(key);
-
 			// Return the new root object
-			RefId refId = new RefId(IoHash.Compute(keyObject.GetView().Span));
 			_logger.LogInformation("Writing ref {RefId} for {StreamId} change {Change}", refId, stream.Id, change);
-			await commitBundle.WriteAsync(GetStreamBucketId(stream.Id), refId, keyObject, true);
-
-			// Print out stats for the new data
-			BundleStats stats = commitBundle.Stats;
-			_logger.LogInformation("Written ref {RefId} for {StreamId} change {Change}. Download: {NewBytes:n0} bytes ({NewRefCount} refs => {NewRefBytes:n0} bytes, {NewBlobCount} blobs => {NewBlobBytes:n0} bytes, {NewNodeCount:n0} nodes => {NewNodeBytes:n0} bytes).", refId, stream.Id, change, stats.NewRefBytes + stats.NewBlobBytes, stats.NewRefCount, stats.NewRefBytes, stats.NewBlobCount, stats.NewBlobBytes, stats.NewNodeCount, stats.NewNodeBytes);
+			await _treeStore.WriteTreeAsync(refId, contents, true, cancellationToken);
 		}
 
 		public QualifiedRefId? GetReplicatedContentRef(IStream stream, int change)
diff --git a/Engine/Source/Programs/Horde/Horde.Build/ServerSettings.cs b/Engine/Source/Programs/Horde/Horde.Build/ServerSettings.cs
--- a/Engine/Source/Programs/Horde/Horde.Build/ServerSettings.cs
+++ b/Engine/Source/Programs/Horde/Horde.Build/ServerSettings.cs
@@ -83,6 +83,46 @@ namespace Horde.Build
 		public string? RelayToken { get; set; }
 	}
 
+	/// <summary>
+	/// Options for configuring a blob store
+	/// </summary>
+	public class BlobStoreOptions : StorageBackendOptions, IBlobStoreOptions
+	{
+		/// <inheritdoc/>
+		public string BlobPrefix { get; set; } = "blobs/";
+
+		/// <inheritdoc/>
+		public string RefPrefix { get; set; } = "refs/";
+	}
+
+	/// <summary>
+	/// Options for configuring the default tree store implementation
+	/// </summary>
+	public interface ITreeStoreOptions
+	{
+		/// <summary>
+		/// Options for creating bundles
+		/// </summary>
+		BundleOptions Bundle { get; }
+
+		/// <summary>
+		/// Options for chunking content
+		/// </summary>
+		ChunkingOptions Chunking { get; }
+	}
+
+	/// <summary>
+	/// Options for storing trees
+	/// </summary>
+	public class TreeStoreOptions : BlobStoreOptions, ITreeStoreOptions
+	{
+		/// <inheritdoc/>
+		public BundleOptions Bundle { get; set; } = new BundleOptions();
+
+		/// <inheritdoc/>
+		public ChunkingOptions Chunking { get; set; } = new ChunkingOptions();
+	}
+
 	/// <summary>
 	/// Specifies the service to use for controlling the size of the fleet
 	/// </summary>
@@ -346,6 +386,11 @@ namespace Horde.Build
 		/// </summary>
 		public StorageBackendOptions ArtifactStorage { get; set; } = new StorageBackendOptions() { BaseDir = "Artifacts" };
 
+		/// <summary>
+		/// Configuration of tree storage
+		/// </summary>
+		public TreeStoreOptions CommitStorage { get; set; } = new TreeStoreOptions() { BaseDir = "Commits" };
+
 		/// <summary>
 		/// Whether to log json to stdout
 		/// </summary>
diff --git a/Engine/Source/Programs/Horde/Horde.Build/Startup.cs b/Engine/Source/Programs/Horde/Horde.Build/Startup.cs
--- a/Engine/Source/Programs/Horde/Horde.Build/Startup.cs
+++ b/Engine/Source/Programs/Horde/Horde.Build/Startup.cs
@@ -449,6 +449,7 @@ namespace Horde.Build
 			// Storage providers
 			services.AddSingleton<IStorageBackend<PersistentLogStorage>>(sp => CreateStorageBackend(sp, settings.LogStorage).ForType<PersistentLogStorage>());
 			services.AddSingleton<IStorageBackend<ArtifactCollection>>(sp => CreateStorageBackend(sp, settings.ArtifactStorage).ForType<ArtifactCollection>());
+			services.AddSingleton<ITreeStore<CommitService>>(sp => CreateTreeStore(sp, settings.CommitStorage).ForType<CommitService>());
 
 			services.AddHordeStorage(settings => configSection.GetSection("Storage").Bind(settings));
 
diff --git a/Engine/Source/Programs/Horde/Horde.Build/Storage/Blobs/BlobStore.cs b/Engine/Source/Programs/Horde/Horde.Build/Storage/Blobs/BlobStore.cs
--- a/Engine/Source/Programs/Horde/Horde.Build/Storage/Blobs/BlobStore.cs
+++ b/Engine/Source/Programs/Horde/Horde.Build/Storage/Blobs/BlobStore.cs
@@ -10,26 +10,39 @@ using EpicGames.Horde.Storage;
 
 namespace Horde.Build.Storage
 {
+	/// <summary>
+	/// Options for configuring the default blob store implementation
+	/// </summary>
+	public interface IBlobStoreOptions
+	{
+		/// <summary>
+		/// Prefix for storing blobs
+		/// </summary>
+		string BlobPrefix { get; }
+
+		/// <summary>
+		/// Prefix for storing refs
+		/// </summary>
+		string RefPrefix { get; }
+	}
+
 	/// <summary>
 	/// Implementation of <see cref="IBlobStore"/> for AWS
 	/// </summary>
 	public class BlobStore : IBlobStore
 	{
 		readonly IStorageBackend _backend;
-		readonly string _blobPrefix;
-		readonly string _refPrefix;
+		readonly IBlobStoreOptions _options;
 
 		/// <summary>
 		/// Constructor
 		/// </summary>
 		/// <param name="backend">Backend to use for storing data</param>
-		/// <param name="blobPrefix">Prefix to append to any blob ids</param>
-		/// <param name="refPrefix">Prefix to append to any ref ids</param>
-		public BlobStore(IStorageBackend backend, string blobPrefix, string refPrefix)
+		/// <param name="options">Options for the blob store</param>
+		public BlobStore(IStorageBackend backend, IBlobStoreOptions options)
 		{
 			_backend = backend;
-			_blobPrefix = blobPrefix;
-			_refPrefix = refPrefix;
+			_options = options;
 		}
 
 		async Task<IBlob?> TryReadAsync(string path, CancellationToken cancellationToken = default)
@@ -50,7 +63,7 @@ namespace Horde.Build.Storage
 
 		#region Blobs
 
-		string GetBlobPath(BlobId id) => $"{_blobPrefix}{id}";
+		string GetBlobPath(BlobId id) => $"{_options.BlobPrefix}{id}";
 
 		/// <inheritdoc/>
 		public Task<IBlob?> TryReadBlobAsync(BlobId id, CancellationToken cancellationToken = default) => TryReadAsync(GetBlobPath(id), cancellationToken);
@@ -67,7 +80,7 @@ namespace Horde.Build.Storage
 
 		#region Refs
 
-		string GetRefPath(RefId id) => $"{_refPrefix}{id}";
+		string GetRefPath(RefId id) => $"{_options.RefPrefix}{id}";
 
 		/// <inheritdoc/>
 		public async Task DeleteRefAsync(RefId id, CancellationToken cancellationToken = default)
@@ -76,21 +89,21 @@ namespace Horde.Build.Storage
 		}
 
 		/// <inheritdoc/>
-		public Task<bool> HasRefAsync(RefId id, CancellationToken cancellationToken = default)
+		public async Task<bool> HasRefAsync(RefId id, CancellationToken cancellationToken = default)
 		{
-			throw new System.NotImplementedException();
+			return await _backend.ExistsAsync(GetRefPath(id), cancellationToken);
 		}
 
 		/// <inheritdoc/>
-		public Task<IBlob?> TryReadRefAsync(RefId id, CancellationToken cancellationToken = default)
+		public async Task<IBlob?> TryReadRefAsync(RefId id, CancellationToken cancellationToken = default)
 		{
-			throw new System.NotImplementedException();
+			return await TryReadAsync(GetRefPath(id), cancellationToken);
 		}
 
 		/// <inheritdoc/>
-		public Task WriteRefAsync(RefId id, ReadOnlySequence<byte> data, IReadOnlyList<BlobId> references, CancellationToken cancellationToken = default)
+		public async Task WriteRefAsync(RefId id, ReadOnlySequence<byte> data, IReadOnlyList<BlobId> references, CancellationToken cancellationToken = default)
 		{
-			throw new System.NotImplementedException();
+			await WriteAsync(GetRefPath(id), data, references, cancellationToken);
 		}
 
 		#endregion
diff --git a/Engine/Source/Programs/Shared/EpicGames.Horde/Storage/Bundles/BundleStore.cs b/Engine/Source/Programs/Shared/EpicGames.Horde/Storage/Bundles/BundleStore.cs
--- a/Engine/Source/Programs/Shared/EpicGames.Horde/Storage/Bundles/BundleStore.cs
+++ b/Engine/Source/Programs/Shared/EpicGames.Horde/Storage/Bundles/BundleStore.cs
@@ -261,6 +261,13 @@ namespace EpicGames.Horde.Storage.Bundles
 			return MountInMemoryBundle(bundle);
 		}
 
+		/// <inheritdoc/>
+		public async Task<T> ReadTreeAsync<T>(RefId refName, CancellationToken cancellationToken = default) where T : TreeNode
+		{
+			ITreeBlob node = await ReadTreeAsync(refName, cancellationToken);
+			return await TreeNode.DeserializeAsync<T>(node, cancellationToken);
+		}
+
 		/// <inheritdoc/>
 		public async Task WriteTreeAsync(RefId id, ITreeBlob root, bool flush = true, CancellationToken cancellationToken = default)
 		{
@@ -306,6 +313,13 @@ namespace EpicGames.Horde.Storage.Bundles
 			}
 		}
 
+		/// <inheritdoc/>
+		public async Task WriteTreeAsync<T>(RefId refName, T root, bool flush = true, CancellationToken cancellationToken = default) where T : TreeNode
+		{
+			ITreeBlob rootBlob = await TreeNode.SerializeAsync(root, this, cancellationToken);
+			await WriteTreeAsync(refName, rootBlob, flush, cancellationToken);
+		}
+
 		/// <inheritdoc/>
 		public Task DeleteTreeAsync(RefId id, CancellationToken cancellationToken) => _blobStore.DeleteRefAsync(id, cancellationToken);
 
@@ -849,32 +863,4 @@ namespace EpicGames.Horde.Storage.Bundles
 			}
 		}
 	}
-
-	/// <summary>
-	/// Manages a collection of bundles.
-	/// </summary>
-	public class BundleStore<T> : BundleStore, ITreeStore<T> where T : TreeNode
-	{
-		/// <summary>
-		/// Constructor
-		/// </summary>
-		public BundleStore(IBlobStore blobStore, BundleOptions options, IMemoryCache cache)
-			: base(blobStore, options, cache)
-		{
-		}
-
-		/// <inheritdoc/>
-		public new async Task<T> ReadTreeAsync(RefId refName, CancellationToken cancellationToken = default)
-		{
-			ITreeBlob node = await base.ReadTreeAsync(refName, cancellationToken);
-			return await TreeNode.DeserializeAsync<T>(node, cancellationToken);
-		}
-
-		/// <inheritdoc/>
-		public async Task WriteTreeAsync(RefId refName, T root, bool flush = true, CancellationToken cancellationToken = default)
-		{
-			ITreeBlob rootBlob = await TreeNode.SerializeAsync(root, this, cancellationToken);
-			await base.WriteTreeAsync(refName, rootBlob, flush, cancellationToken);
-		}
-	}
 }
diff --git a/Engine/Source/Programs/Shared/EpicGames.Horde/Storage/IBlobStore.cs b/Engine/Source/Programs/Shared/EpicGames.Horde/Storage/IBlobStore.cs
--- a/Engine/Source/Programs/Shared/EpicGames.Horde/Storage/IBlobStore.cs
+++ b/Engine/Source/Programs/Shared/EpicGames.Horde/Storage/IBlobStore.cs
@@ -91,11 +91,58 @@ namespace EpicGames.Horde.Storage
 		#endregion
 	}
 
+	/// <summary>
+	/// Typed implementation of <see cref="IBlobStore"/> for use with dependency injection
+	/// </summary>
+	public interface IBlobStore<T> : IBlobStore
+	{
+	}
+
 	/// <summary>
 	/// Extension methods for <see cref="IBlobStore"/>
 	/// </summary>
 	public static class BlobStoreExtensions
 	{
+		class TypedBlobStore<T> : IBlobStore<T>
+		{
+			readonly IBlobStore _inner;
+
+			public TypedBlobStore(IBlobStore inner) => _inner = inner;
+
+			#region Blobs
+
+			/// <inheritdoc/>
+			public Task<IBlob?> TryReadBlobAsync(BlobId id, CancellationToken cancellationToken = default) => _inner.TryReadBlobAsync(id, cancellationToken);
+
+			/// <inheritdoc/>
+			public Task<BlobId> WriteBlobAsync(ReadOnlySequence<byte> data, IReadOnlyList<BlobId> references, CancellationToken cancellationToken = default) => _inner.WriteBlobAsync(data, references, cancellationToken);
+
+			#endregion
+
+			#region Refs
+
+			/// <inheritdoc/>
+			public Task DeleteRefAsync(RefId id, CancellationToken cancellationToken = default) => _inner.DeleteRefAsync(id, cancellationToken);
+
+			/// <inheritdoc/>
+			public Task<bool> HasRefAsync(RefId id, CancellationToken cancellationToken = default) => _inner.HasRefAsync(id, cancellationToken);
+
+			/// <inheritdoc/>
+			public Task<IBlob?> TryReadRefAsync(RefId id, CancellationToken cancellationToken = default) => _inner.TryReadRefAsync(id, cancellationToken);
+
+			/// <inheritdoc/>
+			public Task WriteRefAsync(RefId id, ReadOnlySequence<byte> data, IReadOnlyList<BlobId> references, CancellationToken cancellationToken = default) => _inner.WriteRefAsync(id, data, references, cancellationToken);
+
+			#endregion
+		}
+
+		/// <summary>
+		/// Wraps a <see cref="IBlobStore"/> interface with a type argument
+		/// </summary>
+		/// <param name="blobStore">Regular blob store instance</param>
+		/// <returns></returns>
+		public static IBlobStore<T> ForType<T>(this IBlobStore blobStore) => new TypedBlobStore<T>(blobStore);
+
 		/// <summary>
 		/// Reads a blob from the store, throwing an exception if it does not exist
 		/// </summary>
diff --git a/Engine/Source/Programs/Shared/EpicGames.Horde/Storage/ITreeStore.cs b/Engine/Source/Programs/Shared/EpicGames.Horde/Storage/ITreeStore.cs
--- a/Engine/Source/Programs/Shared/EpicGames.Horde/Storage/ITreeStore.cs
+++ b/Engine/Source/Programs/Shared/EpicGames.Horde/Storage/ITreeStore.cs
@@ -72,6 +72,14 @@ namespace EpicGames.Horde.Storage
 		/// <returns></returns>
 		Task<ITreeBlob> ReadTreeAsync(RefId id, CancellationToken cancellationToken = default);
 
+		/// <summary>
+		/// Reads a root node from the store
+		/// </summary>
+		/// <param name="id">Name of the tree to fetch</param>
+		/// <param name="cancellationToken">Cancellation token for the operation</param>
+		/// <returns>Deserialized node stored with this name</returns>
+		Task<T> ReadTreeAsync<T>(RefId id, CancellationToken cancellationToken = default) where T : TreeNode;
+
 		/// <summary>
 		/// Creates a writer for a tree with the given name
 		/// </summary>
@@ -81,20 +89,6 @@ namespace EpicGames.Horde.Storage
 		/// <param name="cancellationToken">Cancellation token for the operation</param>
 		/// <returns>Tree writer instance</returns>
 		Task WriteTreeAsync(RefId id, ITreeBlob root, bool flush = true, CancellationToken cancellationToken = default);
-	}
-
-	/// <summary>
-	/// Interface for a typed tree store
-	/// </summary>
-	public interface ITreeStore<T> : ITreeStore where T : TreeNode
-	{
-		/// <summary>
-		/// Reads a root node from the store
-		/// </summary>
-		/// <param name="id">Name of the tree to fetch</param>
-		/// <param name="cancellationToken">Cancellation token for the operation</param>
-		/// <returns>Deserialized node stored with this name</returns>
-		new Task<T> ReadTreeAsync(RefId id, CancellationToken cancellationToken = default);
 
 		/// <summary>
 		/// Creates a writer for a tree with the given name
@@ -104,6 +98,61 @@ namespace EpicGames.Horde.Storage
 		/// <param name="flush">Whether to flush the complete tree state. If false, an implementation may choose to buffer some writes to allow packing multiple nodes together.</param>
 		/// <param name="cancellationToken">Cancellation token for the operation</param>
 		/// <returns>Tree writer instance</returns>
-		Task WriteTreeAsync(RefId id, T root, bool flush = true, CancellationToken cancellationToken = default);
+		Task WriteTreeAsync<T>(RefId id, T root, bool flush = true, CancellationToken cancellationToken = default) where T : TreeNode;
+	}
+
+	/// <summary>
+	/// Interface for a typed tree store
+	/// </summary>
+	public interface ITreeStore<T> : ITreeStore
+	{
+	}
+
+	/// <summary>
+	/// Extension methods for tree store instances
+	/// </summary>
+	public static class TreeStoreExtensions
+	{
+		sealed class TypedTreeStore<T> : ITreeStore<T>
+		{
+			readonly ITreeStore _inner;
+
+			public TypedTreeStore(ITreeStore inner)
+			{
+				_inner = inner;
+			}
+
+			/// <inheritdoc/>
+			public Task DeleteTreeAsync(RefId id, CancellationToken cancellationToken = default) => _inner.DeleteTreeAsync(id, cancellationToken);
+
+			/// <inheritdoc/>
+			public void Dispose() => _inner.Dispose();
+
+			/// <inheritdoc/>
+			public Task<bool> HasTreeAsync(RefId id, CancellationToken cancellationToken = default) => _inner.HasTreeAsync(id, cancellationToken);
+
+			/// <inheritdoc/>
+			public Task<ITreeBlob> ReadTreeAsync(RefId id, CancellationToken cancellationToken = default) => _inner.ReadTreeAsync(id, cancellationToken);
+
+			/// <inheritdoc/>
+			public Task<TNode> ReadTreeAsync<TNode>(RefId id, CancellationToken cancellationToken = default) where TNode : TreeNode => _inner.ReadTreeAsync<TNode>(id, cancellationToken);
+
+			/// <inheritdoc/>
+			public Task WriteTreeAsync(RefId id, ITreeBlob root, bool flush = true, CancellationToken cancellationToken = default) => _inner.WriteTreeAsync(id, root, flush, cancellationToken);
+
+			/// <inheritdoc/>
+			public Task WriteTreeAsync<TNode>(RefId id, TNode root, bool flush = true, CancellationToken cancellationToken = default) where TNode : TreeNode => _inner.WriteTreeAsync<TNode>(id, root, flush, cancellationToken);
+		}
+
+		/// <summary>
+		/// Wraps an <see cref="ITreeStore"/> interface with a type for dependency injection
+		/// </summary>
+		/// <typeparam name="T"></typeparam>
+		/// <param name="store">The instance to wrap</param>
+		/// <returns>Wrapped instance of the tree store</returns>
+		public static ITreeStore<T> ForType<T>(this ITreeStore store)
+		{
+			return new TypedTreeStore<T>(store);
+		}
 	}
 }
diff --git a/Engine/Source/Programs/Shared/EpicGames.Horde/Storage/Nodes/DirectoryNode.cs b/Engine/Source/Programs/Shared/EpicGames.Horde/Storage/Nodes/DirectoryNode.cs
--- a/Engine/Source/Programs/Shared/EpicGames.Horde/Storage/Nodes/DirectoryNode.cs
+++ b/Engine/Source/Programs/Shared/EpicGames.Horde/Storage/Nodes/DirectoryNode.cs
@@ -99,6 +99,39 @@ namespace EpicGames.Horde.Storage.Nodes
 			_cachedLength = length;
 		}
 
+		/// <summary>
+		/// Appends data to this file
+		/// </summary>
+		/// <param name="data">Data to append to the file</param>
+		/// <param name="options">Options for chunking the data</param>
+		/// <param name="cancellationToken">Cancellation token for the operation</param>
+		public async Task AppendAsync(ReadOnlyMemory<byte> data, ChunkingOptions options, CancellationToken cancellationToken)
+		{
+			FileNode node = await ExpandAsync(cancellationToken);
+			Node = node.Append(data, options);
+			MarkAsDirty();
+		}
+
+		/// <summary>
+		/// Appends data to this file
+		/// </summary>
+		/// <param name="stream">Data to append to the file</param>
+		/// <param name="options">Options for chunking the data</param>
+		/// <param name="cancellationToken">Cancellation token for the operation</param>
+		public async Task AppendAsync(Stream stream, ChunkingOptions options, CancellationToken cancellationToken)
+		{
+			byte[] buffer = new byte[4 * 1024];
+			for (; ; )
+			{
+				int numBytes = await stream.ReadAsync(buffer, cancellationToken);
+				if (numBytes == 0)
+				{
+					break;
+				}
+				await AppendAsync(buffer.AsMemory(0, numBytes), options, cancellationToken);
+			}
+		}
+
 		/// <inheritdoc/>
 		protected override void OnCollapse()
 		{
@@ -218,14 +251,21 @@ namespace EpicGames.Horde.Storage.Nodes
 		/// </summary>
 		/// <param name="name">Name of the new directory</param>
 		/// <param name="flags">Flags for the new file</param>
-		/// <param name="data">Data for the file</param>
-		/// <param name="options">Options for chunking the data</param>
-		/// <param name="cancellationToken">Cancellation token for the operation</param>
 		/// <returns>The new directory object</returns>
-		public async Task<FileNode> AddFileAsync(Utf8String name, FileEntryFlags flags, ReadOnlyMemory<byte> data, ChunkingOptions options, CancellationToken cancellationToken)
+		public FileEntry AddFile(Utf8String name, FileEntryFlags flags)
 		{
-			using ReadOnlyMemoryStream stream = new ReadOnlyMemoryStream(data);
-			return await AddFileAsync(name, flags, stream, options, cancellationToken);
+			if(TryGetDirectoryEntry(name, out _))
+			{
+				throw new ArgumentException($"A directory with the name {name} already exists");
+			}
+
+			FileNode newNode = new LeafFileNode();
+
+			FileEntry entry = new FileEntry(this, name, flags, newNode);
+			_nameToFileEntry[name] = entry;
+			MarkAsDirty();
+
+			return entry;
 		}
 
 		/// <summary>
@@ -237,20 +277,11 @@ namespace EpicGames.Horde.Storage.Nodes
 		/// <param name="options">Options for chunking the data</param>
 		/// <param name="cancellationToken">Cancellation token for the operation</param>
 		/// <returns>The new directory object</returns>
-		public async Task<FileNode> AddFileAsync(Utf8String name, FileEntryFlags flags, Stream stream, ChunkingOptions options, CancellationToken cancellationToken)
+		public async Task<FileEntry> AddFileAsync(Utf8String name, FileEntryFlags flags, Stream stream, ChunkingOptions options, CancellationToken cancellationToken = default)
 		{
-			if(TryGetDirectoryEntry(name, out _))
-			{
-				throw new ArgumentException($"A directory with the name {name} already exists");
-			}
-
-			FileNode newNode = await FileNode.CreateAsync(stream, options, cancellationToken);
-
-			FileEntry entry = new FileEntry(this, name, flags, newNode);
-			_nameToFileEntry[name] = entry;
-			MarkAsDirty();
-
-			return newNode;
+			FileEntry entry = AddFile(name, flags);
+			await entry.AppendAsync(stream, options, cancellationToken);
+			return entry;
 		}
 
 		/// <summary>
@@ -258,11 +289,9 @@ namespace EpicGames.Horde.Storage.Nodes
 		/// </summary>
 		/// <param name="path">Path to the file</param>
 		/// <param name="flags">Flags for the new file</param>
-		/// <param name="stream">Stream to add to the file</param>
-		/// <param name="options">Options for chunking the data</param>
 		/// <param name="cancellationToken">Cancellation token for the operation</param>
 		/// <returns>The new directory object</returns>
-		public async ValueTask<FileNode> AddFileByPathAsync(Utf8String path, FileEntryFlags flags, Stream stream, ChunkingOptions options, CancellationToken cancellationToken)
+		public async ValueTask<FileEntry> AddFileByPathAsync(Utf8String path, FileEntryFlags flags, CancellationToken cancellationToken = default)
 		{
 			DirectoryNode directory = this;
 
@@ -279,7 +308,7 @@ namespace EpicGames.Horde.Storage.Nodes
 				{
 					if (length == remainingPath.Length)
 					{
-						return await directory.AddFileAsync(remainingPath, flags, stream, options, cancellationToken);
+						return directory.AddFile(remainingPath, flags);
 					}
 
 					byte character = remainingPath[length];
@@ -596,7 +625,7 @@ namespace EpicGames.Horde.Storage.Nodes
 			{
 				logger.LogInformation("Adding {File}", fileInfo.FullName);
 				using Stream stream = fileInfo.OpenRead();
-				FileNode fileNode = await AddFileAsync(fileInfo.Name, 0, stream, options, cancellationToken);
+				await AddFileAsync(fileInfo.Name, 0, stream, options, cancellationToken);
 			}
 		}
 
diff --git a/Engine/Source/Programs/Shared/EpicGames.Horde/Storage/Nodes/FileNode.cs b/Engine/Source/Programs/Shared/EpicGames.Horde/Storage/Nodes/FileNode.cs
--- a/Engine/Source/Programs/Shared/EpicGames.Horde/Storage/Nodes/FileNode.cs
+++ b/Engine/Source/Programs/Shared/EpicGames.Horde/Storage/Nodes/FileNode.cs
@@ -315,6 +315,11 @@ namespace EpicGames.Horde.Storage.Nodes
 		/// <inheritdoc/>
 		public override ReadOnlyMemory<byte> AppendData(ReadOnlyMemory<byte> newData, ChunkingOptions options)
 		{
+			if (_isReadOnly)
+			{
+				return newData;
+			}
+
 			// Fast path for appending data to the buffer up to the chunk window size
 			int windowSize = options.LeafOptions.MinSize;
 			if (Length < windowSize)
