Horde: Missing file.

diff --git a/Engine/Source/Programs/Horde/Horde.Build/Startup.cs b/Engine/Source/Programs/Horde/Horde.Build/Startup.cs
--- a/Engine/Source/Programs/Horde/Horde.Build/Startup.cs
+++ b/Engine/Source/Programs/Horde/Horde.Build/Startup.cs
@@ -276,6 +276,18 @@ namespace Horde.Build
 			}
 		}
 
+		static IBlobStore CreateBlobStore(IServiceProvider sp, BlobStoreOptions options)
+		{
+			IStorageBackend backend = CreateStorageBackend(sp, options);
+			return new BlobStore(backend, options);
+		}
+
+		static ITreeStore CreateTreeStore(IServiceProvider sp, TreeStoreOptions options)
+		{
+			IBlobStore store = CreateBlobStore(sp, options);
+			return new BundleStore(store, options.Bundle, sp.GetRequiredService<IMemoryCache>());
+		}
+
 		public Startup(IConfiguration configuration)
 		{
 			Configuration = configuration;
