Horde: Fix agent build errors.

diff --git a/Engine/Source/Programs/Horde/Horde.Agent/Commands/Bundles/CreateCommand.cs b/Engine/Source/Programs/Horde/Horde.Agent/Commands/Bundles/CreateCommand.cs
--- a/Engine/Source/Programs/Horde/Horde.Agent/Commands/Bundles/CreateCommand.cs
+++ b/Engine/Source/Programs/Horde/Horde.Agent/Commands/Bundles/CreateCommand.cs
@@ -108,10 +108,10 @@ namespace Horde.Agent.Commands.Bundles
 			return _blobStore;
 		}
 
-		protected BundleStore<T> CreateTreeStore<T>(ILogger logger, IMemoryCache cache) where T : TreeNode
+		protected ITreeStore CreateTreeStore(ILogger logger, IMemoryCache cache)
 		{
 			IBlobStore blobStore = CreateBlobStore(logger);
-			return new BundleStore<T>(blobStore, new BundleOptions(), cache);
+			return new BundleStore(blobStore, new BundleOptions(), cache);
 		}
 	}
 
@@ -127,7 +127,7 @@ namespace Horde.Agent.Commands.Bundles
 		public override async Task<int> ExecuteAsync(ILogger logger)
 		{
 			using IMemoryCache cache = new MemoryCache(new MemoryCacheOptions());
-			using ITreeStore<DirectoryNode> store = CreateTreeStore<DirectoryNode>(logger, cache);
+			using ITreeStore store = CreateTreeStore(logger, cache);
 
 			DirectoryNode node = new DirectoryNode();
 			await node.CopyFromDirectoryAsync(InputDir.ToDirectoryInfo(), new ChunkingOptions(), logger, CancellationToken.None);
diff --git a/Engine/Source/Programs/Horde/Horde.Agent/Commands/Bundles/ExtractCommand.cs b/Engine/Source/Programs/Horde/Horde.Agent/Commands/Bundles/ExtractCommand.cs
--- a/Engine/Source/Programs/Horde/Horde.Agent/Commands/Bundles/ExtractCommand.cs
+++ b/Engine/Source/Programs/Horde/Horde.Agent/Commands/Bundles/ExtractCommand.cs
@@ -22,9 +22,9 @@ namespace Horde.Agent.Commands.Bundles
 		public override async Task<int> ExecuteAsync(ILogger logger)
 		{
 			using IMemoryCache cache = new MemoryCache(new MemoryCacheOptions());
-			using ITreeStore<DirectoryNode> store = CreateTreeStore<DirectoryNode>(logger, cache);
+			using ITreeStore store = CreateTreeStore(logger, cache);
 
-			DirectoryNode node = await store.ReadTreeAsync(RefId);
+			DirectoryNode node = await store.ReadTreeAsync<DirectoryNode>(RefId);
 			await node.CopyToDirectoryAsync(OutputDir.ToDirectoryInfo(), logger, CancellationToken.None);
 
 			return 0;
