Horde: Prevent link errors creating spurious issues when XGE fails.

diff --git a/Engine/Source/Programs/Horde/Horde.Build/Issues/Handlers/SymbolIssueHandler.cs b/Engine/Source/Programs/Horde/Horde.Build/Issues/Handlers/SymbolIssueHandler.cs
--- a/Engine/Source/Programs/Horde/Horde.Build/Issues/Handlers/SymbolIssueHandler.cs
+++ b/Engine/Source/Programs/Horde/Horde.Build/Issues/Handlers/SymbolIssueHandler.cs
@@ -47,7 +47,7 @@ namespace Horde.Build.Issues.Handlers
 		/// <returns></returns>
 		static bool IsMaskedEventId(EventId eventId)
 		{
-			return eventId == KnownLogEvents.ExitCode;
+			return eventId == KnownLogEvents.ExitCode || eventId == KnownLogEvents.Systemic_Xge_BuildFailed;
 		}
 
 		/// <summary>
