Horde: Fix tests.

diff --git a/Engine/Source/Programs/Horde/Horde.Build.Tests/TestSetup.cs b/Engine/Source/Programs/Horde/Horde.Build.Tests/TestSetup.cs
--- a/Engine/Source/Programs/Horde/Horde.Build.Tests/TestSetup.cs
+++ b/Engine/Source/Programs/Horde/Horde.Build.Tests/TestSetup.cs
@@ -57,6 +57,7 @@ using Horde.Build.Logs.Storage;
 using Horde.Build.Tasks;
 using Horde.Build.Auditing;
 using Horde.Build.Storage.Backends;
+using EpicGames.Horde.Storage.Bundles;
 
 namespace Horde.Build.Tests
 {
@@ -223,6 +224,12 @@ namespace Horde.Build.Tests
 			services.AddSingleton<IStorageBackend<ArtifactCollection>>(sp => new TransientStorageBackend().ForType<ArtifactCollection>());
 			services.AddSingleton<IStorageBackend<BasicStorageClient>>(sp => new TransientStorageBackend().ForType<BasicStorageClient>());
 
+			services.AddSingleton<ITreeStore<CommitService>>(sp =>
+			{
+				IBlobStore blobStore = new BlobStore(new TransientStorageBackend(), new BlobStoreOptions());
+				return new BundleStore(blobStore, new BundleOptions(), sp.GetRequiredService<IMemoryCache>()).ForType<CommitService>();
+			});
+
 			services.AddSingleton<IStorageClient, BasicStorageClient>();
 
 			services.AddSingleton<ISingletonDocument<AgentSoftwareChannels>>(new SingletonDocumentStub<AgentSoftwareChannels>());
