Horde: Include the preflight changelist description in Slack notifications (with the appropriate #preflight tag).

diff --git a/Engine/Source/Programs/Horde/Horde.Build/Notifications/Sinks/SlackNotificationSink.cs b/Engine/Source/Programs/Horde/Horde.Build/Notifications/Sinks/SlackNotificationSink.cs
--- a/Engine/Source/Programs/Horde/Horde.Build/Notifications/Sinks/SlackNotificationSink.cs
+++ b/Engine/Source/Programs/Horde/Horde.Build/Notifications/Sinks/SlackNotificationSink.cs
@@ -436,6 +436,17 @@ namespace Horde.Build.Notifications.Sinks
 			attachment.FallbackText = $"{stream.Name} - {GetJobChangeText(job)} - {job.Name} - {jobOutcome}";
 			attachment.Color = outcomeColor;
 			attachment.Blocks.Add(new SectionBlock($"*<{jobLink}|{stream.Name} - {GetJobChangeText(job)} - {job.Name}>*"));
+
+			if (!String.IsNullOrEmpty(job.PreflightDescription))
+			{
+				string description = job.PreflightDescription.Trim();
+				if (jobOutcome == JobStepOutcome.Success)
+				{
+					description += $"\n#preflight {job.Id}";
+				}
+				attachment.Blocks.Add(new SectionBlock($"```{description}```"));
+			}
+
 			if (jobOutcome == JobStepOutcome.Success)
 			{
 				attachment.Blocks.Add(new SectionBlock($"*Job Succeeded*"));
