Horde Storage - Added logging when getting unhandled exception during inlined-payload, to better debug rare issue (seems to be related to content id missing)

diff --git a/Engine/Source/Programs/Horde/HordeStorage/Horde.Storage/Controllers/ReferencesController.cs b/Engine/Source/Programs/Horde/HordeStorage/Horde.Storage/Controllers/ReferencesController.cs
--- a/Engine/Source/Programs/Horde/HordeStorage/Horde.Storage/Controllers/ReferencesController.cs
+++ b/Engine/Source/Programs/Horde/HordeStorage/Horde.Storage/Controllers/ReferencesController.cs
@@ -349,6 +349,10 @@ namespace Horde.Storage.Controllers
                                     _logger.Error("Failed to find blob {Blob} in namespace {Namespace} but it was expected to exist after reference resolving", attachmentToSend, ns);
                                     throw;
                                 }
+                                catch (Exception ex)
+                                {
+                                    _logger.Error(ex, "Unknown exception encountered while writing body for jupiter inlined payload.");
+                                }
                                 return new EmptyResult();
                             }
                             else if (referencedBlobs.Count == 0)
