Horde Storage - Fixed race when writing cb package response, we can not write two attachments at the same time to the stream. Now we write them in order.

diff --git a/Engine/Source/Programs/Horde/HordeStorage/Horde.Storage/Controllers/ReferencesController.cs b/Engine/Source/Programs/Horde/HordeStorage/Horde.Storage/Controllers/ReferencesController.cs
--- a/Engine/Source/Programs/Horde/HordeStorage/Horde.Storage/Controllers/ReferencesController.cs
+++ b/Engine/Source/Programs/Horde/HordeStorage/Horde.Storage/Controllers/ReferencesController.cs
@@ -227,62 +227,54 @@ namespace Horde.Storage.Controllers
                         CbPackageBuilder writer = new CbPackageBuilder();
                         await writer.AddAttachment(objectRecord.BlobIdentifier.AsIoHash(), CbPackageAttachmentFlags.IsObject, blobMemory);
 
-                        List<Task> fetchBlobTasks = new List<Task>();
                         await foreach (Attachment attachment in attachments)
                         {
-                            Task blobTask = Task.Run(async () =>
-                            {
-                                IoHash attachmentHash = attachment.AsIoHash();
-                                BlobContents contents;
-                                CbPackageAttachmentFlags flags = 0;
+                            IoHash attachmentHash = attachment.AsIoHash();
+                            CbPackageAttachmentFlags flags = 0;
 
-                                try
+                            try
+                            {
+                                BlobContents attachmentContents;
+                                if (attachment is BlobAttachment blobAttachment)
+                                {
+                                    BlobIdentifier referencedBlob = blobAttachment.Identifier;
+                                    attachmentContents = await _blobStore.GetObject(ns, referencedBlob);
+                                }
+                                else if (attachment is ObjectAttachment objectAttachment)
+                                {
+                                    flags &= CbPackageAttachmentFlags.IsObject;
+                                    BlobIdentifier referencedBlob = objectAttachment.Identifier;
+                                    attachmentContents = await _blobStore.GetObject(ns, referencedBlob);
+                                }
+                                else if (attachment is ContentIdAttachment contentIdAttachment)
                                 {
-                                    if (attachment is BlobAttachment blobAttachment)
-                                    {
-                                        BlobIdentifier referencedBlob = blobAttachment.Identifier;
-                                        contents = await _blobStore.GetObject(ns, referencedBlob);
-                                    }
-                                    else if (attachment is ObjectAttachment objectAttachment)
-                                    {
-                                        flags &= CbPackageAttachmentFlags.IsObject;
-                                        BlobIdentifier referencedBlob = objectAttachment.Identifier;
-                                        contents = await _blobStore.GetObject(ns, referencedBlob);
-                                    }
-                                    else if (attachment is ContentIdAttachment contentIdAttachment)
-                                    {
 
-                                        ContentId contentId = contentIdAttachment.Identifier;
-                                        (contents, string mime) = await _blobStore.GetCompressedObject(ns, contentId, _contentIdStore);
-                                        if (mime == CustomMediaTypeNames.UnrealCompressedBuffer)
-                                        {
-                                            flags &= CbPackageAttachmentFlags.IsCompressed;
-                                        }
-                                        else
-                                        {
-                                            // this resolved to a uncompressed blob, the content id existed the the compressed blob didn't
-                                            // so resetting flags to indicate this.
-                                            flags = 0;
-                                        }
+                                    ContentId contentId = contentIdAttachment.Identifier;
+                                    (attachmentContents, string mime) = await _blobStore.GetCompressedObject(ns, contentId, _contentIdStore);
+                                    if (mime == CustomMediaTypeNames.UnrealCompressedBuffer)
+                                    {
+                                        flags &= CbPackageAttachmentFlags.IsCompressed;
                                     }
                                     else
                                     {
-                                        throw new NotSupportedException($"Unknown attachment type {attachment.GetType()}");
+                                        // this resolved to a uncompressed blob, the content id existed the the compressed blob didn't
+                                        // so resetting flags to indicate this.
+                                        flags = 0;
                                     }
                                 }
-                                catch (Exception e)
+                                else
                                 {
-                                    (CbObject errorObject, HttpStatusCode _) = ToErrorResult(e);
-                                    await writer.AddAttachment(attachmentHash, CbPackageAttachmentFlags.IsError | CbPackageAttachmentFlags.IsObject, errorObject.GetView().ToArray());
-                                    return;
+                                    throw new NotSupportedException($"Unknown attachment type {attachment.GetType()}");
                                 }
-                                await writer.AddAttachment(attachmentHash, flags, contents.Stream, (ulong)contents.Length);
-                            });
-
-                            fetchBlobTasks.Add(blobTask);
 
+                                await writer.AddAttachment(attachmentHash, flags, attachmentContents.Stream, (ulong)attachmentContents.Length);
+                            }
+                            catch (Exception e)
+                            {
+                                (CbObject errorObject, HttpStatusCode _) = ToErrorResult(e);
+                                await writer.AddAttachment(attachmentHash, CbPackageAttachmentFlags.IsError | CbPackageAttachmentFlags.IsObject, errorObject.GetView().ToArray());
+                            }
                         }
-                        await Task.WhenAll(fetchBlobTasks);
                         await using BlobContents contents = new BlobContents(writer.ToByteArray());
                         await WriteBody(contents, CustomMediaTypeNames.UnrealCompactBinaryPackage);
                         break;
