Horde Storage - Do not raise an error when blob is missing and resolving inline payload. We already return a 404, but made that a bit clearer and avoided the errors as it is working as intended.

diff --git a/Engine/Source/Programs/Horde/HordeStorage/Horde.Storage/Controllers/ReferencesController.cs b/Engine/Source/Programs/Horde/HordeStorage/Horde.Storage/Controllers/ReferencesController.cs
--- a/Engine/Source/Programs/Horde/HordeStorage/Horde.Storage/Controllers/ReferencesController.cs
+++ b/Engine/Source/Programs/Horde/HordeStorage/Horde.Storage/Controllers/ReferencesController.cs
@@ -338,12 +338,12 @@ namespace Horde.Storage.Controllers
                                 }
                                 catch (BlobNotFoundException)
                                 {
-                                    _logger.Error("Failed to find blob {Blob} in namespace {Namespace} but it was expected to exist after reference resolving", attachmentToSend, ns);
-                                    throw;
+                                    return NotFound(new ProblemDetails {Title = $"Object {bucket} {key} in namespace {ns} was missing blob {attachmentToSend}"});
                                 }
                                 catch (Exception ex)
                                 {
                                     _logger.Error(ex, "Unknown exception encountered while writing body for jupiter inlined payload.");
+                                    throw;
                                 }
                                 return new EmptyResult();
                             }
