Horde Storage - Disabling start at random offset for ref cleanup - this runs in a single region and one a single worker so this has no function other then delaying the start of the gc. The purpose of this option is to spread out load on db and other services so that not all background tasks hit them at the same time, which does not apply here.

diff --git a/Engine/Source/Programs/Horde/HordeStorage/Horde.Storage/Implementation/GC/RefCleanupService.cs b/Engine/Source/Programs/Horde/HordeStorage/Horde.Storage/Implementation/GC/RefCleanupService.cs
--- a/Engine/Source/Programs/Horde/HordeStorage/Horde.Storage/Implementation/GC/RefCleanupService.cs
+++ b/Engine/Source/Programs/Horde/HordeStorage/Horde.Storage/Implementation/GC/RefCleanupService.cs
@@ -32,7 +32,7 @@ namespace Horde.Storage.Implementation
         private volatile bool _alreadyPolling;
         private readonly ILogger _logger = Log.ForContext<RefCleanupService>();
 
-        public RefCleanupService(IOptionsMonitor<GCSettings> settings, IRefsStore store, IRefCleanup refCleanup, ILeaderElection leaderElection, IReferencesStore referencesStore) : base(serviceName: nameof(RefCleanupService), settings.CurrentValue.RefCleanupPollFrequency, new RefCleanupState(store, refCleanup))
+        public RefCleanupService(IOptionsMonitor<GCSettings> settings, IRefsStore store, IRefCleanup refCleanup, ILeaderElection leaderElection, IReferencesStore referencesStore) : base(serviceName: nameof(RefCleanupService), settings.CurrentValue.RefCleanupPollFrequency, new RefCleanupState(store, refCleanup), startAtRandomTime: false)
         {
             _settings = settings;
             _leaderElection = leaderElection;
