Horde: Add a setting to allow overriding the AWS region from a config file, and default to using settings from AWS environment variables if not overridden.

diff --git a/Engine/Source/Programs/Horde/Horde.Build/ServerSettings.cs b/Engine/Source/Programs/Horde/Horde.Build/ServerSettings.cs
--- a/Engine/Source/Programs/Horde/Horde.Build/ServerSettings.cs
+++ b/Engine/Source/Programs/Horde/Horde.Build/ServerSettings.cs
@@ -76,6 +76,9 @@ namespace Horde.Build
 		/// <inheritdoc/>
 		public string? AwsProfile { get; set; }
 
+		/// <inheritdoc/>
+		public string? AwsRegion { get; set; }
+
 		/// <inheritdoc/>
 		public string? RelayServer { get; set; }
 
diff --git a/Engine/Source/Programs/Horde/Horde.Build/Startup.cs b/Engine/Source/Programs/Horde/Horde.Build/Startup.cs
--- a/Engine/Source/Programs/Horde/Horde.Build/Startup.cs
+++ b/Engine/Source/Programs/Horde/Horde.Build/Startup.cs
@@ -965,6 +965,7 @@ namespace Horde.Build
 		public static ServiceProvider CreateServiceProvider(IConfiguration configuration, ILoggerProvider loggerProvider)
 		{
 			IServiceCollection services = new ServiceCollection();
+			services.AddSingleton(configuration);
 			services.AddSingleton(loggerProvider);
 			AddServices(services, configuration);
 			return services.BuildServiceProvider();
diff --git a/Engine/Source/Programs/Horde/Horde.Build/Storage/Backends/AwsStorageBackend.cs b/Engine/Source/Programs/Horde/Horde.Build/Storage/Backends/AwsStorageBackend.cs
--- a/Engine/Source/Programs/Horde/Horde.Build/Storage/Backends/AwsStorageBackend.cs
+++ b/Engine/Source/Programs/Horde/Horde.Build/Storage/Backends/AwsStorageBackend.cs
@@ -7,6 +7,7 @@ using System.IO;
 using System.Text;
 using System.Threading;
 using System.Threading.Tasks;
+using Amazon;
 using Amazon.Extensions.NETCore.Setup;
 using Amazon.Runtime;
 using Amazon.S3;
@@ -39,24 +40,24 @@ namespace Horde.Build.Storage.Backends
 	public enum AwsCredentialsType
 	{
 		/// <summary>
-		/// Assume a particular role. Should specify ARN in <see cref="IAwsStorageOptions.AwsRole"/>
+		/// Use default credentials from the AWS SDK
 		/// </summary>
-		AssumeRole,
+		Default,
 
 		/// <summary>
-		/// Assume a particular role using the current environment variables.
+		/// Read credentials from the <see cref="IAwsStorageOptions.AwsProfile"/> profile in the AWS config file
 		/// </summary>
-		AssumeRoleWebIdentity,
+		Profile,
 
 		/// <summary>
-		/// Use default credentials from the AWS SDK
+		/// Assume a particular role. Should specify ARN in <see cref="IAwsStorageOptions.AwsRole"/>
 		/// </summary>
-		Fallback,
+		AssumeRole,
 
 		/// <summary>
-		/// Read credentials from a profile in the AWS config file
+		/// Assume a particular role using the current environment variables.
 		/// </summary>
-		SharedCredentials,
+		AssumeRoleWebIdentity,
 	}
 
 	/// <summary>
@@ -88,6 +89,11 @@ namespace Horde.Build.Storage.Backends
 		/// The AWS profile to read credentials form
 		/// </summary>
 		public string? AwsProfile { get; }
+
+		/// <summary>
+		/// Region to connect to
+		/// </summary>
+		public string? AwsRegion { get; }
 	}
 
 	/// <summary>
@@ -142,23 +148,18 @@ namespace Horde.Build.Storage.Backends
 		static AWSOptions GetAwsOptions(IConfiguration configuration, IAwsStorageOptions options)
 		{
 			AWSOptions awsOptions = configuration.GetAWSOptions();
+			if (options.AwsRegion != null)
+			{
+				awsOptions.Region = RegionEndpoint.GetBySystemName(options.AwsRegion);
+			}
+
 			switch (options.AwsCredentials)
 			{
-				case AwsCredentialsType.AssumeRole:
-					if(options.AwsRole == null)
-					{
-						throw new AwsException($"Missing {nameof(IAwsStorageOptions.AwsRole)} setting for configuring {nameof(AwsStorageBackend)}", null);
-					}
-					awsOptions.Credentials = new AssumeRoleAWSCredentials(FallbackCredentialsFactory.GetCredentials(), options.AwsRole, "Horde");
-					break;
-				case AwsCredentialsType.AssumeRoleWebIdentity:
-					awsOptions.Credentials = AssumeRoleWithWebIdentityCredentials.FromEnvironmentVariables();
-					break;
-				case AwsCredentialsType.Fallback:
+				case AwsCredentialsType.Default:
 					// Using the fallback credentials from the AWS SDK, it will pick up credentials through a number of default mechanisms.
 					awsOptions.Credentials = FallbackCredentialsFactory.GetCredentials();
 					break;
-				case AwsCredentialsType.SharedCredentials:
+				case AwsCredentialsType.Profile:
 					if (options.AwsProfile == null)
 					{
 						throw new AwsException($"Missing {nameof(IAwsStorageOptions.AwsProfile)} setting for configuring {nameof(AwsStorageBackend)}", null);
@@ -167,6 +168,16 @@ namespace Horde.Build.Storage.Backends
 					(string accessKey, string secretAccessKey, string secretToken) = AwsHelper.ReadAwsCredentials(options.AwsProfile);
 					awsOptions.Credentials = new Amazon.SecurityToken.Model.Credentials(accessKey, secretAccessKey, secretToken, DateTime.Now + TimeSpan.FromHours(12));
 					break;
+				case AwsCredentialsType.AssumeRole:
+					if(options.AwsRole == null)
+					{
+						throw new AwsException($"Missing {nameof(IAwsStorageOptions.AwsRole)} setting for configuring {nameof(AwsStorageBackend)}", null);
+					}
+					awsOptions.Credentials = new AssumeRoleAWSCredentials(FallbackCredentialsFactory.GetCredentials(), options.AwsRole, "Horde");
+					break;
+				case AwsCredentialsType.AssumeRoleWebIdentity:
+					awsOptions.Credentials = AssumeRoleWithWebIdentityCredentials.FromEnvironmentVariables();
+					break;
 			}
 			return awsOptions;
 		}
