Horde: Strike through issue threads once they've been resolved for all workflows. Strike through titles in summary when resolved for the stream/template they apply to.

diff --git a/Engine/Source/Programs/Horde/Horde.Build/Notifications/Sinks/SlackNotificationSink.cs b/Engine/Source/Programs/Horde/Horde.Build/Notifications/Sinks/SlackNotificationSink.cs
--- a/Engine/Source/Programs/Horde/Horde.Build/Notifications/Sinks/SlackNotificationSink.cs
+++ b/Engine/Source/Programs/Horde/Horde.Build/Notifications/Sinks/SlackNotificationSink.cs
@@ -439,7 +439,7 @@ namespace Horde.Build.Notifications.Sinks
 
 			if (!String.IsNullOrEmpty(job.PreflightDescription))
 			{
-				string description = job.PreflightDescription.Trim();
+				string description = job.PreflightDescription.TrimEnd();
 				if (jobOutcome == JobStepOutcome.Success)
 				{
 					description += $"\n#preflight {job.Id}";
@@ -784,7 +784,7 @@ namespace Horde.Build.Notifications.Sinks
 
 				string issueSummary = issue.UserSummary ?? issue.Summary;
 				string text = $"{workflow.TriagePrefix}{GetPrefixForSeverity(issue.Severity)}*New Issue <{issueUrl}|{issue.Id}>*: {issueSummary}{workflow.TriageSuffix}";
-				if (!spans.Any(x => x.NextSuccess == null))
+				if (!spans.Any(x => x.NextSuccess == null && x.LastFailure.Annotations.WorkflowId != null)) // Thread may be shared by multiple workflows
 				{
 					text = $"~{text}~";
 				}
@@ -1288,6 +1288,12 @@ namespace Horde.Build.Notifications.Sinks
 
 		class ReportBlock
 		{
+			[JsonPropertyName("h")]
+			public bool TemplateHeader { get; set; }
+
+			[JsonPropertyName("s")]
+			public StreamId StreamId { get; set; }
+
 			[JsonPropertyName("t")]
 			public TemplateRefId TemplateId { get; set; }
 
@@ -1318,7 +1324,7 @@ namespace Horde.Build.Notifications.Sinks
 				ReportState state = new ReportState();
 				state.Time = report.Time.UtcDateTime;
 
-				List<List<(IIssue, IIssueSpan?)>> issuesByBlock = new List<List<(IIssue, IIssueSpan?)>>();
+				List<List<(IIssue, IIssueSpan?, bool)>> issuesByBlock = new List<List<(IIssue, IIssueSpan?, bool)>>();
 				if (report.GroupByTemplate)
 				{
 					Dictionary<int, IIssue> issueIdToInfo = new Dictionary<int, IIssue>();
@@ -1334,6 +1340,10 @@ namespace Horde.Build.Notifications.Sinks
 						{
 							continue;
 						}
+						if (!IsIssueOpenForWorkflow(report.Stream.Id, group.Key, report.WorkflowId, group))
+						{
+							continue;
+						}
 
 						Dictionary<IIssue, IIssueSpan> pairs = new Dictionary<IIssue, IIssueSpan>();
 						foreach (IIssueSpan span in group)
@@ -1346,16 +1356,18 @@ namespace Horde.Build.Notifications.Sinks
 
 						if (pairs.Count > 0)
 						{
-							TemplateRefId templateId = group.Key;
+							bool templateHeader = true;
 							foreach (IReadOnlyList<KeyValuePair<IIssue, IIssueSpan>> batch in pairs.OrderBy(x => x.Key.Id).Batch(MaxIssuesPerMessage))
 							{
 								ReportBlock block = new ReportBlock();
-								block.TemplateId = templateId;
+								block.TemplateHeader = templateHeader;
+								block.StreamId = report.Stream.Id;
+								block.TemplateId = group.Key;
 								block.IssueIds.AddRange(batch.Select(x => x.Key.Id));
 								state.Blocks.Add(block);
 
-								issuesByBlock.Add(batch.Select(x => (x.Key, (IIssueSpan?)x.Value)).ToList());
-								templateId = TemplateRefId.Empty;
+								issuesByBlock.Add(batch.Select(x => (x.Key, (IIssueSpan?)x.Value, true)).ToList());
+								templateHeader = false;
 							}
 						}
 					}
@@ -1365,10 +1377,11 @@ namespace Horde.Build.Notifications.Sinks
 					foreach (IReadOnlyList<IIssue> batch in report.Issues.OrderByDescending(x => x.Id).Batch(MaxIssuesPerMessage))
 					{
 						ReportBlock block = new ReportBlock();
+						block.StreamId = report.Stream.Id;
 						block.IssueIds.AddRange(batch.Select(x => x.Id));
 						state.Blocks.Add(block);
 
-						issuesByBlock.Add(batch.Select(x => (x, report.IssueSpans.FirstOrDefault(y => y.IssueId == x.Id))).ToList());
+						issuesByBlock.Add(batch.Select(x => (x, report.IssueSpans.FirstOrDefault(y => y.IssueId == x.Id), true)).ToList());
 					}
 				}
 
@@ -1408,6 +1421,21 @@ namespace Horde.Build.Notifications.Sinks
 			}
 		}
 
+		static bool IsIssueOpenForWorkflow(StreamId streamId, TemplateRefId templateId, WorkflowId workflowId, IEnumerable<IIssueSpan> spans)
+		{
+			foreach (IIssueSpan span in spans)
+			{
+				if (span.NextSuccess == null && span.LastFailure.Annotations.WorkflowId == workflowId)
+				{
+					if ((streamId.IsEmpty || span.StreamId == streamId) && (templateId.IsEmpty || span.TemplateRefId == templateId))
+					{
+						return true;
+					}
+				}
+			}
+			return false;
+		}
+
 		async Task UpdateReportsAsync(IIssue issue, IReadOnlyList<IIssueSpan> spans)
 		{
 			_logger.LogInformation("Checking for report updates to issue {IssueId}", issue.Id);
@@ -1470,7 +1498,7 @@ namespace Horde.Build.Notifications.Sinks
 						string blockEventId = GetReportBlockEventId(messageState.Ts, idx);
 						_logger.LogInformation("Updating report block {EventId}", blockEventId);
 
-						List<(IIssue, IIssueSpan?)> pairs = new List<(IIssue, IIssueSpan?)>();
+						List<(IIssue, IIssueSpan?, bool)> issues = new List<(IIssue, IIssueSpan?, bool)>();
 						foreach (int issueId in block.IssueIds)
 						{
 							IIssueDetails? details = await _issueService.GetIssueDetailsAsync(issueId);
@@ -1481,17 +1509,19 @@ namespace Horde.Build.Notifications.Sinks
 								{
 									otherSpan = details.Spans[0];
 								}
-								pairs.Add((details.Issue, otherSpan));
+
+								bool open = IsIssueOpenForWorkflow(stream.Id, block.TemplateId, workflowId.Value, details.Spans);
+								issues.Add((details.Issue, otherSpan, open));
 							}
 						}
 
-						await UpdateReportBlockAsync(workflow.ReportChannel, blockEventId, state.Time, stream, block.TemplateId, pairs);
+						await UpdateReportBlockAsync(workflow.ReportChannel, blockEventId, state.Time, stream, block.TemplateId, issues);
 					}
 				}
 			}
 		}
 
-		async Task UpdateReportBlockAsync(string channel, string eventId, DateTime reportTime, IStream stream, TemplateRefId? templateId, List<(IIssue, IIssueSpan?)> pairs)
+		async Task UpdateReportBlockAsync(string channel, string eventId, DateTime reportTime, IStream stream, TemplateRefId? templateId, List<(IIssue, IIssueSpan?, bool)> issues)
 		{
 			StringBuilder body = new StringBuilder();
 
@@ -1509,7 +1539,7 @@ namespace Horde.Build.Notifications.Sinks
 				}
 			}
 
-			foreach ((IIssue issue, IIssueSpan? span) in pairs)
+			foreach ((IIssue issue, IIssueSpan? span, bool open) in issues)
 			{
 				if (body.Length > 0)
 				{
@@ -1517,7 +1547,7 @@ namespace Horde.Build.Notifications.Sinks
 				}
 
 				string text = await FormatIssueAsync(issue, span, channel, reportTime);
-				if (issue.VerifiedAt != null)
+				if (!open)
 				{
 					text = $"~{text}~";
 				}
