Horde: Fix template header appearing ahead of each output block.

diff --git a/Engine/Source/Programs/Horde/Horde.Build/Notifications/Sinks/SlackNotificationSink.cs b/Engine/Source/Programs/Horde/Horde.Build/Notifications/Sinks/SlackNotificationSink.cs
--- a/Engine/Source/Programs/Horde/Horde.Build/Notifications/Sinks/SlackNotificationSink.cs
+++ b/Engine/Source/Programs/Horde/Horde.Build/Notifications/Sinks/SlackNotificationSink.cs
@@ -1408,7 +1408,7 @@ namespace Horde.Build.Notifications.Sinks
 					for (int idx = 0; idx < state.Blocks.Count; idx++)
 					{
 						string blockEventId = GetReportBlockEventId(response.Ts, idx);
-						await UpdateReportBlockAsync(report.Channel, blockEventId, report.Time.UtcDateTime, report.Stream, state.Blocks[idx].TemplateId, issuesByBlock[idx]);
+						await UpdateReportBlockAsync(report.Channel, blockEventId, report.Time.UtcDateTime, report.Stream, state.Blocks[idx].TemplateId, issuesByBlock[idx], state.Blocks[idx].TemplateHeader);
 					}
 
 					if (report.WorkflowStats.NumSteps > 0)
@@ -1515,25 +1515,25 @@ namespace Horde.Build.Notifications.Sinks
 							}
 						}
 
-						await UpdateReportBlockAsync(workflow.ReportChannel, blockEventId, state.Time, stream, block.TemplateId, issues);
+						await UpdateReportBlockAsync(workflow.ReportChannel, blockEventId, state.Time, stream, block.TemplateId, issues, block.TemplateHeader);
 					}
 				}
 			}
 		}
 
-		async Task UpdateReportBlockAsync(string channel, string eventId, DateTime reportTime, IStream stream, TemplateRefId? templateId, List<(IIssue, IIssueSpan?, bool)> issues)
+		async Task UpdateReportBlockAsync(string channel, string eventId, DateTime reportTime, IStream stream, TemplateRefId templateId, List<(IIssue, IIssueSpan?, bool)> issues, bool templateHeader)
 		{
 			StringBuilder body = new StringBuilder();
 
-			if (templateId != null)
+			if (templateHeader && !templateId.IsEmpty)
 			{
 				TemplateRefConfig? templateConfig;
-				if (stream.Config.TryGetTemplate(templateId.Value, out templateConfig))
+				if (stream.Config.TryGetTemplate(templateId, out templateConfig))
 				{
-					CreateJobsTabRequest? tab = stream.Config.Tabs.OfType<CreateJobsTabRequest>().FirstOrDefault(x => x.Templates != null && x.Templates.Contains(templateId.Value));
+					CreateJobsTabRequest? tab = stream.Config.Tabs.OfType<CreateJobsTabRequest>().FirstOrDefault(x => x.Templates != null && x.Templates.Contains(templateId));
 					if (tab != null)
 					{
-						Uri templateUrl = new Uri(_settings.DashboardUrl, $"stream/{stream.Id}?tab={tab.Title}&template={templateId.Value}");
+						Uri templateUrl = new Uri(_settings.DashboardUrl, $"stream/{stream.Id}?tab={tab.Title}&template={templateId}");
 						body.Append($"*<{templateUrl}|{templateConfig.Name}>*:");
 					}
 				}
