Horde: Include the fix CL for any fixed issues in Slack reports.

diff --git a/Engine/Source/Programs/Horde/Horde.Build/Notifications/Sinks/SlackNotificationSink.cs b/Engine/Source/Programs/Horde/Horde.Build/Notifications/Sinks/SlackNotificationSink.cs
--- a/Engine/Source/Programs/Horde/Horde.Build/Notifications/Sinks/SlackNotificationSink.cs
+++ b/Engine/Source/Programs/Horde/Horde.Build/Notifications/Sinks/SlackNotificationSink.cs
@@ -1572,17 +1572,33 @@ namespace Horde.Build.Notifications.Sinks
 				issueUrl = new Uri(issueUrl, $"job/{lastFailure.JobId}?step={lastFailure.StepId}&issue={issue.Id}");
 			}
 
-			string status = "*Unassigned*";
+			IUser? owner = null;
 			if (issue.OwnerId != null)
 			{
-				IUser? user = await _userCollection.GetCachedUserAsync(issue.OwnerId.Value);
-				if (user == null)
+				owner = await _userCollection.GetCachedUserAsync(issue.OwnerId.Value);
+			}
+
+			string status = "*Unassigned*";
+			if (issue.FixChange != null)
+			{
+				if (owner == null)
+				{
+					status = $"Fixed in CL {issue.FixChange.Value}";
+				}
+				else
+				{
+					status = $"Fixed in CL {issue.FixChange.Value} by {owner.Name}";
+				}
+			}
+			else if (issue.OwnerId != null)
+			{
+				if (owner == null)
 				{
 					status = $"Assigned to user {issue.OwnerId.Value}";
 				}
 				else
 				{
-					status = $"Assigned to {user.Name}";
+					status = $"Assigned to {owner.Name}";
 				}
 				if (issue.AcknowledgedAt == null)
 				{
