Horde: Put the warning/error icon before the message prefix on issue threads.

diff --git a/Engine/Source/Programs/Horde/Horde.Build/Notifications/Sinks/SlackNotificationSink.cs b/Engine/Source/Programs/Horde/Horde.Build/Notifications/Sinks/SlackNotificationSink.cs
--- a/Engine/Source/Programs/Horde/Horde.Build/Notifications/Sinks/SlackNotificationSink.cs
+++ b/Engine/Source/Programs/Horde/Horde.Build/Notifications/Sinks/SlackNotificationSink.cs
@@ -783,7 +783,7 @@ namespace Horde.Build.Notifications.Sinks
 				}
 
 				string issueSummary = issue.UserSummary ?? issue.Summary;
-				string text = $"{workflow.TriagePrefix}{GetPrefixForSeverity(issue.Severity)}*New Issue <{issueUrl}|{issue.Id}>*: {issueSummary}{workflow.TriageSuffix}";
+				string text = $"{GetPrefixForSeverity(issue.Severity)}{workflow.TriagePrefix}*New Issue <{issueUrl}|{issue.Id}>*: {issueSummary}{workflow.TriageSuffix}";
 				if (!spans.Any(x => x.NextSuccess == null && x.LastFailure.Annotations.WorkflowId != null)) // Thread may be shared by multiple workflows
 				{
 					text = $"~{text}~";
