Horde: Remove the thread emoji from new issue messages. We now include the status emoji instead.

diff --git a/Engine/Source/Programs/Horde/Horde.Build/Server/Config.cs b/Engine/Source/Programs/Horde/Horde.Build/Server/Config.cs
--- a/Engine/Source/Programs/Horde/Horde.Build/Server/Config.cs
+++ b/Engine/Source/Programs/Horde/Horde.Build/Server/Config.cs
@@ -362,7 +362,7 @@ namespace Horde.Build.Server
 		/// <summary>
 		/// Prefix for all triage messages
 		/// </summary>
-		public string? TriagePrefix { get; set; } = ":thread: ";
+		public string? TriagePrefix { get; set; }
 
 		/// <summary>
 		/// Suffix for all triage messages
