Horde: Move stats for number of succeeded steps into the summary.

diff --git a/Engine/Source/Programs/Horde/Horde.Build/Notifications/Sinks/SlackNotificationSink.cs b/Engine/Source/Programs/Horde/Horde.Build/Notifications/Sinks/SlackNotificationSink.cs
--- a/Engine/Source/Programs/Horde/Horde.Build/Notifications/Sinks/SlackNotificationSink.cs
+++ b/Engine/Source/Programs/Horde/Horde.Build/Notifications/Sinks/SlackNotificationSink.cs
@@ -1388,15 +1388,22 @@ namespace Horde.Build.Notifications.Sinks
 				List<BlockBase> blocks = new List<BlockBase>();
 				blocks.Add(new HeaderBlock(AddEnvironmentAnnotation($"{report.Stream.Name}: {report.Time:d}")));
 
+				StringBuilder header = new StringBuilder();
 				if (report.Issues.Count == 0)
 				{
-					blocks.Add(new SectionBlock(":tick: No issues open."));
+					header.Append(":tick: No issues open.");
 				}
 				else
 				{
 					TimeSpan averageAge = TimeSpan.FromHours(report.Issues.Select(x => (report.Time - x.CreatedAt).TotalHours).Average());
-					blocks.Add(new SectionBlock($"*{report.Issues.Count} unique issues* currently open (average age {FormatReadableTimeSpan(averageAge)})."));
+					header.Append($"*{report.Issues.Count} unique issues* currently open (average age {FormatReadableTimeSpan(averageAge)}).");
 				}
+				if (report.WorkflowStats.NumSteps > 0)
+				{
+					double totalPct = (report.WorkflowStats.NumPassingSteps * 100.0) / report.WorkflowStats.NumSteps;
+					header.Append($" {report.WorkflowStats.NumPassingSteps:n0} of {report.WorkflowStats.NumSteps:n0} (*{totalPct:0.0}%*) build steps succeeded since last status update.");
+				}
+				blocks.Add(new SectionBlock(header.ToString()));
 
 				PostMessageResponse? response = await SendMessageAsync(report.Channel, blocks: blocks.ToArray(), withEnvironment: false);
 				if (response != null && response.Ts != null)
@@ -1410,13 +1417,6 @@ namespace Horde.Build.Notifications.Sinks
 						string blockEventId = GetReportBlockEventId(response.Ts, idx);
 						await UpdateReportBlockAsync(report.Channel, blockEventId, report.Time.UtcDateTime, report.Stream, state.Blocks[idx].TemplateId, issuesByBlock[idx], state.Blocks[idx].TemplateHeader);
 					}
-
-					if (report.WorkflowStats.NumSteps > 0)
-					{
-						double totalPct = (report.WorkflowStats.NumPassingSteps * 100.0) / report.WorkflowStats.NumSteps;
-						string footer = $"{report.WorkflowStats.NumPassingSteps:n0} of {report.WorkflowStats.NumSteps:n0} (*{totalPct:0.0}%*) build steps succeeded since last status update.";
-						await SendMessageAsync(report.Channel, text: footer, withEnvironment: false);
-					}
 				}
 			}
 		}
