Horde: Move report summary out of section block, to prevent narrow formatting.

diff --git a/Engine/Source/Programs/Horde/Horde.Build/Notifications/Sinks/SlackNotificationSink.cs b/Engine/Source/Programs/Horde/Horde.Build/Notifications/Sinks/SlackNotificationSink.cs
--- a/Engine/Source/Programs/Horde/Horde.Build/Notifications/Sinks/SlackNotificationSink.cs
+++ b/Engine/Source/Programs/Horde/Horde.Build/Notifications/Sinks/SlackNotificationSink.cs
@@ -783,7 +783,7 @@ namespace Horde.Build.Notifications.Sinks
 				}
 
 				string issueSummary = issue.UserSummary ?? issue.Summary;
-				string text = $"{GetPrefixForSeverity(issue.Severity)}{workflow.TriagePrefix}*New Issue <{issueUrl}|{issue.Id}>*: {issueSummary}{workflow.TriageSuffix}";
+				string text = $"{workflow.TriagePrefix}*New Issue <{issueUrl}|{issue.Id}>*: {issueSummary}{workflow.TriageSuffix}";
 				if (!spans.Any(x => x.NextSuccess == null && x.LastFailure.Annotations.WorkflowId != null)) // Thread may be shared by multiple workflows
 				{
 					text = $"~{text}~";
@@ -1388,23 +1388,6 @@ namespace Horde.Build.Notifications.Sinks
 				List<BlockBase> blocks = new List<BlockBase>();
 				blocks.Add(new HeaderBlock(AddEnvironmentAnnotation($"{report.Stream.Name}: {report.Time:d}")));
 
-				StringBuilder header = new StringBuilder();
-				if (report.Issues.Count == 0)
-				{
-					header.Append(":tick: No issues open.");
-				}
-				else
-				{
-					TimeSpan averageAge = TimeSpan.FromHours(report.Issues.Select(x => (report.Time - x.CreatedAt).TotalHours).Average());
-					header.Append($"*{report.Issues.Count} unique issues* currently open (average age {FormatReadableTimeSpan(averageAge)}).");
-				}
-				if (report.WorkflowStats.NumSteps > 0)
-				{
-					double totalPct = (report.WorkflowStats.NumPassingSteps * 100.0) / report.WorkflowStats.NumSteps;
-					header.Append($" {report.WorkflowStats.NumPassingSteps:n0} of {report.WorkflowStats.NumSteps:n0} (*{totalPct:0.0}%*) build steps succeeded since last status update.");
-				}
-				blocks.Add(new SectionBlock(header.ToString()));
-
 				PostMessageResponse? response = await SendMessageAsync(report.Channel, blocks: blocks.ToArray(), withEnvironment: false);
 				if (response != null && response.Ts != null)
 				{
@@ -1412,6 +1395,23 @@ namespace Horde.Build.Notifications.Sinks
 					string json = JsonSerializer.Serialize(state, _jsonSerializerOptions);
 					await AddOrUpdateMessageStateAsync(report.Channel, reportEventId, null, json, response.Ts);
 
+					StringBuilder header = new StringBuilder();
+					if (report.Issues.Count == 0)
+					{
+						header.Append(":tick: No issues open.");
+					}
+					else
+					{
+						TimeSpan averageAge = TimeSpan.FromHours(report.Issues.Select(x => (report.Time - x.CreatedAt).TotalHours).Average());
+						header.Append($"*{report.Issues.Count} unique issues* currently open (average age {FormatReadableTimeSpan(averageAge)}).");
+					}
+					if (report.WorkflowStats.NumSteps > 0)
+					{
+						double totalPct = (report.WorkflowStats.NumPassingSteps * 100.0) / report.WorkflowStats.NumSteps;
+						header.Append($" *{totalPct:0.0}%* of build steps succeeded since last status update ({report.WorkflowStats.NumPassingSteps:n0} of {report.WorkflowStats.NumSteps:n0}).");
+					}
+					await SendMessageAsync(report.Channel, text: header.ToString(), withEnvironment: false);
+
 					for (int idx = 0; idx < state.Blocks.Count; idx++)
 					{
 						string blockEventId = GetReportBlockEventId(response.Ts, idx);
