Horde: Tweak stats output at the end of each summary report.

diff --git a/Engine/Source/Programs/Horde/Horde.Build/Notifications/Sinks/SlackNotificationSink.cs b/Engine/Source/Programs/Horde/Horde.Build/Notifications/Sinks/SlackNotificationSink.cs
--- a/Engine/Source/Programs/Horde/Horde.Build/Notifications/Sinks/SlackNotificationSink.cs
+++ b/Engine/Source/Programs/Horde/Horde.Build/Notifications/Sinks/SlackNotificationSink.cs
@@ -1395,28 +1395,24 @@ namespace Horde.Build.Notifications.Sinks
 					string json = JsonSerializer.Serialize(state, _jsonSerializerOptions);
 					await AddOrUpdateMessageStateAsync(report.Channel, reportEventId, null, json, response.Ts);
 
-					StringBuilder header = new StringBuilder();
-					if (report.Issues.Count == 0)
+					if (state.Blocks.Count == 0)
 					{
-						header.Append(":tick: No issues open.");
+						string header = ":tick: No issues open.";
+						await SendMessageAsync(report.Channel, text: header, withEnvironment: false);
 					}
-					else
-					{
-						TimeSpan averageAge = TimeSpan.FromHours(report.Issues.Select(x => (report.Time - x.CreatedAt).TotalHours).Average());
-						header.Append($"*{report.Issues.Count} unique issues* currently open (average age {FormatReadableTimeSpan(averageAge)}).");
-					}
-					if (report.WorkflowStats.NumSteps > 0)
-					{
-						double totalPct = (report.WorkflowStats.NumPassingSteps * 100.0) / report.WorkflowStats.NumSteps;
-						header.Append($" *{totalPct:0.0}%* of build steps succeeded since last status update ({report.WorkflowStats.NumPassingSteps:n0} of {report.WorkflowStats.NumSteps:n0}).");
-					}
-					await SendMessageAsync(report.Channel, text: header.ToString(), withEnvironment: false);
 
 					for (int idx = 0; idx < state.Blocks.Count; idx++)
 					{
 						string blockEventId = GetReportBlockEventId(response.Ts, idx);
 						await UpdateReportBlockAsync(report.Channel, blockEventId, report.Time.UtcDateTime, report.Stream, state.Blocks[idx].TemplateId, issuesByBlock[idx], state.Blocks[idx].TemplateHeader);
 					}
+
+					if (report.WorkflowStats.NumSteps > 0)
+					{
+						double totalPct = (report.WorkflowStats.NumPassingSteps * 100.0) / report.WorkflowStats.NumSteps;
+						string header = $"*{totalPct:0.0}%* of build steps ({report.WorkflowStats.NumPassingSteps:n0} of {report.WorkflowStats.NumSteps:n0}) succeeded since last status update.";
+						await SendMessageAsync(report.Channel, text: String.Join(" ", header), withEnvironment: false);
+					}
 				}
 			}
 		}
