Horde: Note whether a fix has failed or needs merging in the issue summary.

diff --git a/Engine/Source/Programs/Horde/Horde.Build/Issues/IIssue.cs b/Engine/Source/Programs/Horde/Horde.Build/Issues/IIssue.cs
--- a/Engine/Source/Programs/Horde/Horde.Build/Issues/IIssue.cs
+++ b/Engine/Source/Programs/Horde/Horde.Build/Issues/IIssue.cs
@@ -190,6 +190,11 @@ namespace Horde.Build.Issues
 		/// Whether this stream contains the fix change
 		/// </summary>
 		bool? ContainsFix { get; }
+
+		/// <summary>
+		/// Whether the fix failed in this stream
+		/// </summary>
+		bool? FixFailed { get; }
 	}
 
 	/// <summary>
diff --git a/Engine/Source/Programs/Horde/Horde.Build/Issues/IIssueCollection.cs b/Engine/Source/Programs/Horde/Horde.Build/Issues/IIssueCollection.cs
--- a/Engine/Source/Programs/Horde/Horde.Build/Issues/IIssueCollection.cs
+++ b/Engine/Source/Programs/Horde/Horde.Build/Issues/IIssueCollection.cs
@@ -377,6 +377,9 @@ namespace Horde.Build.Issues
 		/// <inheritdoc cref="IIssueStream.ContainsFix"/>
 		public bool? ContainsFix { get; set; }
 
+		/// <inheritdoc cref="IIssueStream.FixFailed"/>
+		public bool? FixFailed { get; set; }
+
 		/// <summary>
 		/// Constructor
 		/// </summary>
diff --git a/Engine/Source/Programs/Horde/Horde.Build/Issues/IssueCollection.cs b/Engine/Source/Programs/Horde/Horde.Build/Issues/IssueCollection.cs
--- a/Engine/Source/Programs/Horde/Horde.Build/Issues/IssueCollection.cs
+++ b/Engine/Source/Programs/Horde/Horde.Build/Issues/IssueCollection.cs
@@ -173,6 +173,7 @@ namespace Horde.Build.Issues
 			public StreamId StreamId { get; set; }
 			public bool? MergeOrigin { get; set; }
 			public bool? ContainsFix { get; set; }
+			public bool? FixFailed { get; set; }
 
 			public IssueStream()
 			{
@@ -183,6 +184,7 @@ namespace Horde.Build.Issues
 				StreamId = other.StreamId;
 				MergeOrigin = other.MergeOrigin;
 				ContainsFix = other.ContainsFix;
+				FixFailed = other.FixFailed;
 			}
 		}
 
diff --git a/Engine/Source/Programs/Horde/Horde.Build/Issues/IssueService.cs b/Engine/Source/Programs/Horde/Horde.Build/Issues/IssueService.cs
--- a/Engine/Source/Programs/Horde/Horde.Build/Issues/IssueService.cs
+++ b/Engine/Source/Programs/Horde/Horde.Build/Issues/IssueService.cs
@@ -1006,6 +1006,10 @@ namespace Horde.Build.Issues
 						{
 							newStream.ContainsFix = await ContainsFixChange(newStream.StreamId, issue.FixChange.Value, fixChangeCache);
 						}
+						if (spans.Any(x => x.StreamId == newStream.StreamId && x.LastFailure.Change > issue.FixChange.Value))
+						{
+							newStream.FixFailed = true;
+						}
 					}
 				}
 
diff --git a/Engine/Source/Programs/Horde/Horde.Build/Notifications/Sinks/SlackNotificationSink.cs b/Engine/Source/Programs/Horde/Horde.Build/Notifications/Sinks/SlackNotificationSink.cs
--- a/Engine/Source/Programs/Horde/Horde.Build/Notifications/Sinks/SlackNotificationSink.cs
+++ b/Engine/Source/Programs/Horde/Horde.Build/Notifications/Sinks/SlackNotificationSink.cs
@@ -1411,7 +1411,7 @@ namespace Horde.Build.Notifications.Sinks
 					{
 						double totalPct = (report.WorkflowStats.NumPassingSteps * 100.0) / report.WorkflowStats.NumSteps;
 						string header = $"*{totalPct:0.0}%* of build steps ({report.WorkflowStats.NumPassingSteps:n0} of {report.WorkflowStats.NumSteps:n0}) succeeded since last status update.";
-						await SendMessageAsync(report.Channel, text: String.Join(" ", header), withEnvironment: false);
+						await SendMessageAsync(report.Channel, text: header, withEnvironment: false);
 					}
 				}
 			}
@@ -1577,13 +1577,17 @@ namespace Horde.Build.Notifications.Sinks
 			string status = "*Unassigned*";
 			if (issue.FixChange != null)
 			{
-				if (owner == null)
+				if (issue.Streams.Any(x => x.FixFailed ?? false))
 				{
-					status = $"Fixed in CL {issue.FixChange.Value}";
+					status = "*Fix failed*";
+				}
+				else if (span != null && issue.Streams.Any(x => x.StreamId == span.StreamId && (x.MergeOrigin ?? false) && !(x.ContainsFix ?? false)))
+				{
+					status = "*Fix needs merging*";
 				}
 				else
 				{
-					status = $"Fixed in CL {issue.FixChange.Value} by {owner.Name}";
+					status = $"Fixed in CL {issue.FixChange.Value}";
 				}
 			}
 			else if (issue.OwnerId != null)
