Horde: Post a divider between report results.

diff --git a/Engine/Source/Programs/Horde/Horde.Build/Notifications/Sinks/SlackNotificationSink.cs b/Engine/Source/Programs/Horde/Horde.Build/Notifications/Sinks/SlackNotificationSink.cs
--- a/Engine/Source/Programs/Horde/Horde.Build/Notifications/Sinks/SlackNotificationSink.cs
+++ b/Engine/Source/Programs/Horde/Horde.Build/Notifications/Sinks/SlackNotificationSink.cs
@@ -1395,6 +1395,8 @@ namespace Horde.Build.Notifications.Sinks
 					string json = JsonSerializer.Serialize(state, _jsonSerializerOptions);
 					await AddOrUpdateMessageStateAsync(report.Channel, reportEventId, null, json, response.Ts);
 
+					await SendMessageAsync(report.Channel, blocks: new[] { new DividerBlock() }, withEnvironment: false);
+
 					if (state.Blocks.Count == 0)
 					{
 						string header = ":tick: No issues open.";
@@ -1413,6 +1415,8 @@ namespace Horde.Build.Notifications.Sinks
 						string header = $"*{totalPct:0.0}%* of build steps ({report.WorkflowStats.NumPassingSteps:n0} of {report.WorkflowStats.NumSteps:n0}) succeeded since last status update.";
 						await SendMessageAsync(report.Channel, text: header, withEnvironment: false);
 					}
+
+					await SendMessageAsync(report.Channel, blocks: new[] { new DividerBlock() }, withEnvironment: false);
 				}
 			}
 		}
