Horde: Remove divider between issue report heading and list of issues.

diff --git a/Engine/Source/Programs/Horde/Horde.Build/Notifications/Sinks/SlackNotificationSink.cs b/Engine/Source/Programs/Horde/Horde.Build/Notifications/Sinks/SlackNotificationSink.cs
--- a/Engine/Source/Programs/Horde/Horde.Build/Notifications/Sinks/SlackNotificationSink.cs
+++ b/Engine/Source/Programs/Horde/Horde.Build/Notifications/Sinks/SlackNotificationSink.cs
@@ -1395,8 +1395,6 @@ namespace Horde.Build.Notifications.Sinks
 					string json = JsonSerializer.Serialize(state, _jsonSerializerOptions);
 					await AddOrUpdateMessageStateAsync(report.Channel, reportEventId, null, json, response.Ts);
 
-					await SendMessageAsync(report.Channel, blocks: new[] { new DividerBlock() }, withEnvironment: false);
-
 					if (state.Blocks.Count == 0)
 					{
 						string header = ":tick: No issues open.";
