Horde: Experiment with grouping together stream status updates under a common heading.

diff --git a/Engine/Source/Programs/Horde/Horde.Build/Notifications/INotificationService.cs b/Engine/Source/Programs/Horde/Horde.Build/Notifications/INotificationService.cs
--- a/Engine/Source/Programs/Horde/Horde.Build/Notifications/INotificationService.cs
+++ b/Engine/Source/Programs/Horde/Horde.Build/Notifications/INotificationService.cs
@@ -123,6 +123,6 @@ namespace Horde.Build.Notifications
 		/// Post a notification for the open issues in a stream
 		/// </summary>
 		/// <param name="report">The report data to send</param>
-		Task SendIssueReportAsync(IssueReport report);
+		Task SendIssueReportAsync(IssueReportGroup report);
 	}
 }
diff --git a/Engine/Source/Programs/Horde/Horde.Build/Notifications/INotificationSink.cs b/Engine/Source/Programs/Horde/Horde.Build/Notifications/INotificationSink.cs
--- a/Engine/Source/Programs/Horde/Horde.Build/Notifications/INotificationSink.cs
+++ b/Engine/Source/Programs/Horde/Horde.Build/Notifications/INotificationSink.cs
@@ -87,7 +87,7 @@ namespace Horde.Build.Notifications
 		/// </summary>
 		/// <param name="report"></param>
 		/// <returns></returns>
-		Task SendIssueReportAsync(IssueReport report);
+		Task SendIssueReportAsync(IssueReportGroup report);
 
 		#endregion
 
diff --git a/Engine/Source/Programs/Horde/Horde.Build/Notifications/NotificationService.cs b/Engine/Source/Programs/Horde/Horde.Build/Notifications/NotificationService.cs
--- a/Engine/Source/Programs/Horde/Horde.Build/Notifications/NotificationService.cs
+++ b/Engine/Source/Programs/Horde/Horde.Build/Notifications/NotificationService.cs
@@ -735,7 +735,7 @@ namespace Horde.Build.Notifications
 		}
 
 		/// <inheritdoc/>
-		public async Task SendIssueReportAsync(IssueReport report)
+		public async Task SendIssueReportAsync(IssueReportGroup report)
 		{
 			foreach (INotificationSink sink in _sinks)
 			{
diff --git a/Engine/Source/Programs/Horde/Horde.Build/Notifications/Sinks/SlackNotificationSink.cs b/Engine/Source/Programs/Horde/Horde.Build/Notifications/Sinks/SlackNotificationSink.cs
--- a/Engine/Source/Programs/Horde/Horde.Build/Notifications/Sinks/SlackNotificationSink.cs
+++ b/Engine/Source/Programs/Horde/Horde.Build/Notifications/Sinks/SlackNotificationSink.cs
@@ -783,7 +783,7 @@ namespace Horde.Build.Notifications.Sinks
 				}
 
 				string issueSummary = issue.UserSummary ?? issue.Summary;
-				string text = $"{workflow.TriagePrefix}*New Issue <{issueUrl}|{issue.Id}>*: {issueSummary}{workflow.TriageSuffix}";
+				string text = $"{workflow.TriagePrefix}{GetSeverityPrefix(issue.Severity)}Issue <{issueUrl}|{issue.Id}>: {issueSummary}{workflow.TriageSuffix}";
 				if (!spans.Any(x => x.NextSuccess == null && x.LastFailure.Annotations.WorkflowId != null)) // Thread may be shared by multiple workflows
 				{
 					text = $"~{text}~";
@@ -1315,107 +1315,117 @@ namespace Horde.Build.Notifications.Sinks
 		static string GetReportBlockEventId(string ts, int idx) => $"issue-report-block:{ts}:{idx}";
 
 		/// <inheritdoc/>
-		public async Task SendIssueReportAsync(IssueReport report)
+		public async Task SendIssueReportAsync(IssueReportGroup group)
 		{
-			if (report.Channel != null)
+			List<BlockBase> introBlocks = new List<BlockBase>();
+			introBlocks.Add(new HeaderBlock($"Summary for {group.Time:M}, {group.Time:H}:{group.Time:mm} UTC"));
+			introBlocks.Add(new DividerBlock());
+			await SendMessageAsync(group.Channel, blocks: introBlocks.ToArray(), withEnvironment: false);
+
+			foreach (IssueReport report in group.Reports.OrderBy(x => x.WorkflowId).ThenBy(x => x.Stream.Id))
 			{
-				const int MaxIssuesPerMessage = 8;
+				await SendIssueReportForStreamAsync(group.Channel, group.Time, report);
+			}
+		}
+
+		async Task SendIssueReportForStreamAsync(string channel, DateTime time, IssueReport report)
+		{
+			const int MaxIssuesPerMessage = 8;
+
+			ReportState state = new ReportState();
+			state.Time = time;
 
-				ReportState state = new ReportState();
-				state.Time = report.Time.UtcDateTime;
+			List<List<(IIssue, IIssueSpan?, bool)>> issuesByBlock = new List<List<(IIssue, IIssueSpan?, bool)>>();
+			if (report.GroupByTemplate)
+			{
+				Dictionary<int, IIssue> issueIdToInfo = new Dictionary<int, IIssue>();
+				foreach (IIssue issue in report.Issues)
+				{
+					issueIdToInfo[issue.Id] = issue;
+				}
 
-				List<List<(IIssue, IIssueSpan?, bool)>> issuesByBlock = new List<List<(IIssue, IIssueSpan?, bool)>>();
-				if (report.GroupByTemplate)
+				foreach (IGrouping<TemplateRefId, IIssueSpan> group in report.IssueSpans.GroupBy(x => x.TemplateRefId).OrderBy(x => x.Key.ToString()))
 				{
-					Dictionary<int, IIssue> issueIdToInfo = new Dictionary<int, IIssue>();
-					foreach (IIssue issue in report.Issues)
+					TemplateRefConfig? templateConfig;
+					if (!report.Stream.Config.TryGetTemplate(group.Key, out templateConfig))
 					{
-						issueIdToInfo[issue.Id] = issue;
+						continue;
 					}
-
-					foreach (IGrouping<TemplateRefId, IIssueSpan> group in report.IssueSpans.GroupBy(x => x.TemplateRefId).OrderBy(x => x.Key.ToString()))
+					if (!IsIssueOpenForWorkflow(report.Stream.Id, group.Key, report.WorkflowId, group))
 					{
-						TemplateRefConfig? templateConfig;
-						if (!report.Stream.Config.TryGetTemplate(group.Key, out templateConfig))
-						{
-							continue;
-						}
-						if (!IsIssueOpenForWorkflow(report.Stream.Id, group.Key, report.WorkflowId, group))
-						{
-							continue;
-						}
+						continue;
+					}
 
-						Dictionary<IIssue, IIssueSpan> pairs = new Dictionary<IIssue, IIssueSpan>();
-						foreach (IIssueSpan span in group)
+					Dictionary<IIssue, IIssueSpan> pairs = new Dictionary<IIssue, IIssueSpan>();
+					foreach (IIssueSpan span in group)
+					{
+						if (issueIdToInfo.TryGetValue(span.IssueId, out IIssue? issue))
 						{
-							if (issueIdToInfo.TryGetValue(span.IssueId, out IIssue? issue))
-							{
-								pairs[issue] = span;
-							}
+							pairs[issue] = span;
 						}
+					}
 
-						if (pairs.Count > 0)
+					if (pairs.Count > 0)
+					{
+						bool templateHeader = true;
+						foreach (IReadOnlyList<KeyValuePair<IIssue, IIssueSpan>> batch in pairs.OrderBy(x => x.Key.Id).Batch(MaxIssuesPerMessage))
 						{
-							bool templateHeader = true;
-							foreach (IReadOnlyList<KeyValuePair<IIssue, IIssueSpan>> batch in pairs.OrderBy(x => x.Key.Id).Batch(MaxIssuesPerMessage))
-							{
-								ReportBlock block = new ReportBlock();
-								block.TemplateHeader = templateHeader;
-								block.StreamId = report.Stream.Id;
-								block.TemplateId = group.Key;
-								block.IssueIds.AddRange(batch.Select(x => x.Key.Id));
-								state.Blocks.Add(block);
-
-								issuesByBlock.Add(batch.Select(x => (x.Key, (IIssueSpan?)x.Value, true)).ToList());
-								templateHeader = false;
-							}
+							ReportBlock block = new ReportBlock();
+							block.TemplateHeader = templateHeader;
+							block.StreamId = report.Stream.Id;
+							block.TemplateId = group.Key;
+							block.IssueIds.AddRange(batch.Select(x => x.Key.Id));
+							state.Blocks.Add(block);
+
+							issuesByBlock.Add(batch.Select(x => (x.Key, (IIssueSpan?)x.Value, true)).ToList());
+							templateHeader = false;
 						}
 					}
 				}
-				else
+			}
+			else
+			{
+				foreach (IReadOnlyList<IIssue> batch in report.Issues.OrderByDescending(x => x.Id).Batch(MaxIssuesPerMessage))
 				{
-					foreach (IReadOnlyList<IIssue> batch in report.Issues.OrderByDescending(x => x.Id).Batch(MaxIssuesPerMessage))
-					{
-						ReportBlock block = new ReportBlock();
-						block.StreamId = report.Stream.Id;
-						block.IssueIds.AddRange(batch.Select(x => x.Id));
-						state.Blocks.Add(block);
+					ReportBlock block = new ReportBlock();
+					block.StreamId = report.Stream.Id;
+					block.IssueIds.AddRange(batch.Select(x => x.Id));
+					state.Blocks.Add(block);
 
-						issuesByBlock.Add(batch.Select(x => (x, report.IssueSpans.FirstOrDefault(y => y.IssueId == x.Id), true)).ToList());
-					}
+					issuesByBlock.Add(batch.Select(x => (x, report.IssueSpans.FirstOrDefault(y => y.IssueId == x.Id), true)).ToList());
 				}
+			}
 
-				List<BlockBase> blocks = new List<BlockBase>();
-				blocks.Add(new HeaderBlock(AddEnvironmentAnnotation($"{report.Stream.Name}: {report.Time:d}")));
-
-				PostMessageResponse? response = await SendMessageAsync(report.Channel, blocks: blocks.ToArray(), withEnvironment: false);
-				if (response != null && response.Ts != null)
-				{
-					string reportEventId = GetReportEventId(report.Stream.Id, report.WorkflowId);
-					string json = JsonSerializer.Serialize(state, _jsonSerializerOptions);
-					await AddOrUpdateMessageStateAsync(report.Channel, reportEventId, null, json, response.Ts);
+			List<BlockBase> blocks = new List<BlockBase>();
+			blocks.Add(new HeaderBlock(report.Stream.Name));
 
-					if (state.Blocks.Count == 0)
-					{
-						string header = ":tick: No issues open.";
-						await SendMessageAsync(report.Channel, text: header, withEnvironment: false);
-					}
+			PostMessageResponse? response = await SendMessageAsync(channel, blocks: blocks.ToArray(), withEnvironment: false);
+			if (response != null && response.Ts != null)
+			{
+				string reportEventId = GetReportEventId(report.Stream.Id, report.WorkflowId);
+				string json = JsonSerializer.Serialize(state, _jsonSerializerOptions);
+				await AddOrUpdateMessageStateAsync(channel, reportEventId, null, json, response.Ts);
 
-					for (int idx = 0; idx < state.Blocks.Count; idx++)
-					{
-						string blockEventId = GetReportBlockEventId(response.Ts, idx);
-						await UpdateReportBlockAsync(report.Channel, blockEventId, report.Time.UtcDateTime, report.Stream, state.Blocks[idx].TemplateId, issuesByBlock[idx], state.Blocks[idx].TemplateHeader);
-					}
+				if (state.Blocks.Count == 0)
+				{
+					string header = ":tick: No issues open.";
+					await SendMessageAsync(channel, text: header, withEnvironment: false);
+				}
 
-					if (report.WorkflowStats.NumSteps > 0)
-					{
-						double totalPct = (report.WorkflowStats.NumPassingSteps * 100.0) / report.WorkflowStats.NumSteps;
-						string header = $"*{totalPct:0.0}%* of build steps ({report.WorkflowStats.NumPassingSteps:n0} of {report.WorkflowStats.NumSteps:n0}) succeeded since last status update.";
-						await SendMessageAsync(report.Channel, text: header, withEnvironment: false);
-					}
+				for (int idx = 0; idx < state.Blocks.Count; idx++)
+				{
+					string blockEventId = GetReportBlockEventId(response.Ts, idx);
+					await UpdateReportBlockAsync(channel, blockEventId, time, report.Stream, state.Blocks[idx].TemplateId, issuesByBlock[idx], state.Blocks[idx].TemplateHeader);
+				}
 
-					await SendMessageAsync(report.Channel, blocks: new[] { new DividerBlock() }, withEnvironment: false);
+				if (report.WorkflowStats.NumSteps > 0)
+				{
+					double totalPct = (report.WorkflowStats.NumPassingSteps * 100.0) / report.WorkflowStats.NumSteps;
+					string header = $"*{totalPct:0.0}%* of build steps ({report.WorkflowStats.NumPassingSteps:n0}/{report.WorkflowStats.NumSteps:n0}) succeeded since last status update.";
+					await SendMessageAsync(channel, text: header, withEnvironment: false);
 				}
+
+				await SendMessageAsync(channel, blocks: new[] { new DividerBlock() }, withEnvironment: false);
 			}
 		}
 
@@ -1556,7 +1566,7 @@ namespace Horde.Build.Notifications.Sinks
 			await SendOrUpdateMessageAsync(channel, eventId, null, body.ToString(), withEnvironment: false);
 		}
 
-		string GetPrefixForSeverity(IssueSeverity severity)
+		string GetSeverityPrefix(IssueSeverity severity)
 		{
 			return (severity == IssueSeverity.Warning) ? _settings.SlackWarningPrefix : _settings.SlackErrorPrefix;
 		}
@@ -1613,8 +1623,8 @@ namespace Horde.Build.Notifications.Sinks
 				status = $"{status} - *Quarantined*";
 			}
 
-			string prefix = GetPrefixForSeverity(issue.Severity);
-			StringBuilder body = new StringBuilder($"{prefix}*Issue <{issueUrl}|{issue.Id}>");
+			string prefix = GetSeverityPrefix(issue.Severity);
+			StringBuilder body = new StringBuilder($"{prefix}Issue *<{issueUrl}|{issue.Id}>");
 
 			if (triageChannel != null)
 			{
diff --git a/Engine/Source/Programs/Horde/Horde.Build/Server/Config.cs b/Engine/Source/Programs/Horde/Horde.Build/Server/Config.cs
--- a/Engine/Source/Programs/Horde/Horde.Build/Server/Config.cs
+++ b/Engine/Source/Programs/Horde/Horde.Build/Server/Config.cs
@@ -362,7 +362,7 @@ namespace Horde.Build.Server
 		/// <summary>
 		/// Prefix for all triage messages
 		/// </summary>
-		public string? TriagePrefix { get; set; }
+		public string? TriagePrefix { get; set; } = "*[NEW]* ";
 
 		/// <summary>
 		/// Suffix for all triage messages
