Horde: Missing file.

diff --git a/Engine/Source/Programs/Horde/Horde.Build/Issues/IssueReportService.cs b/Engine/Source/Programs/Horde/Horde.Build/Issues/IssueReportService.cs
--- a/Engine/Source/Programs/Horde/Horde.Build/Issues/IssueReportService.cs
+++ b/Engine/Source/Programs/Horde/Horde.Build/Issues/IssueReportService.cs
@@ -31,9 +31,7 @@ namespace Horde.Build.Issues
 
 	public class IssueReport
 	{
-		public string Channel { get; }
 		public WorkflowId WorkflowId { get; }
-		public DateTimeOffset Time { get; }
 		public IStream Stream { get; }
 		public WorkflowStats WorkflowStats { get; }
 		public string? TriageChannel { get; }
@@ -41,17 +39,29 @@ namespace Horde.Build.Issues
 		public List<IIssueSpan> IssueSpans { get; } = new List<IIssueSpan>();
 		public bool GroupByTemplate { get; }
 
-		public IssueReport(string channel, WorkflowId workflowId, DateTimeOffset time, IStream stream, WorkflowStats workflowStats, string? triageChannel, bool groupByTemplate)
+		public IssueReport(WorkflowId workflowId, IStream stream, WorkflowStats workflowStats, string? triageChannel, bool groupByTemplate)
 		{
-			Channel = channel;
 			WorkflowId = workflowId;
-			Time = time;
 			Stream = stream;
 			WorkflowStats = workflowStats;
 			TriageChannel = triageChannel;
 			GroupByTemplate = groupByTemplate;
 		}
 	}
+
+	public class IssueReportGroup
+	{
+		public string Channel { get; }
+		public DateTime Time { get; }
+		public List<IssueReport> Reports { get; } = new List<IssueReport>();
+
+		public IssueReportGroup(string channel, DateTime time)
+		{
+			Channel = channel;
+			Time = time;
+		}
+	}
+
 #pragma warning restore CS1591 // Missing XML comment for publicly visible type or member
 
 	[SingletonDocument("6268871c211d05611b3e4fd8")]
@@ -112,6 +122,9 @@ namespace Horde.Build.Issues
 
 			DateTime currentTime = _clock.UtcNow;
 
+			List<string> updateKeys = new List<string>();
+			List<IssueReportGroup> groups = new List<IssueReportGroup>();
+
 			List<IStream> streams = await _streamCollection.FindAllAsync();
 			foreach (IStream stream in streams)
 			{
@@ -157,7 +170,7 @@ namespace Horde.Build.Issues
 							workflowStats = new WorkflowStats();
 						}
 
-						IssueReport report = new IssueReport(workflow.ReportChannel, workflow.Id, currentTime, stream, workflowStats, workflow.TriageChannel, workflow.GroupIssuesByTemplate);
+						IssueReport report = new IssueReport(workflow.Id, stream, workflowStats, workflow.TriageChannel, workflow.GroupIssuesByTemplate);
 						foreach (IIssueSpan span in spans)
 						{
 							if (span.LastSuccess != null && span.LastFailure.Annotations.WorkflowId == workflow.Id)
@@ -169,23 +182,40 @@ namespace Horde.Build.Issues
 						HashSet<int> issueIds = new HashSet<int>(report.IssueSpans.Select(x => x.IssueId));
 						report.Issues.AddRange(issues.Where(x => issueIds.Contains(x.Id)));
 
-						await _notificationService.SendIssueReportAsync(report);
+						DateTime reportTime = lastScheduledReportTime;
+
+						IssueReportGroup? group = groups.FirstOrDefault(x => x.Channel == workflow.ReportChannel && x.Time == reportTime);
+						if (group == null)
+						{
+							group = new IssueReportGroup(workflow.ReportChannel, reportTime);
+							groups.Add(group);
+						}
+						group.Reports.Add(report);
 
-						state = await _state.UpdateAsync(s => s.KeyToLastReportTime[key] = currentTime);
+						updateKeys.Add(key);
 					}
 				}
 			}
 
-			if (invalidKeys.Count > 0)
+			foreach (IssueReportGroup group in groups)
 			{
-				void RemoveInvalidKeys(IssueReportState state)
+				await _notificationService.SendIssueReportAsync(group);
+			}
+
+			if (updateKeys.Count > 0 || invalidKeys.Count > 0)
+			{
+				void UpdateKeys(IssueReportState state)
 				{
+					foreach (string updateKey in updateKeys)
+					{
+						state.KeyToLastReportTime[updateKey] = currentTime;
+					}
 					foreach (string invalidKey in invalidKeys)
 					{
 						state.KeyToLastReportTime.Remove(invalidKey);
 					}
 				}
-				state = await _state.UpdateAsync(RemoveInvalidKeys);
+				state = await _state.UpdateAsync(UpdateKeys);
 			}
 		}
 
