Horde: Fix tests.

diff --git a/Engine/Source/Programs/Horde/Horde.Build.Tests/NotificationServiceTests.cs b/Engine/Source/Programs/Horde/Horde.Build.Tests/NotificationServiceTests.cs
--- a/Engine/Source/Programs/Horde/Horde.Build.Tests/NotificationServiceTests.cs
+++ b/Engine/Source/Programs/Horde/Horde.Build.Tests/NotificationServiceTests.cs
@@ -42,7 +42,7 @@ namespace Horde.Build.Tests
 		public Task NotifyIssueUpdatedAsync(IIssue issue) { throw new NotImplementedException(); }
 		public Task NotifyConfigUpdateFailureAsync(string errorMessage, string fileName, int? change = null, IUser? author = null, string? description = null) { throw new NotImplementedException(); }
 		public Task NotifyDeviceServiceAsync(string message, IDevice? device = null, IDevicePool? pool = null, IStream? stream = null, IJob? job = null, IJobStep? step = null, INode? node = null, IUser? user = null) { throw new NotImplementedException(); }
-		public Task SendIssueReportAsync(IssueReport report) => throw new NotImplementedException();
+		public Task SendIssueReportAsync(IssueReportGroup report) => throw new NotImplementedException();
 	}
 
 	[TestClass]
