Horde: Fix exception formatting date/time strings in issue reports.

diff --git a/Engine/Source/Programs/Horde/Horde.Build/Notifications/Sinks/SlackNotificationSink.cs b/Engine/Source/Programs/Horde/Horde.Build/Notifications/Sinks/SlackNotificationSink.cs
--- a/Engine/Source/Programs/Horde/Horde.Build/Notifications/Sinks/SlackNotificationSink.cs
+++ b/Engine/Source/Programs/Horde/Horde.Build/Notifications/Sinks/SlackNotificationSink.cs
@@ -1317,8 +1317,11 @@ namespace Horde.Build.Notifications.Sinks
 		/// <inheritdoc/>
 		public async Task SendIssueReportAsync(IssueReportGroup group)
 		{
+			string reportTime = $"<!date^{new DateTimeOffset(group.Time).ToUnixTimeSeconds()}^{{date_long_pretty}} at {{time}}|{group.Time}>";
+
 			List<BlockBase> introBlocks = new List<BlockBase>();
-			introBlocks.Add(new HeaderBlock($"Summary for {group.Time:M}, {group.Time:H}:{group.Time:mm} UTC"));
+			introBlocks.Add(new HeaderBlock("Summary"));
+			introBlocks.Add(new SectionBlock($"Gathered {reportTime}"));
 			introBlocks.Add(new DividerBlock());
 			await SendMessageAsync(group.Channel, blocks: introBlocks.ToArray(), withEnvironment: false);
 
