Horde Storage - Always start ref gc of each namespace so we get a more consistent cleanup timing
Also improved some apm tagging.

diff --git a/Engine/Source/Programs/Horde/HordeStorage/Horde.Storage/Implementation/GC/CacheCleanup.cs b/Engine/Source/Programs/Horde/HordeStorage/Horde.Storage/Implementation/GC/CacheCleanup.cs
--- a/Engine/Source/Programs/Horde/HordeStorage/Horde.Storage/Implementation/GC/CacheCleanup.cs
+++ b/Engine/Source/Programs/Horde/HordeStorage/Horde.Storage/Implementation/GC/CacheCleanup.cs
@@ -95,6 +95,7 @@ namespace Horde.Storage.Implementation
                 _logger.Information("Attempting to delete object {Namespace} {Bucket} {Name} as it was last updated {LastAccessTime} which is older then {CutoffTime}", ns, bucket, name, lastAccessTime, cutoffTime);
                 using IScope scope = Tracer.Instance.StartActive("gc.ref");
                 scope.Span.ResourceName = $"{ns}:{bucket}.{name}";
+                scope.Span.SetTag("namespace", ns.ToString());
                 // delete the old record from the ref refs
 
                 bool storeDelete = false;
diff --git a/Engine/Source/Programs/Horde/HordeStorage/Horde.Storage/Implementation/GC/RefCleanupService.cs b/Engine/Source/Programs/Horde/HordeStorage/Horde.Storage/Implementation/GC/RefCleanupService.cs
--- a/Engine/Source/Programs/Horde/HordeStorage/Horde.Storage/Implementation/GC/RefCleanupService.cs
+++ b/Engine/Source/Programs/Horde/HordeStorage/Horde.Storage/Implementation/GC/RefCleanupService.cs
@@ -60,21 +60,14 @@ namespace Horde.Storage.Implementation
                     return false;
                 }
                 List<NamespaceId>? namespaces = await _referencesStore.GetNamespaces().ToListAsync(cancellationToken);
-                await Parallel.ForEachAsync(namespaces, cancellationToken, async (ns, token) =>
+                List<Task> cleanupTasks = new List<Task>();
+
+                foreach(NamespaceId ns in namespaces)
                 {
-                    _logger.Information("Attempting to run Refs Cleanup of {Namespace}. ", ns);
-                    try
-                    {
-                        int countOfRemovedRecords = await state.RefCleanup.Cleanup(ns, token);
-                        _logger.Information("Ran Refs Cleanup of {Namespace}. Deleted {CountRefRecords}", ns,
-                            countOfRemovedRecords);
-                    }
-                    catch (Exception e)
-                    {
-                        _logger.Error("Error running Refs Cleanup of {Namespace}. {Exception}", ns, e);
-                    }
-                });
+                    cleanupTasks.Add(DoCleanup(ns, state, cancellationToken));
+                }
 
+                await Task.WhenAll(cleanupTasks);
                 // Do not run the cleanup of legacy namespaces as this can take a very long time and we do not care about these anymore
                 /*
                  await foreach (NamespaceId ns in state.Refs.GetNamespaces().WithCancellation(cancellationToken))
@@ -98,5 +91,20 @@ namespace Horde.Storage.Implementation
                 _alreadyPolling = false;
             }
         }
+
+        private async Task DoCleanup(NamespaceId ns, RefCleanupState state, CancellationToken cancellationToken)
+        {
+            _logger.Information("Attempting to run Refs Cleanup of {Namespace}. ", ns);
+            try
+            {
+                int countOfRemovedRecords = await state.RefCleanup.Cleanup(ns, cancellationToken);
+                _logger.Information("Ran Refs Cleanup of {Namespace}. Deleted {CountRefRecords}", ns,
+                    countOfRemovedRecords);
+            }
+            catch (Exception e)
+            {
+                _logger.Error("Error running Refs Cleanup of {Namespace}. {Exception}", ns, e);
+            }
+        }
     }
 }
