Horde: More tweaks to Slack notifications.

diff --git a/Engine/Source/Programs/Horde/Horde.Build/Notifications/Sinks/SlackNotificationSink.cs b/Engine/Source/Programs/Horde/Horde.Build/Notifications/Sinks/SlackNotificationSink.cs
--- a/Engine/Source/Programs/Horde/Horde.Build/Notifications/Sinks/SlackNotificationSink.cs
+++ b/Engine/Source/Programs/Horde/Horde.Build/Notifications/Sinks/SlackNotificationSink.cs
@@ -1317,11 +1317,7 @@ namespace Horde.Build.Notifications.Sinks
 		/// <inheritdoc/>
 		public async Task SendIssueReportAsync(IssueReportGroup group)
 		{
-			string reportTime = $"<!date^{new DateTimeOffset(group.Time).ToUnixTimeSeconds()}^{{date_long_pretty}} at {{time}}|{group.Time}>";
-
 			List<BlockBase> introBlocks = new List<BlockBase>();
-			introBlocks.Add(new HeaderBlock("Summary"));
-			introBlocks.Add(new SectionBlock($"Gathered {reportTime}"));
 			introBlocks.Add(new DividerBlock());
 			await SendMessageAsync(group.Channel, blocks: introBlocks.ToArray(), withEnvironment: false);
 
@@ -1400,7 +1396,7 @@ namespace Horde.Build.Notifications.Sinks
 			}
 
 			List<BlockBase> blocks = new List<BlockBase>();
-			blocks.Add(new HeaderBlock(report.Stream.Name));
+			blocks.Add(new HeaderBlock($"Summary for {report.Stream.Name}"));
 
 			PostMessageResponse? response = await SendMessageAsync(channel, blocks: blocks.ToArray(), withEnvironment: false);
 			if (response != null && response.Ts != null)
