Horde: Remove initial divider block from Slack reports.

diff --git a/Engine/Source/Programs/Horde/Horde.Build/Notifications/Sinks/SlackNotificationSink.cs b/Engine/Source/Programs/Horde/Horde.Build/Notifications/Sinks/SlackNotificationSink.cs
--- a/Engine/Source/Programs/Horde/Horde.Build/Notifications/Sinks/SlackNotificationSink.cs
+++ b/Engine/Source/Programs/Horde/Horde.Build/Notifications/Sinks/SlackNotificationSink.cs
@@ -1317,10 +1317,6 @@ namespace Horde.Build.Notifications.Sinks
 		/// <inheritdoc/>
 		public async Task SendIssueReportAsync(IssueReportGroup group)
 		{
-			List<BlockBase> introBlocks = new List<BlockBase>();
-			introBlocks.Add(new DividerBlock());
-			await SendMessageAsync(group.Channel, blocks: introBlocks.ToArray(), withEnvironment: false);
-
 			foreach (IssueReport report in group.Reports.OrderBy(x => x.WorkflowId).ThenBy(x => x.Stream.Id))
 			{
 				await SendIssueReportForStreamAsync(group.Channel, group.Time, report);
