Horde Storage - Fixed issue were blobs where not being checked for existence before being returned by the reference resolving.

diff --git a/Engine/Source/Programs/Horde/HordeStorage/Horde.Storage/Controllers/ReferencesController.cs b/Engine/Source/Programs/Horde/HordeStorage/Horde.Storage/Controllers/ReferencesController.cs
--- a/Engine/Source/Programs/Horde/HordeStorage/Horde.Storage/Controllers/ReferencesController.cs
+++ b/Engine/Source/Programs/Horde/HordeStorage/Horde.Storage/Controllers/ReferencesController.cs
@@ -493,7 +493,11 @@ namespace Horde.Storage.Controllers
             {
                 return NotFound(new ProblemDetails {Title = $"Namespace {e.Namespace} did not exist"});
             }
-            catch (ObjectNotFoundException e)
+            catch (BlobNotFoundException e)
+            {
+                return NotFound(new ProblemDetails { Title = $"Blob {e.Blob} in namespace {ns} did not exist" });
+            }
+			catch (ObjectNotFoundException e)
             {
                 return NotFound(new ProblemDetails {Title = $"Object {e.Bucket} {e.Key} in namespace {e.Namespace} did not exist"});
             }
diff --git a/Engine/Source/Programs/Horde/HordeStorage/Horde.Storage/Implementation/IReferenceResolver.cs b/Engine/Source/Programs/Horde/HordeStorage/Horde.Storage/Implementation/IReferenceResolver.cs
--- a/Engine/Source/Programs/Horde/HordeStorage/Horde.Storage/Implementation/IReferenceResolver.cs
+++ b/Engine/Source/Programs/Horde/HordeStorage/Horde.Storage/Implementation/IReferenceResolver.cs
@@ -186,14 +186,16 @@ namespace Horde.Storage.Implementation
         public async IAsyncEnumerable<BlobIdentifier> GetReferencedBlobs(NamespaceId ns, CbObject cb)
         {
             List<Task<(ContentId, BlobIdentifier[]?)>> pendingContentIdResolves = new();
-            List<ContentId> unresolvedContentIdReferences = new List<ContentId>();
+            List<Task<(BlobIdentifier, bool)>> pendingBlobExistsChecks = new();
+			List<ContentId> unresolvedContentIdReferences = new List<ContentId>();
+            List<BlobIdentifier> unresolvedBlobReferences = new List<BlobIdentifier>();
 
-            // Resolve all the attachments
-            await foreach (Attachment attachment in GetAttachments(ns, cb))
+			// Resolve all the attachments
+			await foreach (Attachment attachment in GetAttachments(ns, cb))
             {
                 if (attachment is BlobAttachment blobAttachment)
                 {
-                    yield return blobAttachment.Identifier;
+                    pendingBlobExistsChecks.Add(CheckBlobExists(ns, blobAttachment.Identifier));
                 }
                 else if (attachment is ContentIdAttachment contentIdAttachment)
                 {
@@ -202,8 +204,8 @@ namespace Horde.Storage.Implementation
                 }
                 else if (attachment is ObjectAttachment objectAttachment)
                 {
-                    // a object just references the same blob, traversing the object attachment is done in GetAttachments
-                    yield return objectAttachment.Identifier;
+					// a object just references the same blob, traversing the object attachment is done in GetAttachments
+					pendingBlobExistsChecks.Add(CheckBlobExists(ns, objectAttachment.Identifier));
                 }
                 else
                 {
@@ -219,7 +221,7 @@ namespace Horde.Storage.Implementation
                 {
                     foreach (BlobIdentifier b in resolvedBlobs)
                     {
-                        yield return b;
+                        pendingBlobExistsChecks.Add(CheckBlobExists(ns, b));
                     }
                 }
                 else
@@ -228,12 +230,31 @@ namespace Horde.Storage.Implementation
                 }
             }
 
-            // if there were any content ids we did not recognize we throw a partial reference exception
-            if (unresolvedContentIdReferences.Count != 0)
+			// return any verified blobs
+            foreach (Task<(BlobIdentifier, bool)> pendingBlobExistsTask in pendingBlobExistsChecks)
+            {
+                (BlobIdentifier blob, bool exists) = await pendingBlobExistsTask;
+                if (exists)
+                {
+                    yield return blob;
+                }
+                else
+                {
+                    unresolvedBlobReferences.Add(blob);
+                }
+            }
+
+			// if there were any content ids we did not recognize we throw a partial reference exception
+			if (unresolvedContentIdReferences.Count != 0)
             {
                 throw new PartialReferenceResolveException(unresolvedContentIdReferences);
             }
-        }
+            // if there were any blobs missing we throw a partial reference exception
+            if (unresolvedBlobReferences.Count != 0)
+            {
+                throw new ReferenceIsMissingBlobsException(unresolvedBlobReferences);
+            }
+		}
 
         private async Task<CbObject> ParseCompactBinaryAttachment(NamespaceId ns, BlobIdentifier blobIdentifier)
         {
@@ -251,7 +272,12 @@ namespace Horde.Storage.Implementation
             BlobIdentifier[]? resolvedBlobs = await _contentIdStore.Resolve(ns, contentId);
             return (contentId, resolvedBlobs);
         }
-    }
+
+        private async Task<(BlobIdentifier, bool)> CheckBlobExists(NamespaceId ns, BlobIdentifier blob)
+        {
+            return (blob, await _blobStore.Exists(ns, blob));
+        }
+	}
 
     public class PartialReferenceResolveException : Exception
     {
diff --git a/Engine/Source/Programs/Horde/HordeStorage/Horde.Storage/Tests/Functional/Objects/ReferencesTests.cs b/Engine/Source/Programs/Horde/HordeStorage/Horde.Storage/Tests/Functional/Objects/ReferencesTests.cs
--- a/Engine/Source/Programs/Horde/HordeStorage/Horde.Storage/Tests/Functional/Objects/ReferencesTests.cs
+++ b/Engine/Source/Programs/Horde/HordeStorage/Horde.Storage/Tests/Functional/Objects/ReferencesTests.cs
@@ -1249,6 +1249,13 @@ namespace Horde.Storage.FunctionalTests.References
                 Assert.AreEqual(blobHash, BlobIdentifier.FromBlob(roundTrippedBuffer));
             }
 
+            {
+                // verify that head checks find the object
+                using HttpRequestMessage headRequest = new HttpRequestMessage(HttpMethod.Head, new Uri($"api/v1/refs/{TestNamespace}/bucket/{key}", UriKind.Relative));
+                HttpResponseMessage getResponse = await Client.SendAsync(headRequest);
+                getResponse.EnsureSuccessStatusCode();
+            }
+
             {
                 // delete the blob 
                 await Service.DeleteObject(TestNamespace, blobHash);
@@ -1276,8 +1283,26 @@ namespace Horde.Storage.FunctionalTests.References
                 Assert.IsNotNull(problem);
                 Assert.AreEqual($"Object {blobHash} in {TestNamespace} not found", problem.Title);
             }
+
+            {
+                // verify that head checks fail
+                using HttpRequestMessage headRequest = new HttpRequestMessage(HttpMethod.Head, new Uri($"api/v1/refs/{TestNamespace}/bucket/{key}", UriKind.Relative));
+                HttpResponseMessage getResponse = await Client.SendAsync(headRequest);
+                Assert.AreEqual(HttpStatusCode.NotFound, getResponse.StatusCode);
+            }
+
+            {
+                // verify that the exists check correctly returns a missing blob
+                HttpResponseMessage existsResponse = await Client.GetAsync(new Uri($"api/v1/refs/{TestNamespace}/exists?names=bucket.{key}", UriKind.Relative));
+                existsResponse.EnsureSuccessStatusCode();
+
+                ExistCheckMultipleRefsResponse response = await existsResponse.Content.ReadAsAsync<ExistCheckMultipleRefsResponse>();
+                Assert.AreEqual(1, response.Missing.Count);
+				Assert.AreEqual(key, response.Missing[0].Key);
+				Assert.AreEqual("bucket", response.Missing[0].Bucket.ToString());
+            }
         }
-        
+
         [TestMethod]
         public async Task DeleteObject()
         {
