Horde Storage - Fixed issue were exists would raise a 500 due to a blob missing

diff --git a/Engine/Source/Programs/Horde/HordeStorage/Horde.Storage/Implementation/IReferenceResolver.cs b/Engine/Source/Programs/Horde/HordeStorage/Horde.Storage/Implementation/IReferenceResolver.cs
--- a/Engine/Source/Programs/Horde/HordeStorage/Horde.Storage/Implementation/IReferenceResolver.cs
+++ b/Engine/Source/Programs/Horde/HordeStorage/Horde.Storage/Implementation/IReferenceResolver.cs
@@ -118,6 +118,10 @@ namespace Horde.Storage.Implementation
                             {
                                 isContentId = false;
                             }
+                            catch (BlobNotFoundException e)
+                            {
+                                unresolvedBlobReferences.Add(e.Blob);
+                            }
                             
                             if (isContentId)
                             {
@@ -187,11 +191,11 @@ namespace Horde.Storage.Implementation
         {
             List<Task<(ContentId, BlobIdentifier[]?)>> pendingContentIdResolves = new();
             List<Task<(BlobIdentifier, bool)>> pendingBlobExistsChecks = new();
-			List<ContentId> unresolvedContentIdReferences = new List<ContentId>();
+            List<ContentId> unresolvedContentIdReferences = new List<ContentId>();
             List<BlobIdentifier> unresolvedBlobReferences = new List<BlobIdentifier>();
 
-			// Resolve all the attachments
-			await foreach (Attachment attachment in GetAttachments(ns, cb))
+            // Resolve all the attachments
+            await foreach (Attachment attachment in GetAttachments(ns, cb))
             {
                 if (attachment is BlobAttachment blobAttachment)
                 {
@@ -204,8 +208,8 @@ namespace Horde.Storage.Implementation
                 }
                 else if (attachment is ObjectAttachment objectAttachment)
                 {
-					// a object just references the same blob, traversing the object attachment is done in GetAttachments
-					pendingBlobExistsChecks.Add(CheckBlobExists(ns, objectAttachment.Identifier));
+                    // a object just references the same blob, traversing the object attachment is done in GetAttachments
+                    pendingBlobExistsChecks.Add(CheckBlobExists(ns, objectAttachment.Identifier));
                 }
                 else
                 {
@@ -230,7 +234,7 @@ namespace Horde.Storage.Implementation
                 }
             }
 
-			// return any verified blobs
+            // return any verified blobs
             foreach (Task<(BlobIdentifier, bool)> pendingBlobExistsTask in pendingBlobExistsChecks)
             {
                 (BlobIdentifier blob, bool exists) = await pendingBlobExistsTask;
@@ -244,8 +248,8 @@ namespace Horde.Storage.Implementation
                 }
             }
 
-			// if there were any content ids we did not recognize we throw a partial reference exception
-			if (unresolvedContentIdReferences.Count != 0)
+            // if there were any content ids we did not recognize we throw a partial reference exception
+            if (unresolvedContentIdReferences.Count != 0)
             {
                 throw new PartialReferenceResolveException(unresolvedContentIdReferences);
             }
@@ -254,7 +258,7 @@ namespace Horde.Storage.Implementation
             {
                 throw new ReferenceIsMissingBlobsException(unresolvedBlobReferences);
             }
-		}
+        }
 
         private async Task<CbObject> ParseCompactBinaryAttachment(NamespaceId ns, BlobIdentifier blobIdentifier)
         {
@@ -277,7 +281,7 @@ namespace Horde.Storage.Implementation
         {
             return (blob, await _blobStore.Exists(ns, blob));
         }
-	}
+    }
 
     public class PartialReferenceResolveException : Exception
     {
