Horde: Fix exceptions sorting streams for summary posts.

diff --git a/Engine/Source/Programs/Horde/Horde.Build/Utilities/StringId.cs b/Engine/Source/Programs/Horde/Horde.Build/Utilities/StringId.cs
--- a/Engine/Source/Programs/Horde/Horde.Build/Utilities/StringId.cs
+++ b/Engine/Source/Programs/Horde/Horde.Build/Utilities/StringId.cs
@@ -20,7 +20,7 @@ namespace Horde.Build.Utilities
 	[JsonSchemaString]
 	[TypeConverter(typeof(StringIdTypeConverter))]
 	[CbConverter(typeof(CbStringIdConverter<>))]
-	public struct StringId<T> : IEquatable<StringId<T>>
+	public struct StringId<T> : IEquatable<StringId<T>>, IComparable<StringId<T>>
 	{
 		/// <summary>
 		/// Empty string
@@ -165,6 +165,12 @@ namespace Horde.Build.Utilities
 			return _text.Equals(other._text, StringComparison.Ordinal);
 		}
 
+		/// <inheritdoc/>
+		public int CompareTo(StringId<T> other)
+		{
+			return String.CompareOrdinal(_text, other._text);
+		}
+
 		/// <inheritdoc/>
 		public override string ToString()
 		{
