Horde: Fix some static analysis warnings.

diff --git a/Engine/Source/Programs/Horde/Horde.Build/Perforce/CommitService.cs b/Engine/Source/Programs/Horde/Horde.Build/Perforce/CommitService.cs
--- a/Engine/Source/Programs/Horde/Horde.Build/Perforce/CommitService.cs
+++ b/Engine/Source/Programs/Horde/Horde.Build/Perforce/CommitService.cs
@@ -756,10 +756,10 @@ namespace Horde.Build.Perforce
 			for (; ; )
 			{
 				// Update the stream, updating the reservation every 30 seconds
-				Task internalTask = Task.Run(() => UpdateStreamContentInternalAsync(streamId, streamChanges, cancellationToken));
+				Task internalTask = Task.Run(() => UpdateStreamContentInternalAsync(streamId, streamChanges, cancellationToken), cancellationToken);
 				while (!internalTask.IsCompleted)
 				{
-					Task delayTask = Task.Delay(TimeSpan.FromSeconds(15.0));
+					Task delayTask = Task.Delay(TimeSpan.FromSeconds(15.0), cancellationToken);
 					if (await Task.WhenAny(internalTask, delayTask) == delayTask)
 					{
 						DateTime newTime = DateTime.UtcNow + TimeSpan.FromSeconds(30.0);
