Horde: Add tighter log scopes for updating job steps to assist debugging.

diff --git a/Engine/Source/Programs/Horde/Horde.Build/Jobs/JobService.cs b/Engine/Source/Programs/Horde/Horde.Build/Jobs/JobService.cs
--- a/Engine/Source/Programs/Horde/Horde.Build/Jobs/JobService.cs
+++ b/Engine/Source/Programs/Horde/Horde.Build/Jobs/JobService.cs
@@ -883,7 +883,7 @@ namespace Horde.Build.Jobs
 			traceScope.Span.SetTag("BatchId", batchId);
 			traceScope.Span.SetTag("StepId", stepId);
 			
-			using IDisposable scope = _logger.BeginScope("UpdateStepAsync({JobId})", job.Id);
+			using IDisposable scope = _logger.BeginScope("UpdateStepAsync({JobId}:{BatchId}:{StepId})", job.Id, batchId, stepId);
 			for (; ;)
 			{
 				IJob? newJob = await TryUpdateStepAsync(job, batchId, stepId, newState, newOutcome, newError, newAbortRequested, newAbortByUserId, newLogId, newNotificationTriggerId, newRetryByUserId, newPriority, newReports, newProperties);
@@ -927,7 +927,7 @@ namespace Horde.Build.Jobs
 			traceScope.Span.SetTag("BatchId", batchId);
 			traceScope.Span.SetTag("StepId", stepId);
 
-			using IDisposable scope = _logger.BeginScope("TryUpdateStepAsync({JobId})", job.Id);
+			using IDisposable scope = _logger.BeginScope("TryUpdateStepAsync({JobId}:{BatchId}:{StepId})", job.Id, batchId, stepId);
 
 			// Get the graph for this job
 			IGraph graph = await GetGraphAsync(job);
