Horde: Include more information in standard log output about issue matching failures.

diff --git a/Engine/Source/Programs/Horde/Horde.Build/Issues/IssueService.cs b/Engine/Source/Programs/Horde/Horde.Build/Issues/IssueService.cs
--- a/Engine/Source/Programs/Horde/Horde.Build/Issues/IssueService.cs
+++ b/Engine/Source/Programs/Horde/Horde.Build/Issues/IssueService.cs
@@ -651,10 +651,10 @@ namespace Horde.Build.Issues
 			// Print the list of new events
 			HashSet<IssueEventGroup> eventGroups = new HashSet<IssueEventGroup>(fingerprintToEventGroup.Values);
 
-			_logger.LogDebug("UpdateCompleteStep({JobId}, {BatchId}, {StepId}): {NumEvents} events, {NumFingerprints} unique fingerprints", job.Id, batch.Id, step.Id, stepEvents.Count, eventGroups.Count);
+			_logger.LogInformation("UpdateCompleteStep({JobId}, {BatchId}, {StepId}): {NumEvents} events, {NumFingerprints} unique fingerprints", job.Id, batch.Id, step.Id, stepEvents.Count, eventGroups.Count);
 			foreach (IssueEventGroup eventGroup in eventGroups)
 			{
-				_logger.LogDebug("Group {Digest}: Type '{FingerprintType}', keys '{FingerprintKeys}'", eventGroup.Digest.ToString(), eventGroup.Fingerprint.Type, String.Join(", ", eventGroup.Fingerprint.Keys));
+				_logger.LogInformation("Group {Digest}: Type '{FingerprintType}', keys '{FingerprintKeys}', {NumEvents} events", eventGroup.Digest.ToString(), eventGroup.Fingerprint.Type, String.Join(", ", eventGroup.Fingerprint.Keys), eventGroup.Events.Count);
 				foreach (IssueEvent eventItem in eventGroup.Events)
 				{
 					_logger.LogDebug("Group {Digest}: [{Line}] {Message}", eventGroup.Digest.ToString(), eventItem.Event.LineIndex, eventItem.EventData.Message);
