Horde: Invite users to the Slack channel when mentioning them in issues.

diff --git a/Engine/Source/Programs/Horde/Horde.Build/Notifications/Sinks/SlackNotificationSink.cs b/Engine/Source/Programs/Horde/Horde.Build/Notifications/Sinks/SlackNotificationSink.cs
--- a/Engine/Source/Programs/Horde/Horde.Build/Notifications/Sinks/SlackNotificationSink.cs
+++ b/Engine/Source/Programs/Horde/Horde.Build/Notifications/Sinks/SlackNotificationSink.cs
@@ -54,6 +54,7 @@ namespace Horde.Build.Notifications.Sinks
 	{
 		const string AddReactionUrl = "https://slack.com/api/reactions.add";
 		const string RemoveReactionUrl = "https://slack.com/api/reactions.remove";
+		const string ConversationsInviteUrl = "https://slack.com/api/conversations.invite";
 		const string PostMessageUrl = "https://slack.com/api/chat.postMessage";
 		const string UpdateMessageUrl = "https://slack.com/api/chat.update";
 		const string GetPermalinkUrl = "https://slack.com/api/chat.getPermalink";
@@ -93,6 +94,15 @@ namespace Horde.Build.Notifications.Sinks
 			public string? Name { get; set; }
 		}
 
+		class InviteMessage
+		{
+			[JsonPropertyName("channel")]
+			public string? Channel { get; set; }
+
+			[JsonPropertyName("users")]
+			public string? Users { get; set; } // Comma separated list of ids
+		}
+
 		class EventPayload
 		{
 			[JsonPropertyName("type")]
@@ -765,6 +775,36 @@ namespace Horde.Build.Notifications.Sinks
 			await SendRequestAsync<SlackResponse>(RemoveReactionUrl, message);
 		}
 
+		async Task InviteUsersAsync(string channel, List<string> userIds)
+		{
+			InviteMessage message = new InviteMessage();
+			message.Channel = channel;
+			message.Users = String.Join(",", userIds);
+			await SendRequestAsync<SlackResponse>(ConversationsInviteUrl, message);
+		}
+
+		async Task InviteUsersAsync(string channel, IEnumerable<UserId> userIds)
+		{
+			List<string> slackUserIds = new List<string>();
+			foreach (UserId userId in userIds)
+			{
+				IUser? user = await _userCollection.GetUserAsync(userId);
+				if (user != null)
+				{
+					string? slackUserId = await GetSlackUserId(user);
+					if (slackUserId != null)
+					{
+						slackUserIds.Add(slackUserId);
+					}
+				}
+			}
+
+			if (slackUserIds.Count > 0)
+			{
+				await InviteUsersAsync(channel, slackUserIds);
+			}
+		}
+
 		static string GetTriageThreadEventId(int issueId) => $"issue_triage_{issueId}";
 
 		async Task CreateOrUpdateWorkflowThreadAsync(IIssue issue, IIssueSpan span, IReadOnlyList<IIssueSpan> spans, WorkflowConfig workflow)
@@ -837,10 +877,12 @@ namespace Horde.Build.Notifications.Sinks
 					}
 
 					// If it has an owner, show that
+					HashSet<UserId> inviteUserIds = new HashSet<UserId>();
 					if (issue.OwnerId != null)
 					{
 						string mention = await FormatMentionAsync(issue.OwnerId.Value, workflow.AllowMentions);
 						await SendMessageToThread(triageChannel, state.Ts, $"Assigned to {mention}");
+						inviteUserIds.Add(issue.OwnerId.Value);
 					}
 					else
 					{
@@ -855,6 +897,7 @@ namespace Horde.Build.Notifications.Sinks
 								string mention = await FormatMentionAsync(suspectGroup.Key, workflow.AllowMentions);
 								string changes = String.Join(", ", suspectGroup.Select(x => FormatChange(x.Change)));
 								suspectList.Add($"{mention} ({changes})");
+								inviteUserIds.Add(suspectGroup.Key);
 							}
 
 							string suspectMessage = $"Possibly {StringUtils.FormatList(suspectList, "or")}.";
@@ -863,6 +906,11 @@ namespace Horde.Build.Notifications.Sinks
 					}
 
 					await SetMessageTimestampAsync(state.Id, state.Channel, state.Ts, permalink);
+
+					if (workflow.AllowMentions)
+					{
+						await InviteUsersAsync(state.Channel, inviteUserIds);
+					}
 				}
 
 				if (issue.AcknowledgedAt != null)
