Horde: Fix filter for stream copies.

diff --git a/Engine/Source/Programs/Horde/Horde.Build/Perforce/CommitService.cs b/Engine/Source/Programs/Horde/Horde.Build/Perforce/CommitService.cs
--- a/Engine/Source/Programs/Horde/Horde.Build/Perforce/CommitService.cs
+++ b/Engine/Source/Programs/Horde/Horde.Build/Perforce/CommitService.cs
@@ -892,7 +892,8 @@ namespace Horde.Build.Perforce
 			clientInfo.Change = -1;
 
 			Utf8String clientRoot = new Utf8String(clientInfo.Client.Root);
-			string queryPath = $"//{clientInfo.Client.Name}/{filter}";
+			string filterOrDefault = filter ?? "...";
+			string queryPath = $"//{clientInfo.Client.Name}/{filterOrDefault}";
 
 			RefId refId = GetRefId(stream.Id, change, filter, revisionsOnly);
 			if (revisionsOnly)
