Horde: Disable invites for dev environment.

diff --git a/Engine/Source/Programs/Horde/Horde.Build/Notifications/Sinks/SlackNotificationSink.cs b/Engine/Source/Programs/Horde/Horde.Build/Notifications/Sinks/SlackNotificationSink.cs
--- a/Engine/Source/Programs/Horde/Horde.Build/Notifications/Sinks/SlackNotificationSink.cs
+++ b/Engine/Source/Programs/Horde/Horde.Build/Notifications/Sinks/SlackNotificationSink.cs
@@ -907,7 +907,7 @@ namespace Horde.Build.Notifications.Sinks
 
 					await SetMessageTimestampAsync(state.Id, state.Channel, state.Ts, permalink);
 
-					if (workflow.AllowMentions)
+					if (_environment.IsProduction() && workflow.AllowMentions)
 					{
 						await InviteUsersAsync(state.Channel, inviteUserIds);
 					}
