Horde Storage - Test for missing blobs for a content ids and compact binary attachments. To be able to do this I had to implement support for writing compressed buffers so added test for that as well.

diff --git a/Engine/Source/Programs/Horde/HordeStorage/Horde.Storage/Implementation/ContentId/CompressedBufferUtils.cs b/Engine/Source/Programs/Horde/HordeStorage/Horde.Storage/Implementation/ContentId/CompressedBufferUtils.cs
--- a/Engine/Source/Programs/Horde/HordeStorage/Horde.Storage/Implementation/ContentId/CompressedBufferUtils.cs
+++ b/Engine/Source/Programs/Horde/HordeStorage/Horde.Storage/Implementation/ContentId/CompressedBufferUtils.cs
@@ -2,8 +2,11 @@
 
 using System;
 using System.Buffers.Binary;
+using System.Collections.Generic;
 using System.IO;
+using System.Linq;
 using Datadog.Trace;
+using EpicGames.Core;
 using Force.Crc32;
 using Jupiter.Implementation;
 using Jupiter.Utils;
@@ -129,6 +132,37 @@ namespace Horde.Storage.Implementation
             return header;
         }
 
+        static void WriteHeader(Header header, BinaryWriter writer)
+        {
+            // the header is always stored big endian
+            bool needsByteSwap = BitConverter.IsLittleEndian;
+            if (needsByteSwap)
+            {
+                header.ByteSwap();
+            }
+
+            writer.Write(header.Magic);
+            writer.Write(header.Crc32);
+            writer.Write((byte)header.Method);
+            writer.Write((byte)header.CompressionLevel);
+            writer.Write((byte)header.CompressionMethodUsed);
+            writer.Write((byte)header.BlockSizeExponent);
+            writer.Write(header.BlockCount);
+            writer.Write(header.TotalRawSize);
+            writer.Write(header.TotalCompressedSize);
+            writer.Write(header.RawHash, 0, 20); // write the first 20 bytes as iohashes are 20 bytes
+            for (int i = 0; i < 12; i++)
+            {
+                // the last 12 bytes should be 0 as they are reserved
+                writer.Write((byte)0);
+            }
+
+            if (needsByteSwap)
+            {
+                header.ByteSwap();
+            }
+        }
+
         public byte[] DecompressContent(byte[] content)
         {
             Header header = ExtractHeader(content);
@@ -238,9 +272,96 @@ namespace Horde.Storage.Implementation
                     throw new NotImplementedException($"Method {header.Method} is not a support value");
             }
         }
+
+        public byte[] CompressContent(OoodleCompressorMethod method, OoodleCompressionLevel compressionLevel, byte[] rawContents)
+        {
+            OodleLZ_Compressor oodleMethod = OodleUtils.ToOodleApiCompressor(method);
+            OodleLZ_CompressionLevel oodleLevel = OodleUtils.ToOodleApiCompressionLevel(compressionLevel);
+
+            const long DefaultBlockSize = 256 * 1024;
+            long blockSize = DefaultBlockSize;
+            long blockCount = (rawContents.LongLength + blockSize - 1) / blockSize;
+
+            byte blockSizeExponent = (byte)Math.Floor(Math.Log2(blockSize));
+
+            Span<byte> contentsSpan = new Span<byte>(rawContents);
+            List<byte[]> compressedBlocks = new List<byte[]>();
+
+            for (int i = 0; i < blockCount; i++)
+            {
+                int rawBlockSize = Math.Min(rawContents.Length - (i * (int)blockSize), (int)blockSize);
+                Span<byte> bufferToCompress = contentsSpan.Slice((int)(i * blockSize), rawBlockSize);
+
+                long encodedSize = _oodleCompressor.Compress(oodleMethod, bufferToCompress.ToArray(), oodleLevel, out byte[] compressedBlock);
+
+                if (encodedSize == 0)
+                {
+                    throw new Exception("Failed to compress content");
+                }
+
+                compressedBlocks.Add(compressedBlock);
+            }
+
+            uint compressedContentLength = (uint)compressedBlocks.Sum(b => b.LongLength);
+
+            Header header = new Header
+            {
+                Magic = Header.ExpectedMagic,
+                Crc32 = 0,
+                Method = Header.CompressionMethod.Oodle,
+                CompressionLevel = (byte)compressionLevel,
+                CompressionMethodUsed = (byte)method,
+                BlockSizeExponent = blockSizeExponent,
+                BlockCount = (uint)blockCount,
+                TotalRawSize = (ulong)rawContents.LongLength,
+                TotalCompressedSize = (ulong)compressedContentLength,
+                RawHash = IoHash.Compute(rawContents).ToByteArray()
+            };
+
+            uint blocksByteUsed = (uint)blockCount * sizeof(uint);
+
+            byte[] headerBuffer = new byte[Header.HeaderLength + blocksByteUsed + compressedContentLength];
+
+            // write the compressed buffer, but with the wrong crc which we update and rewrite later
+            {
+                using MemoryStream ms = new MemoryStream(headerBuffer);
+                using BinaryWriter writer = new BinaryWriter(ms);
+
+                WriteHeader(header, writer);
+
+                for (int i = 0; i < blockCount; i++)
+                {
+                    uint value = (uint)compressedBlocks[i].Length;
+                    if (BitConverter.IsLittleEndian)
+                    {
+                        value = BinaryPrimitives.ReverseEndianness(value);
+                    }
+                    writer.Write(value);
+                }
+
+                for (int i = 0; i < blockCount; i++)
+                {
+                    writer.Write(compressedBlocks[i]);
+                }
+            }
+
+            // calculate the crc from the start of the method field (skipping magic which is a constant and the crc field itself)
+            const int MethodOffset = sizeof(uint) + sizeof(uint);
+            uint calculatedCrc = Crc32Algorithm.Compute(headerBuffer, MethodOffset, (int)(Header.HeaderLength - MethodOffset + blocksByteUsed));
+            header.Crc32 = calculatedCrc;
+
+            // write the header again now that we have the crc
+            {
+                using MemoryStream ms = new MemoryStream(headerBuffer);
+                using BinaryWriter writer = new BinaryWriter(ms);
+                WriteHeader(header, writer);
+            }
+
+            return headerBuffer;
+        }
     }
 
-	public class InvalidHashException : Exception
+    public class InvalidHashException : Exception
     {
         public InvalidHashException(uint headerCrc32, uint calculatedCrc) : base($"Header specified crc \"{headerCrc32}\" but calculated hash was \"{calculatedCrc}\"")
         {
@@ -248,18 +369,18 @@ namespace Horde.Storage.Implementation
         }
     }
 
-	public class InvalidMagicException : Exception
+    public class InvalidMagicException : Exception
     {
         public InvalidMagicException(uint headerMagic, uint expectedMagic) : base($"Header magic \"{headerMagic}\" was incorrect, expected to be {expectedMagic}")
         {
         }
     }
 
-	// from OodleDataCompression.h , we define our own enums for oodle compressions used and convert to the ones expected in the oodle api
+    // from OodleDataCompression.h , we define our own enums for oodle compressions used and convert to the ones expected in the oodle api
 #pragma warning disable CA1028 // Enum Storage should be Int32
-	public enum OoodleCompressorMethod: byte
+    public enum OoodleCompressorMethod: byte
 
-	{
+    {
         NotSet = 0,
         Selkie = 1,
         Mermaid = 2,
diff --git a/Engine/Source/Programs/Horde/HordeStorage/Horde.Storage/Tests/Functional/Objects/ReferencesTests.cs b/Engine/Source/Programs/Horde/HordeStorage/Horde.Storage/Tests/Functional/Objects/ReferencesTests.cs
--- a/Engine/Source/Programs/Horde/HordeStorage/Horde.Storage/Tests/Functional/Objects/ReferencesTests.cs
+++ b/Engine/Source/Programs/Horde/HordeStorage/Horde.Storage/Tests/Functional/Objects/ReferencesTests.cs
@@ -1210,43 +1210,63 @@ namespace Horde.Storage.FunctionalTests.References
         }
 
         [TestMethod]
-        public async Task GetMissingBlobRecord()
+        public async Task GetMissingContentIdRecord()
         {
             string blobContents = "This is a blob";
             byte[] blobData = Encoding.ASCII.GetBytes(blobContents);
-            BlobIdentifier blobHash = BlobIdentifier.FromBlob(blobData);
-            IoHashKey key = IoHashKey.FromName("newReferenceObject");
+            BlobIdentifier uncompressedHash = BlobIdentifier.FromBlob(blobData);
+            IoHashKey key = IoHashKey.FromName("compressedObject");
 
-            using HttpContent requestContent = new ByteArrayContent(blobData);
-            requestContent.Headers.ContentType = new MediaTypeHeaderValue(MediaTypeNames.Application.Octet);
-            requestContent.Headers.Add(CommonHeaders.HashHeaderName, blobHash.ToString());
+            CompressedBufferUtils bufferUtils = Server!.Services.GetService<CompressedBufferUtils>()!;
+            byte[] compressedBuffer = bufferUtils.CompressContent(OoodleCompressorMethod.Mermaid, OoodleCompressionLevel.VeryFast, blobData);
+            BlobIdentifier compressedHash = BlobIdentifier.FromBlob(compressedBuffer);
 
-            HttpResponseMessage result = await Client!.PutAsync(new Uri($"api/v1/refs/{TestNamespace}/bucket/{key}.uecb", UriKind.Relative), requestContent);
-            result.EnsureSuccessStatusCode();
+			CbObject cbObject = CbObject.Build(writer => writer.WriteBinaryAttachment("Attachment", uncompressedHash.AsIoHash()));
+            byte[] cbObjectData = cbObject.GetView().ToArray();
+            BlobIdentifier cbObjectHash = BlobIdentifier.FromBlob(cbObjectData);
 
             {
-                Assert.AreEqual(CustomMediaTypeNames.UnrealCompactBinary, result!.Content.Headers.ContentType!.MediaType);
-                // check that no blobs are missing
-                await using MemoryStream ms = new MemoryStream();
-                await result.Content.CopyToAsync(ms);
-                byte[] roundTrippedBuffer = ms.ToArray();
-                CbObject cb = new CbObject(roundTrippedBuffer);
-                CbField needsField = cb["needs"];
-                List<BlobIdentifier> missingBlobs = needsField.AsArray().Select(field => BlobIdentifier.FromIoHash(field.AsHash())).ToList();
-                Assert.AreEqual(0, missingBlobs.Count);
+                // upload compressed blob
+                using HttpContent requestContent = new ByteArrayContent(compressedBuffer);
+                requestContent.Headers.ContentType = new MediaTypeHeaderValue(CustomMediaTypeNames.UnrealCompressedBuffer);
+
+                HttpResponseMessage result = await Client!.PutAsync(new Uri($"api/v1/compressed-blobs/{TestNamespace}/{uncompressedHash}", UriKind.Relative), requestContent);
+                result.EnsureSuccessStatusCode();
             }
 
-            {
+			{
+				// upload ref
+                using HttpContent requestContent = new ByteArrayContent(cbObjectData);
+                requestContent.Headers.ContentType = new MediaTypeHeaderValue(CustomMediaTypeNames.UnrealCompactBinary);
+                requestContent.Headers.Add(CommonHeaders.HashHeaderName, cbObjectHash.ToString());
+
+                HttpResponseMessage result = await Client!.PutAsync(new Uri($"api/v1/refs/{TestNamespace}/bucket/{key}.uecb", UriKind.Relative), requestContent);
+                result.EnsureSuccessStatusCode();
+
+                {
+                    Assert.AreEqual(CustomMediaTypeNames.UnrealCompactBinary, result!.Content.Headers.ContentType!.MediaType);
+                    // check that no blobs are missing
+                    await using MemoryStream ms = new MemoryStream();
+                    await result.Content.CopyToAsync(ms);
+                    byte[] roundTrippedBuffer = ms.ToArray();
+                    CbObject cb = new CbObject(roundTrippedBuffer);
+                    CbField needsField = cb["needs"];
+                    List<BlobIdentifier> missingBlobs = needsField.AsArray().Select(field => BlobIdentifier.FromIoHash(field.AsHash())).ToList();
+                    Assert.AreEqual(0, missingBlobs.Count);
+                }
+            }
+
+			{
                 // verify we can fetch the blob properly
-                HttpResponseMessage getResponse = await Client.GetAsync(new Uri($"api/v1/refs/{TestNamespace}/bucket/{key}.raw", UriKind.Relative));
+                HttpResponseMessage getResponse = await Client.GetAsync(new Uri($"api/v1/refs/{TestNamespace}/bucket/{key}.uecb", UriKind.Relative));
                 getResponse.EnsureSuccessStatusCode();
                 await using MemoryStream ms = new MemoryStream();
                 await getResponse.Content.CopyToAsync(ms);
 
                 byte[] roundTrippedBuffer = ms.ToArray();
 
-                CollectionAssert.AreEqual(blobData, roundTrippedBuffer);
-                Assert.AreEqual(blobHash, BlobIdentifier.FromBlob(roundTrippedBuffer));
+                CollectionAssert.AreEqual(cbObjectData, roundTrippedBuffer);
+                Assert.AreEqual(cbObjectHash, BlobIdentifier.FromBlob(roundTrippedBuffer));
             }
 
             {
@@ -1257,33 +1277,25 @@ namespace Horde.Storage.FunctionalTests.References
             }
 
             {
-                // delete the blob 
-                await Service.DeleteObject(TestNamespace, blobHash);
+                // verify that the exists check doesnt find any issue
+                HttpResponseMessage existsResponse = await Client.GetAsync(new Uri($"api/v1/refs/{TestNamespace}/exists?names=bucket.{key}", UriKind.Relative));
+                existsResponse.EnsureSuccessStatusCode();
+
+                ExistCheckMultipleRefsResponse response = await existsResponse.Content.ReadAsAsync<ExistCheckMultipleRefsResponse>();
+                Assert.AreEqual(0, response.Missing.Count);
             }
 
-            {
-                // the compact binary object still exists, so this returns a success (trying to resolve the attachment will fail)
-                HttpResponseMessage getResponse = await Client.GetAsync(new Uri($"api/v1/refs/{TestNamespace}/bucket/{key}.uecb", UriKind.Relative));
-                getResponse.EnsureSuccessStatusCode();
+			{
+                // delete the blob referenced by the compressed buffer
+                await Service.DeleteObject(TestNamespace, compressedHash);
             }
 
             {
                 // the compact binary object still exists, so this returns a success (trying to resolve the attachment will fail)
-                HttpResponseMessage getResponse = await Client.GetAsync(new Uri($"api/v1/refs/{TestNamespace}/bucket/{key}.json", UriKind.Relative));
+                HttpResponseMessage getResponse = await Client.GetAsync(new Uri($"api/v1/refs/{TestNamespace}/bucket/{key}.uecb", UriKind.Relative));
                 getResponse.EnsureSuccessStatusCode();
             }
 
-            {
-                // we should now see a 404 as the blob is missing
-                HttpResponseMessage getResponse = await Client.GetAsync(new Uri($"api/v1/refs/{TestNamespace}/bucket/{key}.raw", UriKind.Relative));
-                Assert.AreEqual(HttpStatusCode.NotFound, getResponse.StatusCode);
-                Assert.AreEqual("application/problem+json", getResponse.Content.Headers.ContentType!.MediaType);
-                string s = await getResponse.Content.ReadAsStringAsync();
-                ProblemDetails? problem = JsonConvert.DeserializeObject<ProblemDetails>(s);
-                Assert.IsNotNull(problem);
-                Assert.AreEqual($"Object {blobHash} in {TestNamespace} not found", problem.Title);
-            }
-
             {
                 // verify that head checks fail
                 using HttpRequestMessage headRequest = new HttpRequestMessage(HttpMethod.Head, new Uri($"api/v1/refs/{TestNamespace}/bucket/{key}", UriKind.Relative));
@@ -1303,7 +1315,213 @@ namespace Horde.Storage.FunctionalTests.References
             }
         }
 
-        [TestMethod]
+
+		[TestMethod]
+		public async Task GetMissingCompressedBufferAttachment()
+		{
+            CbObject cbObjectAttachment = CbObject.Build(writer => writer.WriteString("ValueField", "This field has a value"));
+            byte[] cbAttachmentData = cbObjectAttachment.GetView().ToArray();
+            BlobIdentifier cbAttachmentHash = BlobIdentifier.FromBlob(cbAttachmentData);
+			IoHashKey key = IoHashKey.FromName("compressedAttachedObject");
+
+			CbObject cbObject = CbObject.Build(writer => writer.WriteObjectAttachment("Attachment", cbAttachmentHash.AsIoHash()));
+			byte[] cbObjectData = cbObject.GetView().ToArray();
+			BlobIdentifier cbObjectHash = BlobIdentifier.FromBlob(cbObjectData);
+
+			{
+				// upload compressed blob
+				using HttpContent requestContent = new ByteArrayContent(cbAttachmentData);
+				requestContent.Headers.ContentType = new MediaTypeHeaderValue(CustomMediaTypeNames.UnrealCompactBinary);
+
+				HttpResponseMessage result = await Client!.PutAsync(new Uri($"api/v1/objects/{TestNamespace}/{cbAttachmentHash}", UriKind.Relative), requestContent);
+				result.EnsureSuccessStatusCode();
+			}
+
+			{
+				// upload ref
+				using HttpContent requestContent = new ByteArrayContent(cbObjectData);
+				requestContent.Headers.ContentType = new MediaTypeHeaderValue(CustomMediaTypeNames.UnrealCompactBinary);
+				requestContent.Headers.Add(CommonHeaders.HashHeaderName, cbObjectHash.ToString());
+
+				HttpResponseMessage result = await Client!.PutAsync(new Uri($"api/v1/refs/{TestNamespace}/bucket/{key}.uecb", UriKind.Relative), requestContent);
+				result.EnsureSuccessStatusCode();
+
+				{
+					Assert.AreEqual(CustomMediaTypeNames.UnrealCompactBinary, result!.Content.Headers.ContentType!.MediaType);
+					// check that no blobs are missing
+					await using MemoryStream ms = new MemoryStream();
+					await result.Content.CopyToAsync(ms);
+					byte[] roundTrippedBuffer = ms.ToArray();
+					CbObject cb = new CbObject(roundTrippedBuffer);
+					CbField needsField = cb["needs"];
+					List<BlobIdentifier> missingBlobs = needsField.AsArray().Select(field => BlobIdentifier.FromIoHash(field.AsHash())).ToList();
+					Assert.AreEqual(0, missingBlobs.Count);
+				}
+			}
+
+			{
+				// verify we can fetch the blob properly
+				HttpResponseMessage getResponse = await Client.GetAsync(new Uri($"api/v1/refs/{TestNamespace}/bucket/{key}.uecb", UriKind.Relative));
+				getResponse.EnsureSuccessStatusCode();
+				await using MemoryStream ms = new MemoryStream();
+				await getResponse.Content.CopyToAsync(ms);
+
+				byte[] roundTrippedBuffer = ms.ToArray();
+
+				CollectionAssert.AreEqual(cbObjectData, roundTrippedBuffer);
+				Assert.AreEqual(cbObjectHash, BlobIdentifier.FromBlob(roundTrippedBuffer));
+			}
+
+			{
+				// verify that head checks find the object
+				using HttpRequestMessage headRequest = new HttpRequestMessage(HttpMethod.Head, new Uri($"api/v1/refs/{TestNamespace}/bucket/{key}", UriKind.Relative));
+				HttpResponseMessage getResponse = await Client.SendAsync(headRequest);
+				getResponse.EnsureSuccessStatusCode();
+			}
+
+			{
+				// verify that the exists check doesnt find any issue
+				HttpResponseMessage existsResponse = await Client.GetAsync(new Uri($"api/v1/refs/{TestNamespace}/exists?names=bucket.{key}", UriKind.Relative));
+				existsResponse.EnsureSuccessStatusCode();
+
+				ExistCheckMultipleRefsResponse response = await existsResponse.Content.ReadAsAsync<ExistCheckMultipleRefsResponse>();
+				Assert.AreEqual(0, response.Missing.Count);
+			}
+
+			{
+				// delete the blob referenced by the compressed buffer
+				await Service.DeleteObject(TestNamespace, cbAttachmentHash);
+			}
+
+			{
+				// the compact binary object still exists, so this returns a success (trying to resolve the attachment will fail)
+				HttpResponseMessage getResponse = await Client.GetAsync(new Uri($"api/v1/refs/{TestNamespace}/bucket/{key}.uecb", UriKind.Relative));
+				getResponse.EnsureSuccessStatusCode();
+			}
+
+			{
+				// verify that head checks fail
+				using HttpRequestMessage headRequest = new HttpRequestMessage(HttpMethod.Head, new Uri($"api/v1/refs/{TestNamespace}/bucket/{key}", UriKind.Relative));
+				HttpResponseMessage getResponse = await Client.SendAsync(headRequest);
+				Assert.AreEqual(HttpStatusCode.NotFound, getResponse.StatusCode);
+			}
+
+			{
+				// verify that the exists check correctly returns a missing blob
+				HttpResponseMessage existsResponse = await Client.GetAsync(new Uri($"api/v1/refs/{TestNamespace}/exists?names=bucket.{key}", UriKind.Relative));
+				existsResponse.EnsureSuccessStatusCode();
+
+				ExistCheckMultipleRefsResponse response = await existsResponse.Content.ReadAsAsync<ExistCheckMultipleRefsResponse>();
+				Assert.AreEqual(1, response.Missing.Count);
+				Assert.AreEqual(key, response.Missing[0].Key);
+				Assert.AreEqual("bucket", response.Missing[0].Bucket.ToString());
+			}
+		}
+
+		[TestMethod]
+		public async Task GetMissingBlobRecord()
+		{
+			string blobContents = "This is a blob";
+			byte[] blobData = Encoding.ASCII.GetBytes(blobContents);
+			BlobIdentifier blobHash = BlobIdentifier.FromBlob(blobData);
+			IoHashKey key = IoHashKey.FromName("newReferenceObject");
+
+			using HttpContent requestContent = new ByteArrayContent(blobData);
+			requestContent.Headers.ContentType = new MediaTypeHeaderValue(MediaTypeNames.Application.Octet);
+			requestContent.Headers.Add(CommonHeaders.HashHeaderName, blobHash.ToString());
+
+			HttpResponseMessage result = await Client!.PutAsync(new Uri($"api/v1/refs/{TestNamespace}/bucket/{key}.uecb", UriKind.Relative), requestContent);
+			result.EnsureSuccessStatusCode();
+
+			{
+				Assert.AreEqual(CustomMediaTypeNames.UnrealCompactBinary, result!.Content.Headers.ContentType!.MediaType);
+				// check that no blobs are missing
+				await using MemoryStream ms = new MemoryStream();
+				await result.Content.CopyToAsync(ms);
+				byte[] roundTrippedBuffer = ms.ToArray();
+				CbObject cb = new CbObject(roundTrippedBuffer);
+				CbField needsField = cb["needs"];
+				List<BlobIdentifier> missingBlobs = needsField.AsArray().Select(field => BlobIdentifier.FromIoHash(field.AsHash())).ToList();
+				Assert.AreEqual(0, missingBlobs.Count);
+			}
+
+			{
+				// verify we can fetch the blob properly
+				HttpResponseMessage getResponse = await Client.GetAsync(new Uri($"api/v1/refs/{TestNamespace}/bucket/{key}.raw", UriKind.Relative));
+				getResponse.EnsureSuccessStatusCode();
+				await using MemoryStream ms = new MemoryStream();
+				await getResponse.Content.CopyToAsync(ms);
+
+				byte[] roundTrippedBuffer = ms.ToArray();
+
+				CollectionAssert.AreEqual(blobData, roundTrippedBuffer);
+				Assert.AreEqual(blobHash, BlobIdentifier.FromBlob(roundTrippedBuffer));
+			}
+
+			{
+				// verify that head checks find the object
+				using HttpRequestMessage headRequest = new HttpRequestMessage(HttpMethod.Head, new Uri($"api/v1/refs/{TestNamespace}/bucket/{key}", UriKind.Relative));
+				HttpResponseMessage getResponse = await Client.SendAsync(headRequest);
+				getResponse.EnsureSuccessStatusCode();
+			}
+
+			{
+				// verify that the exists check doesnt find any issue
+				HttpResponseMessage existsResponse = await Client.GetAsync(new Uri($"api/v1/refs/{TestNamespace}/exists?names=bucket.{key}", UriKind.Relative));
+				existsResponse.EnsureSuccessStatusCode();
+
+				ExistCheckMultipleRefsResponse response = await existsResponse.Content.ReadAsAsync<ExistCheckMultipleRefsResponse>();
+				Assert.AreEqual(0, response.Missing.Count);
+			}
+
+			{
+				// delete the blob 
+				await Service.DeleteObject(TestNamespace, blobHash);
+			}
+
+			{
+				// the compact binary object still exists, so this returns a success (trying to resolve the attachment will fail)
+				HttpResponseMessage getResponse = await Client.GetAsync(new Uri($"api/v1/refs/{TestNamespace}/bucket/{key}.uecb", UriKind.Relative));
+				getResponse.EnsureSuccessStatusCode();
+			}
+
+			{
+				// the compact binary object still exists, so this returns a success (trying to resolve the attachment will fail)
+				HttpResponseMessage getResponse = await Client.GetAsync(new Uri($"api/v1/refs/{TestNamespace}/bucket/{key}.json", UriKind.Relative));
+				getResponse.EnsureSuccessStatusCode();
+			}
+
+			{
+				// we should now see a 404 as the blob is missing
+				HttpResponseMessage getResponse = await Client.GetAsync(new Uri($"api/v1/refs/{TestNamespace}/bucket/{key}.raw", UriKind.Relative));
+				Assert.AreEqual(HttpStatusCode.NotFound, getResponse.StatusCode);
+				Assert.AreEqual("application/problem+json", getResponse.Content.Headers.ContentType!.MediaType);
+				string s = await getResponse.Content.ReadAsStringAsync();
+				ProblemDetails? problem = JsonConvert.DeserializeObject<ProblemDetails>(s);
+				Assert.IsNotNull(problem);
+				Assert.AreEqual($"Object {blobHash} in {TestNamespace} not found", problem.Title);
+			}
+
+			{
+				// verify that head checks fail
+				using HttpRequestMessage headRequest = new HttpRequestMessage(HttpMethod.Head, new Uri($"api/v1/refs/{TestNamespace}/bucket/{key}", UriKind.Relative));
+				HttpResponseMessage getResponse = await Client.SendAsync(headRequest);
+				Assert.AreEqual(HttpStatusCode.NotFound, getResponse.StatusCode);
+			}
+
+			{
+				// verify that the exists check correctly returns a missing blob
+				HttpResponseMessage existsResponse = await Client.GetAsync(new Uri($"api/v1/refs/{TestNamespace}/exists?names=bucket.{key}", UriKind.Relative));
+				existsResponse.EnsureSuccessStatusCode();
+
+				ExistCheckMultipleRefsResponse response = await existsResponse.Content.ReadAsAsync<ExistCheckMultipleRefsResponse>();
+				Assert.AreEqual(1, response.Missing.Count);
+				Assert.AreEqual(key, response.Missing[0].Key);
+				Assert.AreEqual("bucket", response.Missing[0].Bucket.ToString());
+			}
+		}
+
+		[TestMethod]
         public async Task DeleteObject()
         {
             const string ObjectContents = "This is treated as a opaque blob";
diff --git a/Engine/Source/Programs/Horde/HordeStorage/Horde.Storage/Tests/UnitTests/CompactBinaryTests.cs b/Engine/Source/Programs/Horde/HordeStorage/Horde.Storage/Tests/UnitTests/CompactBinaryTests.cs
--- a/Engine/Source/Programs/Horde/HordeStorage/Horde.Storage/Tests/UnitTests/CompactBinaryTests.cs
+++ b/Engine/Source/Programs/Horde/HordeStorage/Horde.Storage/Tests/UnitTests/CompactBinaryTests.cs
@@ -7,7 +7,7 @@ using EpicGames.Core;
 using EpicGames.Serialization;
 using Microsoft.VisualStudio.TestTools.UnitTesting;
 
-namespace EuropaUnit
+namespace Horde.Storage.Tests.Unit
 {
     [TestClass]
     public class CompactBinaryTests
diff --git a/Engine/Source/Programs/Horde/HordeStorage/Horde.Storage/Tests/UnitTests/CompressedBufferTests.cs b/Engine/Source/Programs/Horde/HordeStorage/Horde.Storage/Tests/UnitTests/CompressedBufferTests.cs
new file mode 100644
--- /dev/null
+++ b/Engine/Source/Programs/Horde/HordeStorage/Horde.Storage/Tests/UnitTests/CompressedBufferTests.cs
@@ -0,0 +1,30 @@
+// Copyright Epic Games, Inc. All Rights Reserved.
+
+using System.Text;
+using Horde.Storage.Implementation;
+using Jupiter.Utils;
+using Microsoft.VisualStudio.TestTools.UnitTesting;
+
+namespace Horde.Storage.Tests.Unit
+{
+    [TestClass]
+    public class CompressedBufferTests
+    {
+
+        [TestMethod]
+        public void CompressAndDecompress()
+        {
+            byte[] bytes = Encoding.UTF8.GetBytes("this is a test string");
+
+            using OodleCompressor compressor = new OodleCompressor();
+            compressor.InitializeOodle();
+            CompressedBufferUtils bufferUtils = new CompressedBufferUtils(compressor);
+
+            byte[] compressedBytes = bufferUtils.CompressContent(OoodleCompressorMethod.Mermaid, OoodleCompressionLevel.VeryFast, bytes);
+
+            byte[] roundTrippedBytes = bufferUtils.DecompressContent(compressedBytes);
+
+            CollectionAssert.AreEqual(bytes, roundTrippedBytes);
+        }
+    }
+}
diff --git a/Engine/Source/Programs/Horde/HordeStorage/Jupiter.Common/Utils/OodleCompressor.cs b/Engine/Source/Programs/Horde/HordeStorage/Jupiter.Common/Utils/OodleCompressor.cs
--- a/Engine/Source/Programs/Horde/HordeStorage/Jupiter.Common/Utils/OodleCompressor.cs
+++ b/Engine/Source/Programs/Horde/HordeStorage/Jupiter.Common/Utils/OodleCompressor.cs
@@ -196,6 +196,16 @@ namespace Jupiter.Utils
             long result = OodleLZ_Compress(compressor, rawBuf, rawBuf.LongLength, compBuf, compressionLevel, 0, null, null, 0, 0);
             return result;
         }
+        public long CompressedBufferSizeNeeded(OodleLZ_Compressor compressor, long rawLength)
+        {
+            if (!_initialized)
+            {
+                throw new InvalidOperationException("Initialize Oodle before using it");
+            }
+
+            long bufferSizeNeeded = OodleLZ_GetCompressedBufferSizeNeeded(compressor, rawLength);
+            return bufferSizeNeeded;
+        }
 
         public long Decompress(byte[] compBuf, long uncompressedSize, out byte[] rawBuf)
         {
