Horde: Add an endpoint to allow Perforce triggers to mark an issue as fixed via a tag.

diff --git a/Engine/Source/Programs/Horde/Horde.Build/Acls/Acl.cs b/Engine/Source/Programs/Horde/Horde.Build/Acls/Acl.cs
--- a/Engine/Source/Programs/Horde/Horde.Build/Acls/Acl.cs
+++ b/Engine/Source/Programs/Horde/Horde.Build/Acls/Acl.cs
@@ -286,6 +286,11 @@ namespace Horde.Build.Acls
         /// </summary>
         ViewIssue,
 
+		/// <summary>
+		/// Mark an issue as fixed via the p4fix endpoint
+		/// </summary>
+		IssueFixViaPerforce,
+
         #endregion
 
         //// TEMPLATES ////
diff --git a/Engine/Source/Programs/Horde/Horde.Build/Issues/IssueMessages.cs b/Engine/Source/Programs/Horde/Horde.Build/Issues/IssueMessages.cs
--- a/Engine/Source/Programs/Horde/Horde.Build/Issues/IssueMessages.cs
+++ b/Engine/Source/Programs/Horde/Horde.Build/Issues/IssueMessages.cs
@@ -931,6 +931,21 @@ namespace Horde.Build.Issues
 		}
 	}
 
+	/// <summary>
+	/// Marks an issue as fixed by another user. Designed for use from a Perforce trigger.
+	/// </summary>
+	public class MarkFixedViaPerforceRequest
+	{
+		/// <summary>
+		/// Name of the user that fixed the issue
+		/// </summary>
+		public string UserName { get; set; } = String.Empty;
+
+		/// <summary>
+		/// Change that fixed the issue
+		/// </summary>
+		public int FixChange { get; set; }
+	}
 
 	/// <summary>
 	/// Request an issue to be created on external issue tracking system
diff --git a/Engine/Source/Programs/Horde/Horde.Build/Issues/IssuesController.cs b/Engine/Source/Programs/Horde/Horde.Build/Issues/IssuesController.cs
--- a/Engine/Source/Programs/Horde/Horde.Build/Issues/IssuesController.cs
+++ b/Engine/Source/Programs/Horde/Horde.Build/Issues/IssuesController.cs
@@ -36,6 +36,7 @@ namespace Horde.Build.Issues
 	[Route("[controller]")]
 	public class IssuesController : HordeControllerBase
 	{
+		private readonly AclService _aclService;
 		private readonly IIssueCollection _issueCollection;
 		private readonly IssueService _issueService;
 		private readonly IExternalIssueService _externalIssueService;
@@ -48,8 +49,9 @@ namespace Horde.Build.Issues
 		/// <summary>
 		/// Constructor
 		/// </summary>
-		public IssuesController(ILogger<IssuesController> logger, IIssueCollection issueCollection, IssueService issueService, JobService jobService, StreamService streamService, IUserCollection userCollection, ILogFileService logFileService, IExternalIssueService externalIssueService)
-		{			
+		public IssuesController(AclService aclService, ILogger<IssuesController> logger, IIssueCollection issueCollection, IssueService issueService, JobService jobService, StreamService streamService, IUserCollection userCollection, ILogFileService logFileService, IExternalIssueService externalIssueService)
+		{
+			_aclService = aclService;
 			_issueCollection = issueCollection;
 			_issueService = issueService;
 			_jobService = jobService;
@@ -302,6 +304,33 @@ namespace Horde.Build.Issues
 			await _issueCollection.GetLogger(issueId).FindAsync(Response.BodyWriter, minTime, maxTime, index, count);
 		}
 
+		/// <summary>
+		/// Hook to allow a Perforce trigger to mark an issue as fixed, via a tag in the changelist description.
+		/// </summary>
+		/// <param name="issueId">Id of the agent to get information about</param>
+		/// <param name="request">Request body</param>
+		/// <returns>Information about the requested agent</returns>
+		[HttpPost]
+		[Route("/api/v1/issues/{issueId}/p4fix")]
+		public async Task<ActionResult> MarkIssueAsFixedViaPerforceAsync(int issueId, [FromBody] MarkFixedViaPerforceRequest request)
+		{
+			if (!await _aclService.AuthorizeAsync(AclAction.IssueFixViaPerforce, User))
+			{
+				return Forbid();
+			}
+
+			IIssueDetails? issue = await _issueService.GetIssueDetailsAsync(issueId);
+			if (issue == null)
+			{
+				return NotFound();
+			}
+
+			IUser user = await _userCollection.FindOrAddUserByLoginAsync(request.UserName);
+			await _issueService.UpdateIssueAsync(issueId, fixChange: request.FixChange, resolvedById: user.Id);
+
+			return Ok();
+		}
+
 		/// <summary>
 		/// Create an issue response object
 		/// </summary>
