Horde: Include a list of open workflows for each issue.

diff --git a/Engine/Source/Programs/Horde/Horde.Build/Issues/IssueMessages.cs b/Engine/Source/Programs/Horde/Horde.Build/Issues/IssueMessages.cs
--- a/Engine/Source/Programs/Horde/Horde.Build/Issues/IssueMessages.cs
+++ b/Engine/Source/Programs/Horde/Horde.Build/Issues/IssueMessages.cs
@@ -762,6 +762,10 @@ namespace Horde.Build.Issues
 		/// </summary>
 		public DateTime? QuarantineTimeUtc { get; set; }
 
+		/// <summary>
+		/// Workflows for which this issue is open
+		/// </summary>
+		public List<WorkflowId> OpenWorkflows { get; set; }
 
 		/// <summary>
 		/// Constructs a new issue
@@ -773,7 +777,8 @@ namespace Horde.Build.Issues
 		/// <param name="quarantinedBy">User that quarantined the issue</param>
 		/// <param name="streamSeverity">The current severity in the stream</param>
 		/// <param name="spans">Spans for this issue</param>
-		public FindIssueResponse(IIssue issue, IUser? owner, IUser? nominatedBy, IUser? resolvedBy, IUser? quarantinedBy, IssueSeverity? streamSeverity, List<FindIssueSpanResponse> spans)
+		/// <param name="openWorkflows">List of workflows for which this issue is open</param>
+		public FindIssueResponse(IIssue issue, IUser? owner, IUser? nominatedBy, IUser? resolvedBy, IUser? quarantinedBy, IssueSeverity? streamSeverity, List<FindIssueSpanResponse> spans, List<WorkflowId> openWorkflows)
 		{
 			Id = issue.Id;
 			CreatedAt = issue.CreatedAt;
@@ -807,6 +812,7 @@ namespace Horde.Build.Issues
 				QuarantinedBy = new GetThinUserInfoResponse(quarantinedBy);
 				QuarantineTimeUtc = issue.QuarantineTimeUtc;
 			}
+			OpenWorkflows = openWorkflows;
 		}
 	}
 
diff --git a/Engine/Source/Programs/Horde/Horde.Build/Issues/IssuesController.cs b/Engine/Source/Programs/Horde/Horde.Build/Issues/IssuesController.cs
--- a/Engine/Source/Programs/Horde/Horde.Build/Issues/IssuesController.cs
+++ b/Engine/Source/Programs/Horde/Horde.Build/Issues/IssuesController.cs
@@ -122,6 +122,8 @@ namespace Horde.Build.Issues
 					{
 						IssueSeverity streamSeverity = IssueSeverity.Unspecified;
 
+						HashSet<WorkflowId> openWorkflowIds = new HashSet<WorkflowId>();
+
 						List<FindIssueSpanResponse> spanResponses = new List<FindIssueSpanResponse>();
 						if (issueIdToSpans.TryGetValue(issue.Id, out List<IIssueSpan>? spansForIssue))
 						{
@@ -141,6 +143,16 @@ namespace Horde.Build.Issues
 							{
 								spanResponses.Add(new FindIssueSpanResponse(span, span.LastFailure.Annotations.WorkflowId));
 							}
+
+							// Find which workflows are affected
+							foreach (IIssueSpan span in spansForIssue)
+							{
+								WorkflowId? workflowId = span.LastFailure.Annotations.WorkflowId;
+								if (workflowId != null && span.NextSuccess == null && span.StreamId == streamId)
+								{
+									openWorkflowIds.Add(workflowId.Value);
+								}
+							}
 						}
 					
 						IUser? owner = null;
@@ -169,7 +181,7 @@ namespace Horde.Build.Issues
 						}
 
 
-						FindIssueResponse response = new FindIssueResponse(issue, owner, nominatedBy, resolvedBy, quarantinedBy, streamSeverity, spanResponses);
+						FindIssueResponse response = new FindIssueResponse(issue, owner, nominatedBy, resolvedBy, quarantinedBy, streamSeverity, spanResponses, openWorkflowIds.ToList());
 						responses.Add(PropertyFilter.Apply(response, filter));
 					}
 				}
