Horde: Prevent users getting blamed for systemic issues.

diff --git a/Engine/Source/Programs/Horde/Horde.Build/Issues/Handlers/SystemicIssueHandler.cs b/Engine/Source/Programs/Horde/Horde.Build/Issues/Handlers/SystemicIssueHandler.cs
--- a/Engine/Source/Programs/Horde/Horde.Build/Issues/Handlers/SystemicIssueHandler.cs
+++ b/Engine/Source/Programs/Horde/Horde.Build/Issues/Handlers/SystemicIssueHandler.cs
@@ -44,6 +44,7 @@ namespace Horde.Build.Issues.Handlers
 		/// <inheritdoc/>
 		public override void RankSuspects(IIssueFingerprint fingerprint, List<SuspectChange> suspects)
 		{
+			suspects.Clear();
 		}
 
 		/// <inheritdoc/>
