Horde: Filter workflow issue in summary view
#rnx
#skipci

diff --git a/Engine/Source/Programs/Horde/HordeDashboard/src/backend/Api.ts b/Engine/Source/Programs/Horde/HordeDashboard/src/backend/Api.ts
--- a/Engine/Source/Programs/Horde/HordeDashboard/src/backend/Api.ts
+++ b/Engine/Source/Programs/Horde/HordeDashboard/src/backend/Api.ts
@@ -4023,6 +4023,9 @@ export type FindIssueResponse = {
 	/** The UTC time when the issue was quarantined */
 	quarantineTimeUtc?: Date | string;
 
+	/** Workflows for which this issue is open */
+	openWorkflows: string[];
+
 }
 
 export type GetAgentSoftwareChannelResponse = {
