Horde: Only put divider after last stream in issue report.

diff --git a/Engine/Source/Programs/Horde/Horde.Build/Notifications/Sinks/SlackNotificationSink.cs b/Engine/Source/Programs/Horde/Horde.Build/Notifications/Sinks/SlackNotificationSink.cs
--- a/Engine/Source/Programs/Horde/Horde.Build/Notifications/Sinks/SlackNotificationSink.cs
+++ b/Engine/Source/Programs/Horde/Horde.Build/Notifications/Sinks/SlackNotificationSink.cs
@@ -1369,6 +1369,8 @@ namespace Horde.Build.Notifications.Sinks
 			{
 				await SendIssueReportForStreamAsync(group.Channel, group.Time, report);
 			}
+
+			await SendMessageAsync(group.Channel, blocks: new[] { new DividerBlock() }, withEnvironment: false);
 		}
 
 		async Task SendIssueReportForStreamAsync(string channel, DateTime time, IssueReport report)
@@ -1467,8 +1469,6 @@ namespace Horde.Build.Notifications.Sinks
 					string header = $"*{totalPct:0.0}%* of build steps ({report.WorkflowStats.NumPassingSteps:n0}/{report.WorkflowStats.NumSteps:n0}) succeeded since last status update.";
 					await SendMessageAsync(channel, text: header, withEnvironment: false);
 				}
-
-				await SendMessageAsync(channel, blocks: new[] { new DividerBlock() }, withEnvironment: false);
 			}
 		}
 
