Horde: Add logging to try to track down missing commit metadata.

diff --git a/Engine/Source/Programs/Horde/Horde.Build/Perforce/CommitService.cs b/Engine/Source/Programs/Horde/Horde.Build/Perforce/CommitService.cs
--- a/Engine/Source/Programs/Horde/Horde.Build/Perforce/CommitService.cs
+++ b/Engine/Source/Programs/Horde/Horde.Build/Perforce/CommitService.cs
@@ -375,6 +375,8 @@ namespace Horde.Build.Perforce
 		/// <returns>List of new commits</returns>
 		public async IAsyncEnumerable<NewCommit> FindCommitsForClusterAsync(string clusterName, Dictionary<IStream, int> streamToFirstChange)
 		{
+			_logger.LogInformation("Replicating metadata for cluster {Name}", clusterName);
+
 			// Create a connection to the server
 			using IPerforceConnection? connection = await _perforceService.GetServiceUserConnection(clusterName);
 			if (connection == null)
@@ -460,6 +462,7 @@ namespace Horde.Build.Perforce
 
 						int originalChange = ParseRobomergeSource(describeRecord.Description) ?? describeRecord.Number;
 
+						_logger.LogInformation("Replicating metadata for CL {Change}: {Description}", changeNumber, StringUtils.Truncate(describeRecord.Description, 40));
 						yield return new NewCommit(stream.Id, describeRecord.Number, originalChange, author.Id, owner.Id, describeRecord.Description, basePath.ToString(), describeRecord.Time);
 					}
 				}
