Horde: Log base changelist for replication.

diff --git a/Engine/Source/Programs/Horde/Horde.Build/Perforce/CommitService.cs b/Engine/Source/Programs/Horde/Horde.Build/Perforce/CommitService.cs
--- a/Engine/Source/Programs/Horde/Horde.Build/Perforce/CommitService.cs
+++ b/Engine/Source/Programs/Horde/Horde.Build/Perforce/CommitService.cs
@@ -767,7 +767,7 @@ namespace Horde.Build.Perforce
 					{
 						DateTime newTime = DateTime.UtcNow + TimeSpan.FromSeconds(30.0);
 						await _redisReservations.AddAsync(streamId, newTime.Ticks);
-						_logger.LogInformation("Extending reservation for tree update of {StreamId} (elapsed: {Time}s)", streamId, (int)timer.Elapsed.TotalSeconds);
+						_logger.LogInformation("Extending reservation for content update of {StreamId} (elapsed: {Time}s)", streamId, (int)timer.Elapsed.TotalSeconds);
 					}
 				}
 
@@ -891,7 +891,7 @@ namespace Horde.Build.Perforce
 			}
 
 			// Apply all the updates
-			_logger.LogInformation("Updating client {Client} to changelist {Change}", clientInfo.Client.Name, change);
+			_logger.LogInformation("Updating client {Client} from changelist {BaseChange} to {Change}", clientInfo.Client.Name, baseChange ?? 0, change);
 			clientInfo.Change = -1;
 
 			Utf8String clientRoot = new Utf8String(clientInfo.Client.Root);
