Horde: More logging to track down missing commits.

diff --git a/Engine/Source/Programs/Horde/Horde.Build/Perforce/CommitService.cs b/Engine/Source/Programs/Horde/Horde.Build/Perforce/CommitService.cs
--- a/Engine/Source/Programs/Horde/Horde.Build/Perforce/CommitService.cs
+++ b/Engine/Source/Programs/Horde/Horde.Build/Perforce/CommitService.cs
@@ -351,6 +351,8 @@ namespace Horde.Build.Perforce
 			await foreach (NewCommit newCommit in FindCommitsForClusterAsync(clusterName, streamToFirstChange))
 			{
 				ICommit commit = await _commitCollection.AddOrReplaceAsync(newCommit);
+				_logger.LogInformation("Replicated CL {Change}: {Description} as {CommitId}", commit.Change, StringUtils.Truncate(commit.Description, 40), commit.Id);
+
 				await NotifyListeners(commit.Id);
 
 				RedisList<int> streamCommitsKey = RedisStreamChanges(commit.StreamId);
@@ -462,7 +464,6 @@ namespace Horde.Build.Perforce
 
 						int originalChange = ParseRobomergeSource(describeRecord.Description) ?? describeRecord.Number;
 
-						_logger.LogInformation("Replicating metadata for CL {Change}: {Description}", changeNumber, StringUtils.Truncate(describeRecord.Description, 40));
 						yield return new NewCommit(stream.Id, describeRecord.Number, originalChange, author.Id, owner.Id, describeRecord.Description, basePath.ToString(), describeRecord.Time);
 					}
 				}
