Horde: Job and step detal navigation updates, with many other optimizations and refactoring
#skipci
#rnx

diff --git a/Engine/Source/Programs/Horde/HordeDashboard/src/backend/Api.ts b/Engine/Source/Programs/Horde/HordeDashboard/src/backend/Api.ts
--- a/Engine/Source/Programs/Horde/HordeDashboard/src/backend/Api.ts
+++ b/Engine/Source/Programs/Horde/HordeDashboard/src/backend/Api.ts
@@ -3010,7 +3010,8 @@ export enum DashboardPreference {
 	ColorError = "ColorError",
 	ColorRunning = "ColorRunning",
 	LocalCache = "LocalCache",
-	LeftAlignLog = "LeftAlignLog"
+	LeftAlignLog = "LeftAlignLog",
+	CompactViews = "CompactViews"
 }
 
 export type DashboardSettings = {
diff --git a/Engine/Source/Programs/Horde/HordeDashboard/src/backend/Backend.ts b/Engine/Source/Programs/Horde/HordeDashboard/src/backend/Backend.ts
--- a/Engine/Source/Programs/Horde/HordeDashboard/src/backend/Backend.ts
+++ b/Engine/Source/Programs/Horde/HordeDashboard/src/backend/Backend.ts
@@ -1226,10 +1226,10 @@ export class Backend {
         });
     }
 
-    getNotification(type: string, jobId: string, labelIdx?: string | null, batchId?: string, stepId?: string | null): Promise<GetNotificationResponse> {
+    getNotification(type: string, jobId: string, labelIdx?: string, batchId?: string, stepId?: string): Promise<GetNotificationResponse> {
 
         let url = `api/v1/jobs/${jobId}`;
-        if (type === "label" && labelIdx != null) {
+        if (type === "label" && labelIdx) {
             url = `${url}/labels/${labelIdx}`;
         }
         else if (type === "step" && batchId && stepId) {
@@ -1244,11 +1244,11 @@ export class Backend {
         });
     }
 
-    updateNotification(request: UpdateNotificationsRequest, type: string, jobId: string, labelIdx?: string | null, batchId?: string, stepId?: string | null): Promise<boolean> {
+    updateNotification(request: UpdateNotificationsRequest, type: string, jobId: string, labelIdx?: string, batchId?: string, stepId?: string | null): Promise<boolean> {
 
         let url = `api/v1/jobs/${jobId}`;
 
-        if (type === "label" && labelIdx != null) {
+        if (type === "label" && labelIdx) {
             url = `${url}/labels/${labelIdx}`;
         }
         else if (type === "step" && batchId && stepId) {
diff --git a/Engine/Source/Programs/Horde/HordeDashboard/src/backend/Dashboard.ts b/Engine/Source/Programs/Horde/HordeDashboard/src/backend/Dashboard.ts
--- a/Engine/Source/Programs/Horde/HordeDashboard/src/backend/Dashboard.ts
+++ b/Engine/Source/Programs/Horde/HordeDashboard/src/backend/Dashboard.ts
@@ -239,11 +239,19 @@ export class Dashboard {
         return this.preferences.get(DashboardPreference.LeftAlignLog) === 'true';
     }
 
+    setCompactViews(value: boolean | undefined) {
+        this.setPreference(DashboardPreference.CompactViews, value ? "true" : "false");
+    }
+
+    get compactViews(): boolean {
+        return this.preferences.get(DashboardPreference.CompactViews) === 'true';
+    }
 
     private hasLoggedLocalCache = false;
 
     get localCache(): boolean {
 
+        /*
         if (this.browser === WebBrowser.Chromium) {
             if (!this.hasLoggedLocalCache) {
                 this.hasLoggedLocalCache = true;
@@ -252,8 +260,16 @@ export class Dashboard {
 
             return true;
         }
+        */
+        
+        const value = this.preferences.get(DashboardPreference.LocalCache) !== 'false';
+        
+        if (value && !this.hasLoggedLocalCache) {
+            this.hasLoggedLocalCache = true;
+            console.log("Local graph and template caching is enabled")
+        }
 
-        return this.preferences.get(DashboardPreference.LocalCache) === 'true';
+        return value;
 
     }
 
diff --git a/Engine/Source/Programs/Horde/HordeDashboard/src/backend/JobDetails.ts b/Engine/Source/Programs/Horde/HordeDashboard/src/backend/JobDetails.ts
--- a/Engine/Source/Programs/Horde/HordeDashboard/src/backend/JobDetails.ts
+++ b/Engine/Source/Programs/Horde/HordeDashboard/src/backend/JobDetails.ts
@@ -19,6 +19,9 @@ export type JobLabel = GetLabelResponse & {
 
 const defaultUpdateMS = 5000;
 
+/**
+ * @deprecated This class evolved over time to be a fragile grab bag of methods related to jobs.  
+ */
 export class JobDetails {
 
     constructor(id?: string, logId?: string, stepId?: string, isLogView?: boolean) {
@@ -760,6 +763,7 @@ export class JobDetails {
             }
 
         });
+        
         this.jobdata?.batches?.forEach(b => {
             b.steps.forEach(s => s.reports?.forEach(r => {
                 if (s.id === this.stepId) {
diff --git a/Engine/Source/Programs/Horde/HordeDashboard/src/backend/JobEventHandler.ts b/Engine/Source/Programs/Horde/HordeDashboard/src/backend/JobEventHandler.ts
--- a/Engine/Source/Programs/Horde/HordeDashboard/src/backend/JobEventHandler.ts
+++ b/Engine/Source/Programs/Horde/HordeDashboard/src/backend/JobEventHandler.ts
@@ -33,7 +33,7 @@ export class JobEventHandler extends PollBase {
     }
 
     clear() {
-        super.clear();
+        super.stop();
         this.details = undefined;
         this.stepEvents.clear();
     }
diff --git a/Engine/Source/Programs/Horde/HordeDashboard/src/backend/Notices.ts b/Engine/Source/Programs/Horde/HordeDashboard/src/backend/Notices.ts
--- a/Engine/Source/Programs/Horde/HordeDashboard/src/backend/Notices.ts
+++ b/Engine/Source/Programs/Horde/HordeDashboard/src/backend/Notices.ts
@@ -13,7 +13,7 @@ class Notices extends PollBase {
     }
 
     clear() {
-        super.clear();
+        super.stop();
     }
 
     async poll(): Promise<void> {
diff --git a/Engine/Source/Programs/Horde/HordeDashboard/src/backend/PollBase.ts b/Engine/Source/Programs/Horde/HordeDashboard/src/backend/PollBase.ts
--- a/Engine/Source/Programs/Horde/HordeDashboard/src/backend/PollBase.ts
+++ b/Engine/Source/Programs/Horde/HordeDashboard/src/backend/PollBase.ts
@@ -10,7 +10,7 @@ export abstract class PollBase {
         this.pollTime = pollTime;
     }
 
-    clear() {
+    stop() {
 
         clearTimeout(this.timeoutId);
         this.timeoutId = undefined;
@@ -85,7 +85,7 @@ export abstract class PollBase {
 
     updating = false;
 
-    private timeoutId: any;
+    protected timeoutId: any;
 
     protected pollTime: number;
 
diff --git a/Engine/Source/Programs/Horde/HordeDashboard/src/base/components/SideRail/SideRail.styles.ts b/Engine/Source/Programs/Horde/HordeDashboard/src/base/components/SideRail/SideRail.styles.ts
--- a/Engine/Source/Programs/Horde/HordeDashboard/src/base/components/SideRail/SideRail.styles.ts
+++ b/Engine/Source/Programs/Horde/HordeDashboard/src/base/components/SideRail/SideRail.styles.ts
@@ -1,4 +1,5 @@
 import { FontWeights, getFocusOutlineStyle, IStyleFunction } from '@fluentui/react/lib/index';
+import { modeColors } from '../../../styles/Styles';
 import { ISideRailStyleProps, ISideRailStyles } from './SideRail.types';
 
 
@@ -16,20 +17,21 @@ export const getStyles: IStyleFunction<ISideRailStyleProps, ISideRailStyles> = p
       marginBottom: appPaddingSm,
       selectors: {
         '&:last-child': {
-          marginBottom: 0
+          marginBottom: 0,
+          'a:active,a:visited,a:hover' : {color: modeColors.text, textDecoration: "none"}
         }
       }
     },
     sectionTitle: {
       fontSize: theme.fonts.mediumPlus.fontSize,
       fontWeight: FontWeights.semibold,
-      color: theme.palette.neutralSecondary,
+      color: modeColors.text,
       marginTop: 0,
       paddingLeft: 8
     },
     links: {
       margin: 0,
-      padding: 0
+      padding: 0      
     },
     linkWrapper: {
       display: 'flex',
@@ -39,9 +41,9 @@ export const getStyles: IStyleFunction<ISideRailStyleProps, ISideRailStyles> = p
           {
             display: 'block',
             flex: '1',
-            padding: '4px 8px',
-            selectors: {
-              '&:hover': { background: theme.palette.neutralLight }
+            padding: '4px 0px',
+            selectors: {              
+              ':active,:visited,:hover,:focus': { background: theme.palette.neutralLight }
             }
           },
           getFocusOutlineStyle(theme, 1)
@@ -53,7 +55,7 @@ export const getStyles: IStyleFunction<ISideRailStyleProps, ISideRailStyles> = p
         'ul li': [
           {
             fontSize: theme.fonts.medium.fontSize,
-            padding: '4px 8px',
+            padding: '4px 0px',
             selectors: {
               '&:hover': { background: theme.palette.neutralLight }
             }
@@ -68,11 +70,12 @@ export const getStyles: IStyleFunction<ISideRailStyleProps, ISideRailStyles> = p
     jumpLink: {
       color: theme.palette.neutralPrimary,
       borderLeft: '2px solid transparent',
-      paddingLeft: 6, // 8px - 2px border
+      paddingLeft: 0, // 8px - 2px border
       selectors: {
         '&:focus': {
-          color: theme.palette.neutralPrimary
-        }
+          color: modeColors.text
+        }, 
+        '&:active,&:visited,&:hover' : {color: modeColors.text, textDecoration: "none"}
       }
     },
     jumpLinkActive: {
diff --git a/Engine/Source/Programs/Horde/HordeDashboard/src/base/components/SideRail/SideRail.types.ts b/Engine/Source/Programs/Horde/HordeDashboard/src/base/components/SideRail/SideRail.types.ts
--- a/Engine/Source/Programs/Horde/HordeDashboard/src/base/components/SideRail/SideRail.types.ts
+++ b/Engine/Source/Programs/Horde/HordeDashboard/src/base/components/SideRail/SideRail.types.ts
@@ -41,6 +41,12 @@ export interface ISideRailLink {
    * - For jump links, just an element ID (no #)
    */
   url: string;
+
+  // url target
+  target?: string;
+
+  /** Whether to use react router for the link */
+  reactRouterLink?: boolean;
 }
 
 export type ISideRailStyleProps = Pick<ISideRailProps, 'theme'>;
diff --git a/Engine/Source/Programs/Horde/HordeDashboard/src/index.css b/Engine/Source/Programs/Horde/HordeDashboard/src/index.css
--- a/Engine/Source/Programs/Horde/HordeDashboard/src/index.css
+++ b/Engine/Source/Programs/Horde/HordeDashboard/src/index.css
@@ -75,11 +75,6 @@ code {
   border-width: 0;
 }
 
-
-div[data-automation-key="ViewLogColumn"] {
-  padding: 0
-}
-
 @keyframes red-pulse {
   0% {
     background-color: #CC0000;
