Horde: Process the list of pending changes queued up by the commit service, even if we don't replicate content.

diff --git a/Engine/Source/Programs/Horde/Horde.Build/Perforce/CommitService.cs b/Engine/Source/Programs/Horde/Horde.Build/Perforce/CommitService.cs
--- a/Engine/Source/Programs/Horde/Horde.Build/Perforce/CommitService.cs
+++ b/Engine/Source/Programs/Horde/Horde.Build/Perforce/CommitService.cs
@@ -219,9 +219,6 @@ namespace Horde.Build.Perforce
 			if (Options.Metadata)
 			{
 				await StartMetadataReplicationAsync();
-			}
-			if (Options.Content)
-			{
 				await StartContentReplicationAsync();
 			}
 		}
@@ -232,9 +229,6 @@ namespace Horde.Build.Perforce
 			if (Options.Content)
 			{
 				await StopContentReplicationAsync();
-			}
-			if (Options.Metadata)
-			{
 				await StopMetadataReplicationAsync();
 			}
 			await StopNotificationsAsync();
@@ -819,14 +813,20 @@ namespace Horde.Build.Perforce
 					prevChange = null;
 				}
 
-				// If we don't have a previous commit tree yet, perform a full snapshot of that change
-				if (prevChange == null)
+				// No-op unless content mirroring is enabled
+				if (Options.Content)
 				{
-					await WriteCommitTreeAsync(stream, values[0], null, cancellationToken);
+					// If we don't have a previous commit tree yet, perform a full snapshot of that change
+					if (prevChange == null)
+					{
+						await WriteCommitTreeAsync(stream, values[0], null, cancellationToken);
+					}
+
+					// Perform a snapshot of the new change, then remove it from the list
+					await WriteCommitTreeAsync(stream, values[1], prevChange, cancellationToken);
 				}
 
-				// Perform a snapshot of the new change, then remove it from the list
-				await WriteCommitTreeAsync(stream, values[1], prevChange, cancellationToken);
+				// Move to the next value
 				prevChange = values[1];
 				await changes.LeftPopAsync();
 			}
