Horde: Various P4 replication changes.

* Added a RefName class for identifying references. Unlike the older RefId class, this supports arbitrary ANSI strings with the exclusion of non-printable and invalid filesystem characters.
* Refactored IBlobStore and ITreeStore to support RefName identifiers.
* Fixed commit replication restarting from a fresh sync whenever a new iteration runs.
* Removed replicated executor and unused content ref id from job task definition.

diff --git a/Engine/Source/Programs/Horde/Horde.Agent/AgentSettings.cs b/Engine/Source/Programs/Horde/Horde.Agent/AgentSettings.cs
--- a/Engine/Source/Programs/Horde/Horde.Agent/AgentSettings.cs
+++ b/Engine/Source/Programs/Horde/Horde.Agent/AgentSettings.cs
@@ -29,11 +29,6 @@ namespace Horde.Agent
 		/// Manage workspaces and run steps in a synced branch
 		/// </summary>
 		Perforce,
-
-		/// <summary>
-		/// Pull replicated data from Horde Storage
-		/// </summary>
-		Replicated,
 	}
 
 	/// <summary>
diff --git a/Engine/Source/Programs/Horde/Horde.Agent/Commands/Bundles/CreateCommand.cs b/Engine/Source/Programs/Horde/Horde.Agent/Commands/Bundles/CreateCommand.cs
--- a/Engine/Source/Programs/Horde/Horde.Agent/Commands/Bundles/CreateCommand.cs
+++ b/Engine/Source/Programs/Horde/Horde.Agent/Commands/Bundles/CreateCommand.cs
@@ -30,7 +30,7 @@ namespace Horde.Agent.Commands.Bundles
 				DirectoryReference.CreateDirectory(_rootDir);
 			}
 
-			FileReference GetRefFile(RefId id) => FileReference.Combine(_rootDir, id.Hash.ToString() + ".ref");
+			FileReference GetRefFile(RefName name) => FileReference.Combine(_rootDir, name.ToString() + ".ref");
 			FileReference GetBlobFile(BlobId id) => FileReference.Combine(_rootDir, id.Inner.ToString() + ".blob");
 
 			static async Task WriteAsync(FileReference file, ReadOnlySequence<byte> sequence, CancellationToken cancellationToken)
@@ -49,6 +49,11 @@ namespace Horde.Agent.Commands.Bundles
 			public async Task<IBlob?> TryReadBlobAsync(BlobId id, CancellationToken cancellationToken = default)
 			{
 				FileReference file = GetBlobFile(id);
+				if(!FileReference.Exists(file))
+				{
+					return null;
+				}
+
 				_logger.LogInformation("Reading {File}", file);
 				byte[] bytes = await FileReference.ReadAllBytesAsync(file, cancellationToken);
 				return BlobUtils.Deserialize(bytes);
@@ -67,27 +72,32 @@ namespace Horde.Agent.Commands.Bundles
 
 			#region Refs
 
-			public Task DeleteRefAsync(RefId id, CancellationToken cancellationToken = default)
+			public Task DeleteRefAsync(RefName name, CancellationToken cancellationToken = default)
 			{
 				throw new NotImplementedException();
 			}
 
-			public Task<bool> HasRefAsync(RefId id, CancellationToken cancellationToken = default)
+			public Task<bool> HasRefAsync(RefName name, CancellationToken cancellationToken = default)
 			{
 				throw new NotImplementedException();
 			}
 
-			public async Task<IBlob?> TryReadRefAsync(RefId id, CancellationToken cancellationToken = default)
+			public async Task<IBlob?> TryReadRefAsync(RefName name, CancellationToken cancellationToken = default)
 			{
-				FileReference file = GetRefFile(id);
+				FileReference file = GetRefFile(name);
+				if(!FileReference.Exists(file))
+				{
+					return null;
+				}
+
 				_logger.LogInformation("Reading {File}", file);
 				byte[] bytes = await FileReference.ReadAllBytesAsync(file, cancellationToken);
 				return BlobUtils.Deserialize(bytes);
 			}
 
-			public async Task WriteRefAsync(RefId id, ReadOnlySequence<byte> data, IReadOnlyList<BlobId> references, CancellationToken cancellationToken = default)
+			public async Task WriteRefAsync(RefName name, ReadOnlySequence<byte> data, IReadOnlyList<BlobId> references, CancellationToken cancellationToken = default)
 			{
-				FileReference file = GetRefFile(id);
+				FileReference file = GetRefFile(name);
 				_logger.LogInformation("Writing {File}", file);
 				await WriteAsync(file, BlobUtils.Serialize(data, references), cancellationToken);
 			}
@@ -95,7 +105,7 @@ namespace Horde.Agent.Commands.Bundles
 			#endregion
 		}
 
-		public static RefId DefaultRefId { get; } = new RefId("default-ref");
+		public static RefName DefaultRefName { get; } = new RefName("default-ref");
 
 		[CommandLine("-StorageDir=", Description = "Overrides the default storage server with a local directory")]
 		public DirectoryReference StorageDir { get; set; } = DirectoryReference.Combine(Program.AppDir, "bundles");
@@ -119,7 +129,7 @@ namespace Horde.Agent.Commands.Bundles
 	class CreateCommand : BundleCommandBase
 	{
 		[CommandLine("-Ref=")]
-		public RefId RefId { get; set; } = DefaultRefId;
+		public RefName RefName { get; set; } = DefaultRefName;
 
 		[CommandLine("-InputDir=", Required = true)]
 		public DirectoryReference InputDir { get; set; } = null!;
@@ -131,7 +141,7 @@ namespace Horde.Agent.Commands.Bundles
 
 			DirectoryNode node = new DirectoryNode();
 			await node.CopyFromDirectoryAsync(InputDir.ToDirectoryInfo(), new ChunkingOptions(), logger, CancellationToken.None);
-			await store.WriteTreeAsync(RefId, node);
+			await store.WriteTreeAsync(RefName, node);
 
 			return 0;
 		}
diff --git a/Engine/Source/Programs/Horde/Horde.Agent/Commands/Bundles/DumpCommand.cs b/Engine/Source/Programs/Horde/Horde.Agent/Commands/Bundles/DumpCommand.cs
--- a/Engine/Source/Programs/Horde/Horde.Agent/Commands/Bundles/DumpCommand.cs
+++ b/Engine/Source/Programs/Horde/Horde.Agent/Commands/Bundles/DumpCommand.cs
@@ -16,7 +16,7 @@ namespace Horde.Agent.Commands.Bundles
 	internal class DumpCommand : BundleCommandBase
 	{
 		[CommandLine("-Ref=")]
-		public RefId RefId { get; set; } = DefaultRefId;
+		public RefName RefName { get; set; } = DefaultRefName;
 
 		[CommandLine("-Blob=")]
 		public BlobId? BlobId { get; set; }
@@ -33,8 +33,8 @@ namespace Horde.Agent.Commands.Bundles
 			}
 			else
 			{
-				blob = await blobStore.ReadRefAsync(RefId, CancellationToken.None);
-				logger.LogInformation("Summary for ref {RefId}", RefId);
+				blob = await blobStore.ReadRefAsync(RefName, CancellationToken.None);
+				logger.LogInformation("Summary for ref {RefId}", RefName);
 			}
 
 			IReadOnlyList<BlobId> references = blob.References;
diff --git a/Engine/Source/Programs/Horde/Horde.Agent/Commands/Bundles/ExtractCommand.cs b/Engine/Source/Programs/Horde/Horde.Agent/Commands/Bundles/ExtractCommand.cs
--- a/Engine/Source/Programs/Horde/Horde.Agent/Commands/Bundles/ExtractCommand.cs
+++ b/Engine/Source/Programs/Horde/Horde.Agent/Commands/Bundles/ExtractCommand.cs
@@ -14,7 +14,7 @@ namespace Horde.Agent.Commands.Bundles
 	internal class ExtractCommand : BundleCommandBase
 	{
 		[CommandLine("-Ref=")]
-		public RefId RefId { get; set; } = DefaultRefId;
+		public RefName RefName { get; set; } = DefaultRefName;
 
 		[CommandLine("-OutputDir=", Required = true)]
 		public DirectoryReference OutputDir { get; set; } = null!;
@@ -24,7 +24,7 @@ namespace Horde.Agent.Commands.Bundles
 			using IMemoryCache cache = new MemoryCache(new MemoryCacheOptions());
 			using ITreeStore store = CreateTreeStore(logger, cache);
 
-			DirectoryNode node = await store.ReadTreeAsync<DirectoryNode>(RefId);
+			DirectoryNode node = await store.ReadTreeAsync<DirectoryNode>(RefName);
 			await node.CopyToDirectoryAsync(OutputDir.ToDirectoryInfo(), logger, CancellationToken.None);
 
 			return 0;
diff --git a/Engine/Source/Programs/Horde/Horde.Agent/Execution/ReplicatedExecutor.cs b/Engine/Source/Programs/Horde/Horde.Agent/Execution/ReplicatedExecutor.cs
deleted file mode 100644
diff --git a/Engine/Source/Programs/Horde/Horde.Build/Commands/Bundles/CreateCommand.cs b/Engine/Source/Programs/Horde/Horde.Build/Commands/Bundles/CreateCommand.cs
--- a/Engine/Source/Programs/Horde/Horde.Build/Commands/Bundles/CreateCommand.cs
+++ b/Engine/Source/Programs/Horde/Horde.Build/Commands/Bundles/CreateCommand.cs
@@ -22,7 +22,7 @@ namespace Horde.Build.Commands.Bundles
 	class CreateCommand : Command
 	{
 		[CommandLine("-Ref=")]
-		public RefId RefId { get; set; } = new RefId("default-ref");
+		public RefName RefName { get; set; } = new RefName("default-ref");
 
 		[CommandLine("-InputDir=", Required = true)]
 		public DirectoryReference InputDir { get; set; } = null!;
@@ -44,7 +44,7 @@ namespace Horde.Build.Commands.Bundles
 
 			DirectoryNode node = new DirectoryNode();
 			await node.CopyFromDirectoryAsync(InputDir.ToDirectoryInfo(), new ChunkingOptions(), logger, CancellationToken.None);
-			await store.WriteTreeAsync(RefId, node);
+			await store.WriteTreeAsync(RefName, node);
 
 			return 0;
 		}
diff --git a/Engine/Source/Programs/Horde/Horde.Build/Commands/Bundles/ExtractCommand.cs b/Engine/Source/Programs/Horde/Horde.Build/Commands/Bundles/ExtractCommand.cs
--- a/Engine/Source/Programs/Horde/Horde.Build/Commands/Bundles/ExtractCommand.cs
+++ b/Engine/Source/Programs/Horde/Horde.Build/Commands/Bundles/ExtractCommand.cs
@@ -17,7 +17,7 @@ namespace Horde.Build.Commands.Bundles
 	internal class ExtractCommand : Command
 	{
 		[CommandLine("-Ref=")]
-		public RefId RefId { get; set; } = new RefId("default-ref");
+		public RefName RefName { get; set; } = new RefName("default-ref");
 
 		[CommandLine("-OutputDir=", Required = true)]
 		public DirectoryReference OutputDir { get; set; } = null!;
@@ -37,7 +37,7 @@ namespace Horde.Build.Commands.Bundles
 
 			ITreeStore store = serviceProvider.GetRequiredService<ITreeStore<CommitService>>();
 
-			DirectoryNode node = await store.ReadTreeAsync<DirectoryNode>(RefId);
+			DirectoryNode node = await store.ReadTreeAsync<DirectoryNode>(RefName);
 			await node.CopyToDirectoryAsync(OutputDir.ToDirectoryInfo(), logger, CancellationToken.None);
 
 			return 0;
diff --git a/Engine/Source/Programs/Horde/Horde.Build/Commands/Bundles/PerforceCommand.cs b/Engine/Source/Programs/Horde/Horde.Build/Commands/Bundles/PerforceCommand.cs
--- a/Engine/Source/Programs/Horde/Horde.Build/Commands/Bundles/PerforceCommand.cs
+++ b/Engine/Source/Programs/Horde/Horde.Build/Commands/Bundles/PerforceCommand.cs
@@ -39,7 +39,7 @@ namespace Horde.Build.Commands.Bundles
 		public int Change { get; set; }
 
 		[CommandLine]
-		public int? BaseChange { get; set; }
+		public int BaseChange { get; set; }
 
 		[CommandLine]
 		public int Count { get; set; } = 1;
@@ -50,14 +50,11 @@ namespace Horde.Build.Commands.Bundles
 		[CommandLine]
 		public bool Compact { get; set; } = true;
 
-		[CommandLine]
-		public bool Clean { get; set; }
-
 		[CommandLine]
 		public string Filter { get; set; } = "...";
 
 		[CommandLine]
-		public bool Metadata { get; set; } = false;
+		public bool RevisionsOnly { get; set; } = false;
 
 		[CommandLine]
 		public DirectoryReference? OutputDir { get; set; }
@@ -89,6 +86,16 @@ namespace Horde.Build.Commands.Bundles
 			Dictionary<IStream, int> streamToFirstChange = new Dictionary<IStream, int>();
 			streamToFirstChange[stream] = Change;
 
+			DirectoryNode baseContents;
+			if (BaseChange == 0)
+			{
+				baseContents = new DirectoryNode();
+			}
+			else
+			{
+				baseContents = await commitService.ReadCommitTreeAsync(stream, BaseChange, Filter, RevisionsOnly, CancellationToken.None);
+			}
+
 			await foreach (NewCommit newCommit in commitService.FindCommitsForClusterAsync(stream.ClusterName, streamToFirstChange).Take(Count))
 			{
 				string briefSummary = newCommit.Description.Replace('\n', ' ');
@@ -98,10 +105,7 @@ namespace Horde.Build.Commands.Bundles
 
 				if (Content)
 				{
-					if (Clean || BaseChange != null)
-					{
-						await commitService.WriteCommitTreeAsync(stream, newCommit.Change, BaseChange, Filter, Metadata, CancellationToken.None);
-					}
+					baseContents = await commitService.WriteCommitTreeAsync(stream, newCommit.Change, BaseChange, baseContents, Filter, RevisionsOnly, CancellationToken.None);
 					BaseChange = newCommit.Change;
 				}
 			}
diff --git a/Engine/Source/Programs/Horde/Horde.Build/Jobs/JobTaskSource.cs b/Engine/Source/Programs/Horde/Horde.Build/Jobs/JobTaskSource.cs
--- a/Engine/Source/Programs/Horde/Horde.Build/Jobs/JobTaskSource.cs
+++ b/Engine/Source/Programs/Horde/Horde.Build/Jobs/JobTaskSource.cs
@@ -886,12 +886,6 @@ namespace Horde.Build.Jobs
 
 			task.Workspace = workspaces.Last();
 
-			IStream? Stream;
-			if (_streams.TryGetValue(job.StreamId, out Stream))
-			{
-				task.ContentRef = _commitService.GetReplicatedContentRef(Stream, job.Change);
-			}
-
 			return task;
 		}
 
diff --git a/Engine/Source/Programs/Horde/Horde.Build/Perforce/CommitService.cs b/Engine/Source/Programs/Horde/Horde.Build/Perforce/CommitService.cs
--- a/Engine/Source/Programs/Horde/Horde.Build/Perforce/CommitService.cs
+++ b/Engine/Source/Programs/Horde/Horde.Build/Perforce/CommitService.cs
@@ -47,11 +47,6 @@ namespace Horde.Build.Perforce
 		/// </summary>
 		public bool Content { get; set; } = true;
 
-		/// <summary>
-		/// Namespace to store content replicated from Perforce
-		/// </summary>
-		public NamespaceId NamespaceId { get; set; } = new NamespaceId("horde.p4");
-
 		/// <summary>
 		/// Options for how objects are packed together
 		/// </summary>
@@ -63,69 +58,6 @@ namespace Horde.Build.Perforce
 		public ChunkingOptions Chunking { get; set; } = new ChunkingOptions();
 	}
 
-	/// <summary>
-	/// Key for a commit object. This is hashed to produce a ref id.
-	/// </summary>
-	public class CommitKey
-	{
-		/// <summary>
-		/// Stream that this commit came from
-		/// </summary>
-		[CbField]
-		public StreamId StreamId { get; set; }
-
-		/// <summary>
-		/// The change being mirrored
-		/// </summary>
-		[CbField]
-		public int Change { get; set; }
-
-		/// <summary>
-		/// Filter for paths in the depot included in this tree
-		/// </summary>
-		[CbField]
-		public string? Filter { get; set; }
-
-		/// <summary>
-		/// Whether this commit contains depot path and revision metadata rather than full file contents
-		/// </summary>
-		[CbField]
-		public bool RevisionsOnly { get; set; }
-
-		/// <summary>
-		/// Default constructor
-		/// </summary>
-		public CommitKey()
-		{
-		}
-
-		/// <summary>
-		/// Constructor
-		/// </summary>
-		public CommitKey(StreamId streamId, int change, string? filter, bool revisionsOnly)
-		{
-			StreamId = streamId;
-			Change = change;
-			Filter = filter;
-			RevisionsOnly = revisionsOnly;
-		}
-
-		/// <summary>
-		/// Gets the ref id for this key
-		/// </summary>
-		/// <returns></returns>
-		public RefId GetRefId()
-		{
-			CbObject obj = CbSerializer.Serialize(this);
-			return new RefId(IoHash.Compute(obj.GetView().Span));
-		}
-
-		/// <summary>
-		/// Gets the ref id for a particular commit
-		/// </summary>
-		public static RefId GetRefId(StreamId streamId, int change, string? filter, bool revisionsOnly) => new CommitKey(streamId, change, filter, revisionsOnly).GetRefId();
-	}
-
 	/// <summary>
 	/// Service which mirrors changes from Perforce
 	/// </summary>
@@ -753,27 +685,36 @@ namespace Horde.Build.Perforce
 			RedisList<int> streamChanges = RedisStreamChanges(streamId);
 			for (; ; )
 			{
-				// Update the stream, updating the reservation every 30 seconds
-				Task internalTask = Task.Run(() => UpdateStreamContentInternalAsync(streamId, streamChanges, cancellationToken), cancellationToken);
-				while (!internalTask.IsCompleted)
-				{
-					Task delayTask = Task.Delay(TimeSpan.FromSeconds(15.0), cancellationToken);
-					if (await Task.WhenAny(internalTask, delayTask) == delayTask)
-					{
-						DateTime newTime = DateTime.UtcNow + TimeSpan.FromSeconds(30.0);
-						await _redisReservations.AddAsync(streamId, newTime.Ticks);
-						_logger.LogInformation("Extending reservation for content update of {StreamId} (elapsed: {Time}s)", streamId, (int)timer.Elapsed.TotalSeconds);
-					}
-				}
-
-				// Log any error during the update
-				if (internalTask.IsFaulted)
+				IStream? stream = await _streamCollection.GetAsync(streamId);
+				if (stream == null || stream.ReplicationMode == ContentReplicationMode.None)
 				{
-					_logger.LogError(internalTask.Exception, "Exception while updating stream content for {StreamId}", streamId);
+					// Remove all but the last item
+					await streamChanges.TrimAsync(0, -2);
 				}
 				else
 				{
-					_logger.LogInformation("Finished update for {StreamId}", streamId);
+					// Update the stream, updating the reservation every 30 seconds
+					Task internalTask = Task.Run(() => UpdateStreamContentInternalAsync(stream, streamChanges, cancellationToken), cancellationToken);
+					while (!internalTask.IsCompleted)
+					{
+						Task delayTask = Task.Delay(TimeSpan.FromSeconds(15.0), cancellationToken);
+						if (await Task.WhenAny(internalTask, delayTask) == delayTask)
+						{
+							DateTime newTime = DateTime.UtcNow + TimeSpan.FromSeconds(30.0);
+							await _redisReservations.AddAsync(streamId, newTime.Ticks);
+							_logger.LogInformation("Extending reservation for content update of {StreamId} (elapsed: {Time}s)", streamId, (int)timer.Elapsed.TotalSeconds);
+						}
+					}
+
+					// Log any error during the update
+					if (internalTask.IsFaulted)
+					{
+						_logger.LogError(internalTask.Exception, "Exception while updating stream content for {StreamId}", streamId);
+					}
+					else
+					{
+						_logger.LogInformation("Finished update for {StreamId}", streamId);
+					}
 				}
 
 				// Remove this stream from the dirty list if it's empty
@@ -788,77 +729,104 @@ namespace Horde.Build.Perforce
 			await _redisReservations.RemoveAsync(streamId);
 		}
 
-		async Task UpdateStreamContentInternalAsync(StreamId streamId, RedisList<int> changes, CancellationToken cancellationToken)
+		async Task UpdateStreamContentInternalAsync(IStream stream, RedisList<int> changes, CancellationToken cancellationToken)
 		{
-			IStream? stream = await _streamCollection.GetAsync(streamId);
-			if (stream == null)
-			{
-				return;
-			}
-
-			int? prevChange = null;
+			int prevChange = 0;
+			DirectoryNode prevContents = new DirectoryNode();
 			for (; ; )
 			{
 				// Get the first two changes to be mirrored. The first one should already exist, unless it's the start of replication for this stream.
 				int[] values = await changes.RangeAsync(0, 1);
-				if (values.Length < 2)
+				if(values.Length == 0)
 				{
 					break;
 				}
 
-				// Check that the previous commit is still valid
-				if (prevChange != null && prevChange.Value != values[0])
+				// Get or add the tree for the first change.
+				if (prevChange != values[0])
 				{
-					_logger.LogInformation("Invalidating previous commit; expected {Change}, actually {ActualChange}", prevChange.Value, values[0]);
-					prevChange = null;
-				}
+					RefName prevRefName = GetRefName(stream, values[0]);
 
-				// No-op unless content mirroring is enabled
-				if (Options.Content)
-				{
-					// If we don't have a previous commit tree yet, perform a full snapshot of that change
-					if (prevChange == null)
+					DirectoryNode? contents = await _treeStore.TryReadTreeAsync<DirectoryNode>(prevRefName, cancellationToken);
+					if (contents == null)
+					{
+						_logger.LogInformation("No content for CL {Change}; creating full snapshot", values[0]);
+						prevContents = await WriteCommitTreeAsync(stream, values[0], 0, new DirectoryNode(), cancellationToken);
+					}
+					else
 					{
-						await WriteCommitTreeAsync(stream, values[0], null, cancellationToken);
+						_logger.LogInformation("Reading existing commit tree for CL {Change} from ref {RefName}", values[0], prevRefName);
+						prevContents = contents;
 					}
 
-					// Perform a snapshot of the new change, then remove it from the list
-					await WriteCommitTreeAsync(stream, values[1], prevChange, cancellationToken);
+					prevChange = values[0];
 				}
 
-				// Move to the next value
+				// If there was only one change, quit now
+				if (values.Length == 1)
+				{
+					break;
+				}
+
+				// Perform a snapshot of the new change, then remove it from the list
+				prevContents = await WriteCommitTreeAsync(stream, values[1], prevChange, prevContents, cancellationToken);
 				prevChange = values[1];
-				await changes.LeftPopAsync();
+
+				// Remove the first item from the list
+				ITransaction transaction = _redis.CreateTransaction();
+				transaction.AddCondition(Condition.ListIndexEqual(changes.Key, 0, values[0]));
+				transaction.AddCondition(Condition.ListIndexEqual(changes.Key, 1, values[1]));
+				_ = transaction.With(changes).LeftPopAsync();
+				await transaction.ExecuteAsync();
 			}
 		}
 
 		/// <summary>
-		/// Replicates the contents of a stream to Horde storage, optionally using the given change as a starting point
+		/// Gets the ref for a particular change
 		/// </summary>
-		/// <param name="stream">The stream to replicate</param>
-		/// <param name="change">Commit to store the tree ref</param>
-		/// <param name="baseChange">The base change to update from</param>
-		/// <param name="cancellationToken">Cancellation token for the operation</param>
-		/// <returns>Root tree object</returns>
-		public async Task WriteCommitTreeAsync(IStream stream, int change, int? baseChange, CancellationToken cancellationToken)
+		/// <param name="streamId"></param>
+		/// <param name="change"></param>
+		/// <param name="filter"></param>
+		/// <param name="revisionsOnly"></param>
+		/// <returns></returns>
+		static RefName GetRefName(StreamId streamId, int change, string? filter, bool revisionsOnly)
 		{
-			if (stream.ReplicationMode != ContentReplicationMode.None)
+			StringBuilder builder = new StringBuilder($"v1/{streamId}/{change}");
+			if (filter != null)
 			{
-				await WriteCommitTreeAsync(stream, change, baseChange, stream.ReplicationFilter, stream.ReplicationMode == ContentReplicationMode.RevisionsOnly, cancellationToken);
+				builder.Append($"_filter_{IoHash.Compute(Encoding.UTF8.GetBytes(filter))}");
 			}
+			if (revisionsOnly)
+			{
+				builder.Append("_revs");
+			}
+			return new RefName(builder.ToString());
 		}
 
 		/// <summary>
 		/// Gets the ref for a particular change
 		/// </summary>
-		/// <param name="streamId"></param>
+		/// <param name="stream"></param>
+		/// <param name="change"></param>
+		/// <returns></returns>
+		static RefName GetRefName(IStream stream, int change)
+		{
+			return GetRefName(stream.Id, change, stream.ReplicationFilter, stream.ReplicationMode == ContentReplicationMode.RevisionsOnly);
+		}
+
+		/// <summary>
+		/// Reads a tree from the given stream
+		/// </summary>
+		/// <param name="stream"></param>
 		/// <param name="change"></param>
 		/// <param name="filter"></param>
 		/// <param name="revisionsOnly"></param>
+		/// <param name="cancellationToken"></param>
 		/// <returns></returns>
-		static RefId GetRefId(StreamId streamId, int change, string? filter, bool revisionsOnly)
+		public async Task<DirectoryNode> ReadCommitTreeAsync(IStream stream, int change, string? filter, bool revisionsOnly, CancellationToken cancellationToken)
 		{
-			return new CommitKey(streamId, change, filter, revisionsOnly).GetRefId();
+			RefName name = GetRefName(stream.Id, change, filter, revisionsOnly);
+			return await _treeStore.ReadTreeAsync<DirectoryNode>(name, cancellationToken);
 		}
 
 		/// <summary>
@@ -867,12 +835,30 @@ namespace Horde.Build.Perforce
 		/// <param name="stream">The stream to replicate</param>
 		/// <param name="change">Commit to store the tree ref</param>
 		/// <param name="baseChange">The base change to update from</param>
-		/// <param name="filter">Depot path to query for changes. Will default to the entire depot (but filtered by the workspace)</param>
-		/// <param name="revisionsOnly">Whether to replicate file revisions only</param>
+		/// <param name="baseContents">Initial contents of the tree at baseChange</param>
 		/// <param name="cancellationToken">Cancellation token for the operation</param>
 		/// <returns>Root tree object</returns>
-		public async Task WriteCommitTreeAsync(IStream stream, int change, int? baseChange, string? filter, bool revisionsOnly, CancellationToken cancellationToken)
+		public async Task<DirectoryNode> WriteCommitTreeAsync(IStream stream, int change, int baseChange, DirectoryNode baseContents, CancellationToken cancellationToken)
 		{
+			bool revisionsOnly = stream.ReplicationMode == ContentReplicationMode.RevisionsOnly;
+			return await WriteCommitTreeAsync(stream, change, baseChange, baseContents, stream.ReplicationFilter, revisionsOnly, cancellationToken);
+		}
+
+		/// <summary>
+		/// Replicates the contents of a stream to Horde storage, optionally using the given change as a starting point
+		/// </summary>
+		/// <param name="stream">The stream to replicate</param>
+		/// <param name="change">Commit to store the tree ref</param>
+		/// <param name="baseChange">The base change to update from</param>
+		/// <param name="baseContents">Initial contents of the tree at baseChange</param>
+		/// <param name="filter"></param>
+		/// <param name="revisionsOnly"></param>
+		/// <param name="cancellationToken">Cancellation token for the operation</param>
+		/// <returns>Root tree object</returns>
+		public async Task<DirectoryNode> WriteCommitTreeAsync(IStream stream, int change, int? baseChange, DirectoryNode baseContents, string? filter, bool revisionsOnly, CancellationToken cancellationToken)
+		{
+			DirectoryNode contents = baseContents;
+
 			// Create a client to replicate from this stream
 			ReplicationClient clientInfo = await FindOrAddReplicationClientAsync(stream);
 
@@ -880,17 +866,6 @@ namespace Horde.Build.Perforce
 			using IPerforceConnection perforce = await PerforceConnection.CreateAsync(clientInfo.Settings, _logger);
 			await FlushWorkspaceAsync(clientInfo, perforce, baseChange ?? 0);
 
-			// Get the initial bundle state
-			DirectoryNode contents;
-			if (baseChange == null)
-			{
-				contents = new DirectoryNode();
-			}
-			else
-			{
-				contents = await _treeStore.ReadTreeAsync<DirectoryNode>(GetRefId(stream.Id, baseChange.Value, filter, revisionsOnly), cancellationToken);
-			}
-
 			// Apply all the updates
 			_logger.LogInformation("Updating client {Client} from changelist {BaseChange} to {Change}", clientInfo.Client.Name, baseChange ?? 0, change);
 			clientInfo.Change = -1;
@@ -899,7 +874,7 @@ namespace Horde.Build.Perforce
 			string filterOrDefault = filter ?? "...";
 			string queryPath = $"//{clientInfo.Client.Name}/{filterOrDefault}";
 
-			RefId refId = GetRefId(stream.Id, change, filter, revisionsOnly);
+			RefName refName = GetRefName(stream.Id, change, filter, revisionsOnly);
 			if (revisionsOnly)
 			{
 				int count = 0;
@@ -927,7 +902,7 @@ namespace Horde.Build.Perforce
 
 					if (++count > 1000)
 					{
-						await _treeStore.WriteTreeAsync(refId, contents, false, cancellationToken);
+						await _treeStore.WriteTreeAsync(refName, contents, false, cancellationToken);
 						count = 0;
 					}
 				}
@@ -980,7 +955,7 @@ namespace Horde.Build.Perforce
 						if (dataSize > trimDataSize)
 						{
 							_logger.LogInformation("Trimming working set after receiving {NumBytes:n0}mb...", dataSize / (1024 * 1024));
-							await _treeStore.WriteTreeAsync(refId, contents, false, cancellationToken);
+							await _treeStore.WriteTreeAsync(refName, contents, false, cancellationToken);
 							GC.Collect();
 							trimDataSize = dataSize + TrimInterval;
 							_logger.LogInformation("Trimming complete. Next trim at {NextNumBytes:n0}mb.", trimDataSize / (1024 * 1024));
@@ -992,21 +967,10 @@ namespace Horde.Build.Perforce
 			clientInfo.Change = change;
 
 			// Return the new root object
-			_logger.LogInformation("Writing ref {RefId} for {StreamId} change {Change}", refId, stream.Id, change);
-			await _treeStore.WriteTreeAsync(refId, contents, true, cancellationToken);
-		}
+			_logger.LogInformation("Writing ref {RefId} for {StreamId} change {Change}", refName, stream.Id, change);
+			await _treeStore.WriteTreeAsync(refName, contents, true, cancellationToken);
 
-		public QualifiedRefId? GetReplicatedContentRef(IStream stream, int change)
-		{
-			if (stream.ReplicationMode == ContentReplicationMode.None)
-			{
-				return null;
-			}
-			else
-			{
-				CommitKey key = new CommitKey(stream.Id, change, stream.ReplicationFilter, stream.ReplicationMode == ContentReplicationMode.RevisionsOnly);
-				return new QualifiedRefId(Options.NamespaceId, GetStreamBucketId(stream.Id), key.GetRefId());
-			}
+			return contents;
 		}
 
 		async Task FlushWorkspaceAsync(ReplicationClient clientInfo, IPerforceConnection perforce, int change)
@@ -1116,16 +1080,6 @@ namespace Horde.Build.Perforce
 			return clientInfo;
 		}
 
-		/// <summary>
-		/// Get the bucket used for particular stream refs
-		/// </summary>
-		/// <param name="streamId"></param>
-		/// <returns></returns>
-		static BucketId GetStreamBucketId(StreamId streamId)
-		{
-			return new BucketId(streamId.ToString());
-		}
-
 		#endregion
 	}
 }
diff --git a/Engine/Source/Programs/Horde/Horde.Build/Perforce/ICommitService.cs b/Engine/Source/Programs/Horde/Horde.Build/Perforce/ICommitService.cs
--- a/Engine/Source/Programs/Horde/Horde.Build/Perforce/ICommitService.cs
+++ b/Engine/Source/Programs/Horde/Horde.Build/Perforce/ICommitService.cs
@@ -7,8 +7,6 @@ using Horde.Build.Utilities;
 
 namespace Horde.Build.Perforce
 {
-	using StreamId = StringId<IStream>;
-
 	/// <summary>
 	/// Provides information about commits
 	/// </summary>
@@ -20,13 +18,5 @@ namespace Horde.Build.Perforce
 		/// <param name="onAddCommit">Callback for a new commit being added</param>
 		/// <returns>Disposable handler.</returns>
 		IDisposable AddListener(Action<ICommit> onAddCommit);
-
-		/// <summary>
-		/// Gets reference to replicated content for the given stream and change
-		/// </summary>
-		/// <param name="stream"></param>
-		/// <param name="change"></param>
-		/// <returns></returns>
-		QualifiedRefId? GetReplicatedContentRef(IStream stream, int change);
 	}
 }
diff --git a/Engine/Source/Programs/Horde/Horde.Build/Storage/Blobs/BlobStore.cs b/Engine/Source/Programs/Horde/Horde.Build/Storage/Blobs/BlobStore.cs
--- a/Engine/Source/Programs/Horde/Horde.Build/Storage/Blobs/BlobStore.cs
+++ b/Engine/Source/Programs/Horde/Horde.Build/Storage/Blobs/BlobStore.cs
@@ -80,30 +80,30 @@ namespace Horde.Build.Storage
 
 		#region Refs
 
-		string GetRefPath(RefId id) => $"{_options.RefPrefix}{id}";
+		string GetRefPath(RefName name) => $"{_options.RefPrefix}{name}";
 
 		/// <inheritdoc/>
-		public async Task DeleteRefAsync(RefId id, CancellationToken cancellationToken = default)
+		public async Task DeleteRefAsync(RefName name, CancellationToken cancellationToken = default)
 		{
-			await _backend.DeleteAsync(GetRefPath(id), cancellationToken);
+			await _backend.DeleteAsync(GetRefPath(name), cancellationToken);
 		}
 
 		/// <inheritdoc/>
-		public async Task<bool> HasRefAsync(RefId id, CancellationToken cancellationToken = default)
+		public async Task<bool> HasRefAsync(RefName name, CancellationToken cancellationToken = default)
 		{
-			return await _backend.ExistsAsync(GetRefPath(id), cancellationToken);
+			return await _backend.ExistsAsync(GetRefPath(name), cancellationToken);
 		}
 
 		/// <inheritdoc/>
-		public async Task<IBlob?> TryReadRefAsync(RefId id, CancellationToken cancellationToken = default)
+		public async Task<IBlob?> TryReadRefAsync(RefName name, CancellationToken cancellationToken = default)
 		{
-			return await TryReadAsync(GetRefPath(id), cancellationToken);
+			return await TryReadAsync(GetRefPath(name), cancellationToken);
 		}
 
 		/// <inheritdoc/>
-		public async Task WriteRefAsync(RefId id, ReadOnlySequence<byte> data, IReadOnlyList<BlobId> references, CancellationToken cancellationToken = default)
+		public async Task WriteRefAsync(RefName name, ReadOnlySequence<byte> data, IReadOnlyList<BlobId> references, CancellationToken cancellationToken = default)
 		{
-			await WriteAsync(GetRefPath(id), data, references, cancellationToken);
+			await WriteAsync(GetRefPath(name), data, references, cancellationToken);
 		}
 
 		#endregion
diff --git a/Engine/Source/Programs/Shared/EpicGames.Horde/Storage/Backends/InMemoryBlobStore.cs b/Engine/Source/Programs/Shared/EpicGames.Horde/Storage/Backends/InMemoryBlobStore.cs
--- a/Engine/Source/Programs/Shared/EpicGames.Horde/Storage/Backends/InMemoryBlobStore.cs
+++ b/Engine/Source/Programs/Shared/EpicGames.Horde/Storage/Backends/InMemoryBlobStore.cs
@@ -22,13 +22,13 @@ namespace EpicGames.Horde.Storage.Backends
 		/// <summary>
 		/// Map of ref name to ref data
 		/// </summary>
-		readonly ConcurrentDictionary<RefId, IBlob> _refs = new ConcurrentDictionary<RefId, IBlob>();
+		readonly ConcurrentDictionary<RefName, IBlob> _refs = new ConcurrentDictionary<RefName, IBlob>();
 
 		/// <inheritdoc cref="_blobs"/>
 		public IReadOnlyDictionary<BlobId, IBlob> Blobs => _blobs;
 
 		/// <inheritdoc cref="_refs"/>
-		public IReadOnlyDictionary<RefId, IBlob> Refs => _refs;
+		public IReadOnlyDictionary<RefName, IBlob> Refs => _refs;
 
 		#region Blobs
 
@@ -52,22 +52,22 @@ namespace EpicGames.Horde.Storage.Backends
 		#region Refs
 
 		/// <inheritdoc/>
-		public Task DeleteRefAsync(RefId id, CancellationToken cancellationToken) => Task.FromResult(_refs.TryRemove(id, out _));
+		public Task DeleteRefAsync(RefName name, CancellationToken cancellationToken) => Task.FromResult(_refs.TryRemove(name, out _));
 
 		/// <inheritdoc/>
-		public Task<bool> HasRefAsync(RefId id, CancellationToken cancellationToken) => Task.FromResult(_refs.ContainsKey(id));
+		public Task<bool> HasRefAsync(RefName name, CancellationToken cancellationToken) => Task.FromResult(_refs.ContainsKey(name));
 
 		/// <inheritdoc/>
-		public Task<IBlob?> TryReadRefAsync(RefId id, CancellationToken cancellationToken)
+		public Task<IBlob?> TryReadRefAsync(RefName name, CancellationToken cancellationToken)
 		{
-			_refs.TryGetValue(id, out IBlob? blob);
+			_refs.TryGetValue(name, out IBlob? blob);
 			return Task.FromResult(blob);
 		}
 
-		/// <inheritdoc/>
-		public Task WriteRefAsync(RefId id, ReadOnlySequence<byte> data, IReadOnlyList<BlobId> references, CancellationToken cancellationToken)
+		/// <inheritdoc/>s
+		public Task WriteRefAsync(RefName name, ReadOnlySequence<byte> data, IReadOnlyList<BlobId> references, CancellationToken cancellationToken)
 		{
-			_refs[id] = BlobUtils.FromMemory(data.ToArray(), references);
+			_refs[name] = BlobUtils.FromMemory(data.ToArray(), references);
 			return Task.CompletedTask;
 		}
 
diff --git a/Engine/Source/Programs/Shared/EpicGames.Horde/Storage/Bundles/Bundle.cs b/Engine/Source/Programs/Shared/EpicGames.Horde/Storage/Bundles/Bundle.cs
--- a/Engine/Source/Programs/Shared/EpicGames.Horde/Storage/Bundles/Bundle.cs
+++ b/Engine/Source/Programs/Shared/EpicGames.Horde/Storage/Bundles/Bundle.cs
@@ -355,12 +355,16 @@ namespace EpicGames.Horde.Storage.Bundles
 		/// Read a bundle from a blob store
 		/// </summary>
 		/// <param name="store">Blob store to read from</param>
-		/// <param name="id">Identifier for the blob</param>
+		/// <param name="name">Identifier for the blob</param>
 		/// <param name="cancellationToken">Cancellation token for the operation</param>
 		/// <returns>The bundle that was read</returns>
-		public static async Task<Bundle> ReadBundleAsync(this IBlobStore store, RefId id, CancellationToken cancellationToken)
+		public static async Task<Bundle?> TryReadBundleAsync(this IBlobStore store, RefName name, CancellationToken cancellationToken)
 		{
-			IBlob blob = await store.ReadRefAsync(id, cancellationToken);
+			IBlob? blob = await store.ReadRefAsync(name, cancellationToken);
+			if (blob == null)
+			{
+				return null;
+			}
 			return ReadBundle(blob);
 		}
 
@@ -368,14 +372,14 @@ namespace EpicGames.Horde.Storage.Bundles
 		/// Writes a bundle to a blob store
 		/// </summary>
 		/// <param name="store">Blob store to write to</param>
-		/// <param name="id">Identifier for the blob</param>
+		/// <param name="name">Identifier for the blob</param>
 		/// <param name="bundle">Bundle to write</param>
 		/// <param name="cancellationToken">Cancellation token for the operation</param>
-		public static Task WriteBundleAsync(this IBlobStore store, RefId id, Bundle bundle, CancellationToken cancellationToken)
+		public static Task WriteBundleAsync(this IBlobStore store, RefName name, Bundle bundle, CancellationToken cancellationToken)
 		{
 			ReadOnlySequence<byte> data = bundle.AsSequence();
 			List<BlobId> imports = bundle.Header.Imports.Select(x => x.BlobId).ToList();
-			return store.WriteRefAsync(id, data, imports, cancellationToken);
+			return store.WriteRefAsync(name, data, imports, cancellationToken);
 		}
 
 		/// <summary>
diff --git a/Engine/Source/Programs/Shared/EpicGames.Horde/Storage/Bundles/BundleStore.cs b/Engine/Source/Programs/Shared/EpicGames.Horde/Storage/Bundles/BundleStore.cs
--- a/Engine/Source/Programs/Shared/EpicGames.Horde/Storage/Bundles/BundleStore.cs
+++ b/Engine/Source/Programs/Shared/EpicGames.Horde/Storage/Bundles/BundleStore.cs
@@ -252,24 +252,32 @@ namespace EpicGames.Horde.Storage.Bundles
 		}
 
 		/// <inheritdoc/>
-		public Task<bool> HasTreeAsync(RefId id, CancellationToken cancellationToken) => _blobStore.HasRefAsync(id, cancellationToken);
+		public Task<bool> HasTreeAsync(RefName name, CancellationToken cancellationToken) => _blobStore.HasRefAsync(name, cancellationToken);
 
 		/// <inheritdoc/>
-		public async Task<ITreeBlob> ReadTreeAsync(RefId id, CancellationToken cancellationToken = default)
+		public async Task<ITreeBlob?> TryReadTreeAsync(RefName name, CancellationToken cancellationToken = default)
 		{
-			Bundle bundle = await _blobStore.ReadBundleAsync(id, cancellationToken);
+			Bundle? bundle = await _blobStore.TryReadBundleAsync(name, cancellationToken);
+			if (bundle == null)
+			{
+				return null;
+			}
 			return MountInMemoryBundle(bundle);
 		}
 
 		/// <inheritdoc/>
-		public async Task<T> ReadTreeAsync<T>(RefId refName, CancellationToken cancellationToken = default) where T : TreeNode
+		public async Task<T?> TryReadTreeAsync<T>(RefName name, CancellationToken cancellationToken = default) where T : TreeNode
 		{
-			ITreeBlob node = await ReadTreeAsync(refName, cancellationToken);
+			ITreeBlob? node = await TryReadTreeAsync(name, cancellationToken);
+			if (node == null)
+			{
+				return null;
+			}
 			return await TreeNode.DeserializeAsync<T>(node, cancellationToken);
 		}
 
 		/// <inheritdoc/>
-		public async Task WriteTreeAsync(RefId id, ITreeBlob root, bool flush = true, CancellationToken cancellationToken = default)
+		public async Task WriteTreeAsync(RefName name, ITreeBlob root, bool flush = true, CancellationToken cancellationToken = default)
 		{
 			// Find all the referenced in-memory nodes from the root
 			List<NodeInfo> nodes = new List<NodeInfo>();
@@ -305,7 +313,7 @@ namespace EpicGames.Horde.Storage.Bundles
 			{
 				// Write the main ref
 				Bundle rootBundle = CreateBundle(bundleNodes);
-				await _blobStore.WriteBundleAsync(id, rootBundle, cancellationToken);
+				await _blobStore.WriteBundleAsync(name, rootBundle, cancellationToken);
 
 				// Copy the stats over
 				Stats = _nextStats;
@@ -314,14 +322,14 @@ namespace EpicGames.Horde.Storage.Bundles
 		}
 
 		/// <inheritdoc/>
-		public async Task WriteTreeAsync<T>(RefId refName, T root, bool flush = true, CancellationToken cancellationToken = default) where T : TreeNode
+		public async Task WriteTreeAsync<T>(RefName name, T root, bool flush = true, CancellationToken cancellationToken = default) where T : TreeNode
 		{
 			ITreeBlob rootBlob = await TreeNode.SerializeAsync(root, this, cancellationToken);
-			await WriteTreeAsync(refName, rootBlob, flush, cancellationToken);
+			await WriteTreeAsync(name, rootBlob, flush, cancellationToken);
 		}
 
 		/// <inheritdoc/>
-		public Task DeleteTreeAsync(RefId id, CancellationToken cancellationToken) => _blobStore.DeleteRefAsync(id, cancellationToken);
+		public Task DeleteTreeAsync(RefName name, CancellationToken cancellationToken) => _blobStore.DeleteRefAsync(name, cancellationToken);
 
 		/// <summary>
 		/// Gets the data for a given node
diff --git a/Engine/Source/Programs/Shared/EpicGames.Horde/Storage/IBlobStore.cs b/Engine/Source/Programs/Shared/EpicGames.Horde/Storage/IBlobStore.cs
--- a/Engine/Source/Programs/Shared/EpicGames.Horde/Storage/IBlobStore.cs
+++ b/Engine/Source/Programs/Shared/EpicGames.Horde/Storage/IBlobStore.cs
@@ -9,6 +9,27 @@ using EpicGames.Core;
 
 namespace EpicGames.Horde.Storage
 {
+	/// <summary>
+	/// Exception for a ref not existing
+	/// </summary>
+	public sealed class RefNameNotFoundException : Exception
+	{
+		/// <summary>
+		/// Name of the missing ref
+		/// </summary>
+		public RefName Name { get; }
+
+		/// <summary>
+		/// Constructor
+		/// </summary>
+		/// <param name="name"></param>
+		public RefNameNotFoundException(RefName name)
+			: base($"Ref name '{name}' not found")
+		{
+			Name = name;
+		}
+	}
+
 	/// <summary>
 	/// Interface for accessing an object's data
 	/// </summary>
@@ -58,35 +79,35 @@ namespace EpicGames.Horde.Storage
 		/// <summary>
 		/// Checks if the given ref exists
 		/// </summary>
-		/// <param name="id">Name of the reference to look for</param>
+		/// <param name="name">Name of the reference to look for</param>
 		/// <param name="cancellationToken">Cancellation token for the operation</param>
 		/// <returns>True if the ref exists, false if it did not exist</returns>
-		Task<bool> HasRefAsync(RefId id, CancellationToken cancellationToken = default);
+		Task<bool> HasRefAsync(RefName name, CancellationToken cancellationToken = default);
 
 		/// <summary>
 		/// Reads data for a ref from the store
 		/// </summary>
-		/// <param name="id">The blob identifier</param>
+		/// <param name="name">The ref name</param>
 		/// <param name="cancellationToken">Cancellation token for the operation</param>
 		/// <returns></returns>
-		Task<IBlob?> TryReadRefAsync(RefId id, CancellationToken cancellationToken = default);
+		Task<IBlob?> TryReadRefAsync(RefName name, CancellationToken cancellationToken = default);
 
 		/// <summary>
-		/// Writes a new blob to the store
+		/// Writes a new ref to the store
 		/// </summary>
-		/// <param name="id">Ref for which the blob is a part</param>
+		/// <param name="name">Ref to write</param>
 		/// <param name="data">Payload for the object</param>
 		/// <param name="references">Object references</param>
 		/// <param name="cancellationToken">Cancellation token for the operation</param>
 		/// <returns>Unique identifier for the blob</returns>
-		Task WriteRefAsync(RefId id, ReadOnlySequence<byte> data, IReadOnlyList<BlobId> references, CancellationToken cancellationToken = default);
+		Task WriteRefAsync(RefName name, ReadOnlySequence<byte> data, IReadOnlyList<BlobId> references, CancellationToken cancellationToken = default);
 
 		/// <summary>
 		/// Reads data for a ref from the store
 		/// </summary>
-		/// <param name="id">The ref identifier</param>
+		/// <param name="name">The ref identifier</param>
 		/// <param name="cancellationToken">Cancellation token for the operation</param>
-		Task DeleteRefAsync(RefId id, CancellationToken cancellationToken = default);
+		Task DeleteRefAsync(RefName name, CancellationToken cancellationToken = default);
 
 		#endregion
 	}
@@ -122,16 +143,16 @@ namespace EpicGames.Horde.Storage
 			#region Refs
 
 			/// <inheritdoc/>
-			public Task DeleteRefAsync(RefId id, CancellationToken cancellationToken = default) => _inner.DeleteRefAsync(id, cancellationToken);
+			public Task DeleteRefAsync(RefName name, CancellationToken cancellationToken = default) => _inner.DeleteRefAsync(name, cancellationToken);
 
 			/// <inheritdoc/>
-			public Task<bool> HasRefAsync(RefId id, CancellationToken cancellationToken = default) => _inner.HasRefAsync(id, cancellationToken);
+			public Task<bool> HasRefAsync(RefName name, CancellationToken cancellationToken = default) => _inner.HasRefAsync(name, cancellationToken);
 
 			/// <inheritdoc/>
-			public Task<IBlob?> TryReadRefAsync(RefId id, CancellationToken cancellationToken = default) => _inner.TryReadRefAsync(id, cancellationToken);
+			public Task<IBlob?> TryReadRefAsync(RefName name, CancellationToken cancellationToken = default) => _inner.TryReadRefAsync(name, cancellationToken);
 
 			/// <inheritdoc/>
-			public Task WriteRefAsync(RefId id, ReadOnlySequence<byte> data, IReadOnlyList<BlobId> references, CancellationToken cancellationToken = default) => _inner.WriteRefAsync(id, data, references, cancellationToken);
+			public Task WriteRefAsync(RefName name, ReadOnlySequence<byte> data, IReadOnlyList<BlobId> references, CancellationToken cancellationToken = default) => _inner.WriteRefAsync(name, data, references, cancellationToken);
 
 			#endregion
 		}
@@ -164,15 +185,15 @@ namespace EpicGames.Horde.Storage
 		/// Reads a ref from the store, throwing an exception if it does not exist
 		/// </summary>
 		/// <param name="store">The store instance to read from</param>
-		/// <param name="id">Id for the ref</param>
+		/// <param name="name">Id for the ref</param>
 		/// <param name="cancellationToken">Cancellation token for the operation</param>
 		/// <returns>The blob instance</returns>
-		public static async Task<IBlob> ReadRefAsync(this IBlobStore store, RefId id, CancellationToken cancellationToken = default)
+		public static async Task<IBlob> ReadRefAsync(this IBlobStore store, RefName name, CancellationToken cancellationToken = default)
 		{
-			IBlob? blob = await store.TryReadRefAsync(id, cancellationToken);
+			IBlob? blob = await store.TryReadRefAsync(name, cancellationToken);
 			if (blob == null)
 			{
-				throw new KeyNotFoundException($"Unable to find ref {id}");
+				throw new RefNameNotFoundException(name);
 			}
 			return blob;
 		}
diff --git a/Engine/Source/Programs/Shared/EpicGames.Horde/Storage/ITreeStore.cs b/Engine/Source/Programs/Shared/EpicGames.Horde/Storage/ITreeStore.cs
--- a/Engine/Source/Programs/Shared/EpicGames.Horde/Storage/ITreeStore.cs
+++ b/Engine/Source/Programs/Shared/EpicGames.Horde/Storage/ITreeStore.cs
@@ -52,53 +52,53 @@ namespace EpicGames.Horde.Storage
 		/// <summary>
 		/// Tests whether a tree exists with the given name
 		/// </summary>
-		/// <param name="id">Name of the ref used to store the tree</param>
+		/// <param name="name">Name of the ref used to store the tree</param>
 		/// <param name="cancellationToken">Cancellation token for the operation</param>
 		/// <returns>Tree reader instance</returns>
-		Task<bool> HasTreeAsync(RefId id, CancellationToken cancellationToken = default);
+		Task<bool> HasTreeAsync(RefName name, CancellationToken cancellationToken = default);
 
 		/// <summary>
 		/// Deletes a tree with the given name
 		/// </summary>
-		/// <param name="id">Name of the ref used to store the tree</param>
+		/// <param name="name">Name of the ref used to store the tree</param>
 		/// <param name="cancellationToken">Cancellation token for the operation</param>
-		Task DeleteTreeAsync(RefId id, CancellationToken cancellationToken = default);
+		Task DeleteTreeAsync(RefName name, CancellationToken cancellationToken = default);
 
 		/// <summary>
 		/// Reads the root blob from a tree in storage
 		/// </summary>
-		/// <param name="id">Identifier for the tree</param>
+		/// <param name="name">Identifier for the tree</param>
 		/// <param name="cancellationToken">Cancellation token for the operation</param>
 		/// <returns></returns>
-		Task<ITreeBlob> ReadTreeAsync(RefId id, CancellationToken cancellationToken = default);
+		Task<ITreeBlob?> TryReadTreeAsync(RefName name, CancellationToken cancellationToken = default);
 
 		/// <summary>
 		/// Reads a root node from the store
 		/// </summary>
-		/// <param name="id">Name of the tree to fetch</param>
+		/// <param name="name">Name of the tree to fetch</param>
 		/// <param name="cancellationToken">Cancellation token for the operation</param>
 		/// <returns>Deserialized node stored with this name</returns>
-		Task<T> ReadTreeAsync<T>(RefId id, CancellationToken cancellationToken = default) where T : TreeNode;
+		Task<T?> TryReadTreeAsync<T>(RefName name, CancellationToken cancellationToken = default) where T : TreeNode;
 
 		/// <summary>
 		/// Creates a writer for a tree with the given name
 		/// </summary>
-		/// <param name="id">Name of the ref used to store the tree</param>
+		/// <param name="name">Name of the ref used to store the tree</param>
 		/// <param name="root">Node to flush</param>
 		/// <param name="flush">Whether to flush the complete tree state. If false, an implementation may choose to buffer some writes to allow packing multiple nodes together.</param>
 		/// <param name="cancellationToken">Cancellation token for the operation</param>
 		/// <returns>Tree writer instance</returns>
-		Task WriteTreeAsync(RefId id, ITreeBlob root, bool flush = true, CancellationToken cancellationToken = default);
+		Task WriteTreeAsync(RefName name, ITreeBlob root, bool flush = true, CancellationToken cancellationToken = default);
 
 		/// <summary>
 		/// Creates a writer for a tree with the given name
 		/// </summary>
-		/// <param name="id">Name of the ref used to store the tree</param>
+		/// <param name="name">Name of the ref used to store the tree</param>
 		/// <param name="root">Node to flush</param>
 		/// <param name="flush">Whether to flush the complete tree state. If false, an implementation may choose to buffer some writes to allow packing multiple nodes together.</param>
 		/// <param name="cancellationToken">Cancellation token for the operation</param>
 		/// <returns>Tree writer instance</returns>
-		Task WriteTreeAsync<T>(RefId id, T root, bool flush = true, CancellationToken cancellationToken = default) where T : TreeNode;
+		Task WriteTreeAsync<T>(RefName name, T root, bool flush = true, CancellationToken cancellationToken = default) where T : TreeNode;
 	}
 
 	/// <summary>
@@ -123,25 +123,25 @@ namespace EpicGames.Horde.Storage
 			}
 
 			/// <inheritdoc/>
-			public Task DeleteTreeAsync(RefId id, CancellationToken cancellationToken = default) => _inner.DeleteTreeAsync(id, cancellationToken);
+			public Task DeleteTreeAsync(RefName name, CancellationToken cancellationToken = default) => _inner.DeleteTreeAsync(name, cancellationToken);
 
 			/// <inheritdoc/>
 			public void Dispose() => _inner.Dispose();
 
 			/// <inheritdoc/>
-			public Task<bool> HasTreeAsync(RefId id, CancellationToken cancellationToken = default) => _inner.HasTreeAsync(id, cancellationToken);
+			public Task<bool> HasTreeAsync(RefName name, CancellationToken cancellationToken = default) => _inner.HasTreeAsync(name, cancellationToken);
 
 			/// <inheritdoc/>
-			public Task<ITreeBlob> ReadTreeAsync(RefId id, CancellationToken cancellationToken = default) => _inner.ReadTreeAsync(id, cancellationToken);
+			public Task<ITreeBlob?> TryReadTreeAsync(RefName name, CancellationToken cancellationToken = default) => _inner.TryReadTreeAsync(name, cancellationToken);
 
 			/// <inheritdoc/>
-			public Task<TNode> ReadTreeAsync<TNode>(RefId id, CancellationToken cancellationToken = default) where TNode : TreeNode => _inner.ReadTreeAsync<TNode>(id, cancellationToken);
+			public Task<TNode?> TryReadTreeAsync<TNode>(RefName name, CancellationToken cancellationToken = default) where TNode : TreeNode => _inner.TryReadTreeAsync<TNode>(name, cancellationToken);
 
 			/// <inheritdoc/>
-			public Task WriteTreeAsync(RefId id, ITreeBlob root, bool flush = true, CancellationToken cancellationToken = default) => _inner.WriteTreeAsync(id, root, flush, cancellationToken);
+			public Task WriteTreeAsync(RefName name, ITreeBlob root, bool flush = true, CancellationToken cancellationToken = default) => _inner.WriteTreeAsync(name, root, flush, cancellationToken);
 
 			/// <inheritdoc/>
-			public Task WriteTreeAsync<TNode>(RefId id, TNode root, bool flush = true, CancellationToken cancellationToken = default) where TNode : TreeNode => _inner.WriteTreeAsync<TNode>(id, root, flush, cancellationToken);
+			public Task WriteTreeAsync<TNode>(RefName name, TNode root, bool flush = true, CancellationToken cancellationToken = default) where TNode : TreeNode => _inner.WriteTreeAsync<TNode>(name, root, flush, cancellationToken);
 		}
 
 		/// <summary>
@@ -154,5 +154,39 @@ namespace EpicGames.Horde.Storage
 		{
 			return new TypedTreeStore<T>(store);
 		}
+
+		/// <summary>
+		/// Reads a tree, throwing an exception if it doesn't exist
+		/// </summary>
+		/// <param name="store">Store to read from</param>
+		/// <param name="name">Name of the ref</param>
+		/// <param name="cancellationToken">Cancellation token for the operation</param>
+		/// <returns></returns>
+		public static async Task<ITreeBlob> ReadTreeAsync(this ITreeStore store, RefName name, CancellationToken cancellationToken = default)
+		{
+			ITreeBlob? blob = await store.TryReadTreeAsync(name, cancellationToken);
+			if (blob == null)
+			{
+				throw new RefNameNotFoundException(name);
+			}
+			return blob;
+		}
+
+		/// <summary>
+		/// Reads a tree, throwing an exception if it doesn't exist
+		/// </summary>
+		/// <param name="store">Store to read from</param>
+		/// <param name="name">Name of the ref</param>
+		/// <param name="cancellationToken">Cancellation token for the operation</param>
+		/// <returns></returns>
+		public static async Task<T> ReadTreeAsync<T>(this ITreeStore store, RefName name, CancellationToken cancellationToken = default) where T : TreeNode
+		{
+			T? node = await store.TryReadTreeAsync<T>(name, cancellationToken);
+			if (node == null)
+			{
+				throw new RefNameNotFoundException(name);
+			}
+			return node;
+		}
 	}
 }
diff --git a/Engine/Source/Programs/Shared/EpicGames.Horde/Storage/RefName.cs b/Engine/Source/Programs/Shared/EpicGames.Horde/Storage/RefName.cs
new file mode 100644
--- /dev/null
+++ b/Engine/Source/Programs/Shared/EpicGames.Horde/Storage/RefName.cs
@@ -0,0 +1,160 @@
+// Copyright Epic Games, Inc. All Rights Reserved.
+
+using EpicGames.Core;
+using EpicGames.Serialization;
+using System;
+using System.ComponentModel;
+using System.Globalization;
+using System.Text;
+using System.Text.Json;
+using System.Text.Json.Serialization;
+
+namespace EpicGames.Horde.Storage
+{
+	/// <summary>
+	/// Identifier for a storage namespace
+	/// </summary>
+	[JsonConverter(typeof(RefNameJsonConverter))]
+	[TypeConverter(typeof(RefNameTypeConverter))]
+	[CbConverter(typeof(RefNameCbConverter))]
+	public struct RefName : IEquatable<RefName>, IComparable<RefName>
+	{
+		/// <summary>
+		/// String for the ref name
+		/// </summary>
+		public Utf8String Text { get; }
+
+		/// <summary>
+		/// Constructor
+		/// </summary>
+		/// <param name="text"></param>
+		public RefName(Utf8String text)
+		{
+			Text = text;
+
+			for (int idx = 0; idx < text.Length; idx++)
+			{
+				if (!IsValidCharacter(text[idx]))
+				{
+					throw new ArgumentException($"Character '{(char)text[idx]}' is not valid in ref name ('{text}')", nameof(text));
+				}
+			}
+
+			if (text.Length == 0)
+			{
+				throw new ArgumentException($"Ref name cannot be empty", nameof(text));
+			}
+		}
+
+		/// <summary>
+		/// Constructor
+		/// </summary>
+		/// <param name="text"></param>
+		public RefName(string text)
+			: this(new Utf8String(text))
+		{
+		}
+
+		/// <summary>
+		/// Test whether a character is valid in a ref name
+		/// </summary>
+		/// <param name="c">Character to test</param>
+		/// <returns></returns>
+		public static bool IsValidCharacter(byte c)
+		{
+			switch (c)
+			{
+				case (byte)'<':
+				case (byte)'>':
+				case (byte)':':
+				case (byte)'"':
+				case (byte)'\\':
+				case (byte)'|':
+				case (byte)'*':
+				case (byte)'?':
+					return false;
+				default:
+					return c > 0x20 && c < 0x7f;
+			}
+		}
+
+		/// <inheritdoc/>
+		public override bool Equals(object? obj) => obj is RefName refId && Equals(refId);
+
+		/// <inheritdoc/>
+		public override int GetHashCode() => Text.GetHashCode();
+
+		/// <inheritdoc/>
+		public bool Equals(RefName refName) => Text == refName.Text;
+
+		/// <inheritdoc/>
+		public int CompareTo(RefName other) => Text.CompareTo(other.Text);
+
+		/// <inheritdoc/>
+		public override string ToString() => Text.ToString();
+
+		/// <inheritdoc/>
+		public static bool operator ==(RefName lhs, RefName rhs) => lhs.Equals(rhs);
+
+		/// <inheritdoc/>
+		public static bool operator !=(RefName lhs, RefName rhs) => !lhs.Equals(rhs);
+
+		/// <inheritdoc/>
+		public static bool operator <(RefName lhs, RefName rhs) => lhs.CompareTo(rhs) < 0;
+
+		/// <inheritdoc/>
+		public static bool operator <=(RefName lhs, RefName rhs) => lhs.CompareTo(rhs) <= 0;
+
+		/// <inheritdoc/>
+		public static bool operator >(RefName lhs, RefName rhs) => lhs.CompareTo(rhs) > 0;
+
+		/// <inheritdoc/>
+		public static bool operator >=(RefName lhs, RefName rhs) => lhs.CompareTo(rhs) >= 0;
+
+	}
+
+	/// <summary>
+	/// Type converter for IoHash to and from JSON
+	/// </summary>
+	sealed class RefNameJsonConverter : JsonConverter<RefName>
+	{
+		/// <inheritdoc/>
+		public override RefName Read(ref Utf8JsonReader reader, Type typeToConvert, JsonSerializerOptions options) => new RefName(new Utf8String(reader.ValueSpan.ToArray()));
+
+		/// <inheritdoc/>
+		public override void Write(Utf8JsonWriter writer, RefName value, JsonSerializerOptions options) => writer.WriteStringValue(value.Text.Span);
+	}
+
+	/// <summary>
+	/// Type converter from strings to IoHash objects
+	/// </summary>
+	sealed class RefNameTypeConverter : TypeConverter
+	{
+		/// <inheritdoc/>
+		public override bool CanConvertFrom(ITypeDescriptorContext? context, Type sourceType)
+		{
+			return sourceType == typeof(string);
+		}
+
+		/// <inheritdoc/>
+		public override object ConvertFrom(ITypeDescriptorContext? context, CultureInfo? culture, object? value)
+		{
+			return new RefName((string)value!);
+		}
+	}
+
+	/// <summary>
+	/// Type converter to compact binary
+	/// </summary>
+	sealed class RefNameCbConverter : CbConverterBase<RefName>
+	{
+		/// <inheritdoc/>
+		public override RefName Read(CbField field) => new RefName(field.AsUtf8String());
+
+		/// <inheritdoc/>
+		public override void Write(CbWriter writer, RefName value) => writer.WriteUtf8StringValue(value.Text);
+
+		/// <inheritdoc/>
+		public override void WriteNamed(CbWriter writer, Utf8String name, RefName value) => writer.WriteUtf8String(name, value.Text);
+	}
+}
