Horde: Build fixes.

diff --git a/Engine/Source/Programs/Horde/Horde.Agent/Services/WorkerService.cs b/Engine/Source/Programs/Horde/Horde.Agent/Services/WorkerService.cs
--- a/Engine/Source/Programs/Horde/Horde.Agent/Services/WorkerService.cs
+++ b/Engine/Source/Programs/Horde/Horde.Agent/Services/WorkerService.cs
@@ -289,8 +289,6 @@ namespace Horde.Agent.Services
 						new LocalExecutor(rpcClient, executeTask.JobId, executeTask.BatchId, batch.AgentType, _settings.LocalExecutor),
 					ExecutorType.Perforce => (rpcClient, executeTask, batch) =>
 						new PerforceExecutor(rpcClient, executeTask.JobId, executeTask.BatchId, batch.AgentType, executeTask.AutoSdkWorkspace, executeTask.Workspace, _workingDir),
-					ExecutorType.Replicated => (rpcClient, executeTask, batch) =>
-						new ReplicatedExecutor(rpcClient, executeTask.JobId, executeTask.BatchId, batch.AgentType, _storageClient, executeTask.ContentRef, _workingDir),
 
 					_ => throw new InvalidDataException($"Unknown executor type '{_settings.Executor}'")
 				};
