Horde: Add another polling iteration after mirroring a change. New changes may have been added in the meantime.

diff --git a/Engine/Source/Programs/Horde/Horde.Build/Perforce/CommitService.cs b/Engine/Source/Programs/Horde/Horde.Build/Perforce/CommitService.cs
--- a/Engine/Source/Programs/Horde/Horde.Build/Perforce/CommitService.cs
+++ b/Engine/Source/Programs/Horde/Horde.Build/Perforce/CommitService.cs
@@ -761,23 +761,24 @@ namespace Horde.Build.Perforce
 
 					prevChange = values[0];
 				}
-
-				// If there was only one change, quit now
-				if (values.Length == 1)
+				else if (values.Length == 2)
+				{
+					// Perform a snapshot of the new change, then remove it from the list
+					prevContents = await WriteCommitTreeAsync(stream, values[1], prevChange, prevContents, cancellationToken);
+					prevChange = values[1];
+
+					// Remove the first item from the list
+					ITransaction transaction = _redis.CreateTransaction();
+					transaction.AddCondition(Condition.ListIndexEqual(changes.Key, 0, values[0]));
+					transaction.AddCondition(Condition.ListIndexEqual(changes.Key, 1, values[1]));
+					_ = transaction.With(changes).LeftPopAsync();
+					await transaction.ExecuteAsync();
+				}
+				else
 				{
+					// Nothing to do
 					break;
 				}
-
-				// Perform a snapshot of the new change, then remove it from the list
-				prevContents = await WriteCommitTreeAsync(stream, values[1], prevChange, prevContents, cancellationToken);
-				prevChange = values[1];
-
-				// Remove the first item from the list
-				ITransaction transaction = _redis.CreateTransaction();
-				transaction.AddCondition(Condition.ListIndexEqual(changes.Key, 0, values[0]));
-				transaction.AddCondition(Condition.ListIndexEqual(changes.Key, 1, values[1]));
-				_ = transaction.With(changes).LeftPopAsync();
-				await transaction.ExecuteAsync();
 			}
 		}
 
