Horde: Fix running steps not being updated correctly if the containing batch terminates abruptly.

diff --git a/Engine/Source/Programs/Horde/Horde.Build/Jobs/JobCollection.cs b/Engine/Source/Programs/Horde/Horde.Build/Jobs/JobCollection.cs
--- a/Engine/Source/Programs/Horde/Horde.Build/Jobs/JobCollection.cs
+++ b/Engine/Source/Programs/Horde/Horde.Build/Jobs/JobCollection.cs
@@ -714,7 +714,21 @@ namespace Horde.Build.Jobs
 				bool bUpdateState = false;
 				foreach (JobStepDocument step in batch.Steps)
 				{
-					if (step.State == JobStepState.Ready || step.State == JobStepState.Waiting)
+					if (step.State == JobStepState.Running)
+					{
+						step.State = JobStepState.Completed;
+						step.Outcome = JobStepOutcome.Failure;
+						step.Error = JobStepError.Incomplete;
+
+						if (CanRetryNode(jobDocument, batch.GroupIdx, step.NodeIdx))
+						{
+							step.Retry = true;
+							retriedNodes.Add(new NodeRef(batch.GroupIdx, step.NodeIdx));
+						}
+
+						bUpdateState = true;
+					}
+					else if (step.State == JobStepState.Ready || step.State == JobStepState.Waiting)
 					{
 						if (CanRetryNode(jobDocument, batch.GroupIdx, step.NodeIdx))
 						{
@@ -848,6 +862,9 @@ namespace Horde.Build.Jobs
 							// Update the retry flag
 							if (newRetryByUserId != null && step.RetriedByUserId == null)
 							{
+								step.Retry = true;
+								updates.Add(updateBuilder.Set(x => x.Batches[batchIdx].Steps[stepIdx].Retry, true));
+
 								step.RetriedByUserId = newRetryByUserId;
 								updates.Add(updateBuilder.Set(x => x.Batches[batchIdx].Steps[stepIdx].RetriedByUserId, step.RetriedByUserId));
 
@@ -1385,9 +1402,12 @@ namespace Horde.Build.Jobs
 			// Remove all the nodes which have already succeeded
 			foreach (JobStepBatchDocument batch in job.Batches)
 			{
-				foreach (IJobStep step in batch.Steps)
+				foreach (JobStepDocument step in batch.Steps)
 				{
-					if ((step.State == JobStepState.Running && step.RetriedByUserId == null) || (step.State == JobStepState.Completed && step.RetriedByUserId == null) || (step.State == JobStepState.Aborted && step.RetriedByUserId == null) || step.State == JobStepState.Skipped)
+					if ((step.State == JobStepState.Running && !step.Retry) 
+						|| (step.State == JobStepState.Completed && !step.Retry) 
+						|| (step.State == JobStepState.Aborted && !step.Retry) 
+						|| (step.State == JobStepState.Skipped))
 					{
 						newNodesToExecute.Remove(graph.Groups[batch.GroupIdx].Nodes[step.NodeIdx]);
 					}
diff --git a/Engine/Source/Programs/Horde/Horde.Build/Jobs/JobService.cs b/Engine/Source/Programs/Horde/Horde.Build/Jobs/JobService.cs
--- a/Engine/Source/Programs/Horde/Horde.Build/Jobs/JobService.cs
+++ b/Engine/Source/Programs/Horde/Horde.Build/Jobs/JobService.cs
@@ -799,7 +799,7 @@ namespace Horde.Build.Jobs
 				JobStepBatchError? newError = null;
 				if (newState == JobStepBatchState.Complete && batch.AgentId != null)
 				{
-					if (batch.Steps.Any(x => x.State == JobStepState.Waiting || x.State == JobStepState.Ready))
+					if (batch.Steps.Any(x => x.State == JobStepState.Waiting || x.State == JobStepState.Ready || x.State == JobStepState.Running))
 					{
 						// Mark the batch as incomplete
 						newError = JobStepBatchError.Incomplete;
