Horde: Missing files.

diff --git a/Engine/Source/Programs/Horde/Horde.Build.Tests/JobCollectionTests.cs b/Engine/Source/Programs/Horde/Horde.Build.Tests/JobCollectionTests.cs
--- a/Engine/Source/Programs/Horde/Horde.Build.Tests/JobCollectionTests.cs
+++ b/Engine/Source/Programs/Horde/Horde.Build.Tests/JobCollectionTests.cs
@@ -248,5 +248,84 @@ namespace Horde.Build.Tests
 			Assert.AreEqual(1, job.Batches[2].Steps.Count);
 			Assert.AreEqual(JobStepState.Skipped, job.Batches[2].Steps[0].State);
 		}
+
+		[TestMethod]
+		public async Task IncompleteBatchRunningAsync()
+		{
+			// Same as IncompleteBatchAsync, but with steps moved to running state
+
+			Mock<ITemplate> templateMock = new Mock<ITemplate>(MockBehavior.Strict);
+			templateMock.SetupGet(x => x.InitialAgentType).Returns((string?)null);
+
+			IGraph baseGraph = await GraphCollection.AddAsync(templateMock.Object);
+
+			List<string> arguments = new List<string>();
+			arguments.Add("-Target=Step 1");
+			arguments.Add("-Target=Step 3");
+
+			IJob job = await JobCollection.AddAsync(JobId.GenerateNewId(), new StreamId("ue4-main"), new TemplateRefId("test-build"), ContentHash.SHA1("hello"), baseGraph, "Test job", 123, 123, null, null, null, null, null, null, null, null, null, false, false, null, null, arguments);
+			Assert.AreEqual(1, job.Batches.Count);
+
+			// First retry
+
+			job = await StartBatch(job, baseGraph, 0);
+			job = Deref(await JobCollection.TryUpdateStepAsync(job, baseGraph, job.Batches[0].Id, job.Batches[0].Steps[0].Id, newState: JobStepState.Running));
+			job = Deref(await JobCollection.TryUpdateBatchAsync(job, baseGraph, job.Batches[0].Id, null, JobStepBatchState.Complete, JobStepBatchError.Incomplete));
+
+			job = (await JobCollection.GetAsync(job.Id))!;
+			Assert.AreEqual(2, job.Batches.Count);
+
+			Assert.AreEqual(JobStepBatchState.Complete, job.Batches[0].State);
+			Assert.AreEqual(1, job.Batches[0].Steps.Count);
+			Assert.AreEqual(JobStepState.Completed, job.Batches[0].Steps[0].State);
+			Assert.AreEqual(JobStepError.Incomplete, job.Batches[0].Steps[0].Error);
+
+			Assert.AreEqual(JobStepBatchState.Ready, job.Batches[1].State);
+			Assert.AreEqual(1, job.Batches[1].Steps.Count);
+
+			// Second retry
+
+			job = Deref(await JobCollection.TryUpdateStepAsync(job, baseGraph, job.Batches[1].Id, job.Batches[1].Steps[0].Id, newState: JobStepState.Running));
+			job = Deref(await JobCollection.TryUpdateBatchAsync(job, baseGraph, job.Batches[1].Id, null, JobStepBatchState.Complete, JobStepBatchError.Incomplete));
+
+			job = (await JobCollection.GetAsync(job.Id))!;
+			Assert.AreEqual(3, job.Batches.Count);
+
+			Assert.AreEqual(JobStepBatchState.Complete, job.Batches[0].State);
+			Assert.AreEqual(1, job.Batches[0].Steps.Count);
+			Assert.AreEqual(JobStepState.Completed, job.Batches[0].Steps[0].State);
+			Assert.AreEqual(JobStepError.Incomplete, job.Batches[0].Steps[0].Error);
+
+			Assert.AreEqual(JobStepBatchState.Complete, job.Batches[1].State);
+			Assert.AreEqual(1, job.Batches[1].Steps.Count);
+			Assert.AreEqual(JobStepState.Completed, job.Batches[1].Steps[0].State);
+			Assert.AreEqual(JobStepError.Incomplete, job.Batches[1].Steps[0].Error);
+
+			Assert.AreEqual(JobStepBatchState.Ready, job.Batches[2].State);
+			Assert.AreEqual(1, job.Batches[2].Steps.Count);
+
+			// Check it doesn't retry a third time
+
+			job = Deref(await JobCollection.TryUpdateStepAsync(job, baseGraph, job.Batches[2].Id, job.Batches[2].Steps[0].Id, newState: JobStepState.Running));
+			job = Deref(await JobCollection.TryUpdateBatchAsync(job, baseGraph, job.Batches[2].Id, null, JobStepBatchState.Complete, JobStepBatchError.Incomplete));
+
+			job = (await JobCollection.GetAsync(job.Id))!;
+			Assert.AreEqual(3, job.Batches.Count);
+
+			Assert.AreEqual(JobStepBatchState.Complete, job.Batches[0].State);
+			Assert.AreEqual(1, job.Batches[0].Steps.Count);
+			Assert.AreEqual(JobStepState.Completed, job.Batches[0].Steps[0].State);
+			Assert.AreEqual(JobStepError.Incomplete, job.Batches[0].Steps[0].Error);
+
+			Assert.AreEqual(JobStepBatchState.Complete, job.Batches[1].State);
+			Assert.AreEqual(1, job.Batches[1].Steps.Count);
+			Assert.AreEqual(JobStepState.Completed, job.Batches[1].Steps[0].State);
+			Assert.AreEqual(JobStepError.Incomplete, job.Batches[1].Steps[0].Error);
+
+			Assert.AreEqual(JobStepBatchState.Complete, job.Batches[2].State);
+			Assert.AreEqual(1, job.Batches[2].Steps.Count);
+			Assert.AreEqual(JobStepState.Completed, job.Batches[2].Steps[0].State);
+			Assert.AreEqual(JobStepError.Incomplete, job.Batches[2].Steps[0].Error);
+		}
 	}
 }
diff --git a/Engine/Source/Programs/Horde/Horde.Build/Jobs/IJob.cs b/Engine/Source/Programs/Horde/Horde.Build/Jobs/IJob.cs
--- a/Engine/Source/Programs/Horde/Horde.Build/Jobs/IJob.cs
+++ b/Engine/Source/Programs/Horde/Horde.Build/Jobs/IJob.cs
@@ -156,7 +156,7 @@ namespace Horde.Build.Jobs
 	}
 
 	/// <summary>
-	/// Systemic error codes for a job failing
+	/// Systemic error codes for a job step failing
 	/// </summary>
 	public enum JobStepError
 	{
@@ -173,7 +173,12 @@ namespace Horde.Build.Jobs
 		/// <summary>
 		/// Step is in is paused state so was skipped
 		/// </summary>
-		Paused
+		Paused,
+
+		/// <summary>
+		/// Step did not complete because the batch exited
+		/// </summary>
+		Incomplete
 	}
 
 	/// <summary>
