Horde: Sync commits in multiple passes, so we can resume if the pod is terminated.

diff --git a/Engine/Source/Programs/Horde/Horde.Build/Perforce/CommitService.cs b/Engine/Source/Programs/Horde/Horde.Build/Perforce/CommitService.cs
--- a/Engine/Source/Programs/Horde/Horde.Build/Perforce/CommitService.cs
+++ b/Engine/Source/Programs/Horde/Horde.Build/Perforce/CommitService.cs
@@ -1,9 +1,11 @@
 // Copyright Epic Games, Inc. All Rights Reserved.
 
 using System;
+using System.Buffers;
 using System.Collections.Generic;
 using System.Diagnostics;
 using System.Globalization;
+using System.IO;
 using System.Linq;
 using System.Text;
 using System.Text.RegularExpressions;
@@ -58,6 +60,86 @@ namespace Horde.Build.Perforce
 		public ChunkingOptions Chunking { get; set; } = new ChunkingOptions();
 	}
 
+	/// <summary>
+	/// Root node for a commit snapshot
+	/// </summary>
+	[TreeSerializer(typeof(CommitNodeSerializer))]
+	public class CommitNode : TreeNode
+	{
+		/// <summary>
+		/// Paths that have been synced. Empty once complete.
+		/// </summary>
+		public List<Utf8String> Paths { get; }
+
+		/// <summary>
+		/// Contents of this snapshot
+		/// </summary>
+		public TreeNodeRef<DirectoryNode> Contents { get; }
+
+		/// <summary>
+		/// Constructor
+		/// </summary>
+		public CommitNode()
+			: this(new DirectoryNode())
+		{
+		}
+
+		/// <summary>
+		/// Constructor
+		/// </summary>
+		public CommitNode(DirectoryNode contents)
+		{
+			Paths = new List<Utf8String>();
+			Contents = new TreeNodeRef<DirectoryNode>(this, contents);
+		}
+
+		private CommitNode(IEnumerable<Utf8String> paths, ITreeBlob contents)
+		{
+			Paths = new List<Utf8String>(paths);
+			Contents = new TreeNodeRef<DirectoryNode>(this, contents);
+		}
+
+		internal async ValueTask<ITreeBlob> SerializeAsync(ITreeBlobWriter writer, CancellationToken cancellationToken)
+		{
+			List<ITreeBlob> references = new List<ITreeBlob>();
+			references.Add(await Contents.CollapseAsync(writer, cancellationToken));
+
+			ByteArrayBuilder builder = new ByteArrayBuilder();
+			builder.WriteVariableLengthArray(Paths, x => builder.WriteUtf8String(x));
+			return await writer.WriteBlobAsync(builder.AsSequence(), references, cancellationToken);
+		}
+
+		internal static async ValueTask<CommitNode> DeserializeAsync(ITreeBlob blob, CancellationToken cancellationToken)
+		{
+			IReadOnlyList<ITreeBlob> children = await blob.GetReferencesAsync(cancellationToken);
+			ReadOnlySequence<byte> data = await blob.GetDataAsync(cancellationToken);
+
+			MemoryReader reader = new MemoryReader(data.AsSingleSegment());
+			Utf8String[] paths = reader.ReadVariableLengthArray(() => reader.ReadUtf8String());
+
+			return new CommitNode(paths, blob);
+		}
+
+		/// <inheritdoc/>
+		public override IReadOnlyList<TreeNodeRef> GetReferences()
+		{
+			List<TreeNodeRef> refs = new List<TreeNodeRef>();
+			refs.Add(Contents);
+			return refs;
+		}
+	}
+
+	class CommitNodeSerializer : TreeNodeSerializer<CommitNode>
+	{
+		/// <inheritdoc/>
+		public override ValueTask<CommitNode> DeserializeAsync(ITreeBlob node, CancellationToken cancellationToken)
+			=> CommitNode.DeserializeAsync(node, cancellationToken);
+
+		/// <inheritdoc/>
+		public override ValueTask<ITreeBlob> SerializeAsync(ITreeBlobWriter writer, CommitNode node, CancellationToken cancellationToken)
+			=> node.SerializeAsync(writer, cancellationToken);
+	}
+
 	/// <summary>
 	/// Service which mirrors changes from Perforce
 	/// </summary>
@@ -732,7 +814,7 @@ namespace Horde.Build.Perforce
 		async Task UpdateStreamContentInternalAsync(IStream stream, RedisList<int> changes, CancellationToken cancellationToken)
 		{
 			int prevChange = 0;
-			DirectoryNode prevContents = new DirectoryNode();
+			CommitNode prevContents = new CommitNode();
 			for (; ; )
 			{
 				// Get the first two changes to be mirrored. The first one should already exist, unless it's the start of replication for this stream.
@@ -747,11 +829,11 @@ namespace Horde.Build.Perforce
 				{
 					RefName prevRefName = GetRefName(stream, values[0]);
 
-					DirectoryNode? contents = await _treeStore.TryReadTreeAsync<DirectoryNode>(prevRefName, cancellationToken);
+					CommitNode? contents = await _treeStore.TryReadTreeAsync<CommitNode>(prevRefName, cancellationToken);
 					if (contents == null)
 					{
 						_logger.LogInformation("No content for CL {Change}; creating full snapshot", values[0]);
-						prevContents = await WriteCommitTreeAsync(stream, values[0], 0, new DirectoryNode(), cancellationToken);
+						prevContents = await WriteCommitTreeAsync(stream, values[0], 0, new CommitNode(), cancellationToken);
 					}
 					else
 					{
@@ -824,10 +906,10 @@ namespace Horde.Build.Perforce
 		/// <param name="revisionsOnly"></param>
 		/// <param name="cancellationToken"></param>
 		/// <returns></returns>
-		public async Task<DirectoryNode> ReadCommitTreeAsync(IStream stream, int change, string? filter, bool revisionsOnly, CancellationToken cancellationToken)
+		public async Task<CommitNode> ReadCommitTreeAsync(IStream stream, int change, string? filter, bool revisionsOnly, CancellationToken cancellationToken)
 		{
 			RefName name = GetRefName(stream.Id, change, filter, revisionsOnly);
-			return await _treeStore.ReadTreeAsync<DirectoryNode>(name, cancellationToken);
+			return await _treeStore.ReadTreeAsync<CommitNode>(name, cancellationToken);
 		}
 
 		/// <summary>
@@ -839,7 +921,7 @@ namespace Horde.Build.Perforce
 		/// <param name="baseContents">Initial contents of the tree at baseChange</param>
 		/// <param name="cancellationToken">Cancellation token for the operation</param>
 		/// <returns>Root tree object</returns>
-		public async Task<DirectoryNode> WriteCommitTreeAsync(IStream stream, int change, int baseChange, DirectoryNode baseContents, CancellationToken cancellationToken)
+		public async Task<CommitNode> WriteCommitTreeAsync(IStream stream, int change, int baseChange, CommitNode baseContents, CancellationToken cancellationToken)
 		{
 			bool revisionsOnly = stream.ReplicationMode == ContentReplicationMode.RevisionsOnly;
 			return await WriteCommitTreeAsync(stream, change, baseChange, baseContents, stream.ReplicationFilter, revisionsOnly, cancellationToken);
@@ -856,10 +938,8 @@ namespace Horde.Build.Perforce
 		/// <param name="revisionsOnly"></param>
 		/// <param name="cancellationToken">Cancellation token for the operation</param>
 		/// <returns>Root tree object</returns>
-		public async Task<DirectoryNode> WriteCommitTreeAsync(IStream stream, int change, int? baseChange, DirectoryNode baseContents, string? filter, bool revisionsOnly, CancellationToken cancellationToken)
+		public async Task<CommitNode> WriteCommitTreeAsync(IStream stream, int change, int? baseChange, CommitNode baseContents, string? filter, bool revisionsOnly, CancellationToken cancellationToken)
 		{
-			DirectoryNode contents = baseContents;
-
 			// Create a client to replicate from this stream
 			ReplicationClient clientInfo = await FindOrAddReplicationClientAsync(stream);
 
@@ -876,6 +956,18 @@ namespace Horde.Build.Perforce
 			string queryPath = $"//{clientInfo.Client.Name}/{filterOrDefault}";
 
 			RefName refName = GetRefName(stream.Id, change, filter, revisionsOnly);
+
+			// Get the current sync state for this change
+			CommitNode? syncNode = await _treeStore.TryReadTreeAsync<CommitNode>(refName, cancellationToken);
+			if (syncNode == null)
+			{
+				DirectoryNode contents = await baseContents.Contents.ExpandAsync(cancellationToken);
+				syncNode = new CommitNode(contents);
+			}
+
+			// Get the contents of the current tree
+			DirectoryNode root = await syncNode.Contents.ExpandAsync(cancellationToken);
+
 			if (revisionsOnly)
 			{
 				int count = 0;
@@ -898,25 +990,29 @@ namespace Horde.Build.Perforce
 
 					using ReadOnlyMemoryStream dataStream = new ReadOnlyMemoryStream(data);
 
-					FileEntry entry = await contents.AddFileByPathAsync(path, FileEntryFlags.PerforceDepotPathAndRevision, cancellationToken);
+					FileEntry entry = await root.AddFileByPathAsync(path, FileEntryFlags.PerforceDepotPathAndRevision, cancellationToken);
 					await entry.AppendAsync(dataStream, Options.Chunking, cancellationToken);
 
 					if (++count > 1000)
 					{
-						await _treeStore.WriteTreeAsync(refName, contents, false, cancellationToken);
+						await _treeStore.WriteTreeAsync(refName, syncNode, false, cancellationToken);
 						count = 0;
 					}
 				}
 			}
 			else
 			{
-				const long TrimInterval = 1024 * 1024 * 256;
-
-				long dataSize = 0;
-				long trimDataSize = TrimInterval;
+				// Replay the files that have already been synced
+				foreach (Utf8String path in syncNode.Paths)
+				{
+					string flushPath = $"//{clientInfo.Client.Name}/{path}@{change}";
+					_logger.LogInformation("Flushing {FlushPath}", flushPath);
+					await perforce.SyncQuietAsync(SyncOptions.Force | SyncOptions.KeepWorkspaceFiles, -1, flushPath, cancellationToken);
+				}
 
-				Dictionary<int, FileEntry> files = new Dictionary<int, FileEntry>();
-				await foreach (PerforceResponse response in perforce.StreamCommandAsync("sync", Array.Empty<string>(), new string[] { $"{queryPath}@{change}" }, null, typeof(SyncRecord), true, default))
+				// Do a sync preview to find everything that's left, and sort the remaining list of paths
+				List<(Utf8String Path, long Size)> files = new List<(Utf8String, long)>();
+				await foreach (PerforceResponse<SyncRecord> response in perforce.StreamCommandAsync<SyncRecord>("sync", new[] { "-n" }, new string[] { $"{queryPath}@{change}" }, null, cancellationToken))
 				{
 					PerforceError? error = response.Error;
 					if (error != null)
@@ -925,43 +1021,92 @@ namespace Horde.Build.Perforce
 						continue;
 					}
 
-					PerforceIo? io = response.Io;
-					if (io != null)
+					Utf8String file = response.Data.Path.Clone();
+					if (!file.StartsWith(clientRoot, Utf8StringComparer.Ordinal))
 					{
-						if (io.Command == PerforceIoCommand.Open)
-						{
-							Utf8String path = GetClientRelativePath(io.Payload, clientInfo.Client.Root);
-							files[io.File] = await contents.AddFileByPathAsync(path, FileEntryFlags.None, cancellationToken);
-						}
-						else if (io.Command == PerforceIoCommand.Write)
-						{
-							FileEntry file = files[io.File];
-							await file.AppendAsync(io.Payload, Options.Chunking, cancellationToken);
-							dataSize += io.Payload.Length;
-						}
-						else if (io.Command == PerforceIoCommand.Close)
-						{
-							files.Remove(io.File);
-						}
-						else if (io.Command == PerforceIoCommand.Unlink)
-						{
-							Utf8String path = GetClientRelativePath(io.Payload, clientInfo.Client.Root);
-							await contents.DeleteFileByPathAsync(path, cancellationToken);
-						}
-						else
+						throw new ArgumentException($"Unable to make path {file} relative to client root {clientRoot}");
+					}
+
+					byte[] path = new byte[file.Length - clientRoot.Length];
+					for (int idx = clientRoot.Length; idx < file.Length; idx++)
+					{
+						path[idx - clientRoot.Length] = (file[idx] == '\\') ? (byte)'/' : file[idx];
+					}
+
+					files.Add((new Utf8String(path), response.Data.FileSize));
+				}
+				files.SortBy(x => x.Path);
+
+				// Sync incrementally
+				while (files.Count > 0)
+				{
+					// Find the next path to sync
+					const long MaxBatchSize = 100 * 1024 * 1024;
+					(int idx, Utf8String path) = GetSyncBatch(files, MaxBatchSize);
+
+					const long TrimInterval = 1024 * 1024 * 256;
+
+					long dataSize = 0;
+					long trimDataSize = TrimInterval;
+
+					string syncPath = $"//{clientInfo.Client.Name}/{path}@{change}";
+					_logger.LogInformation("Syncing {SyncPath}", syncPath);
+
+					Dictionary<int, FileEntry> handles = new Dictionary<int, FileEntry>();
+					await foreach (PerforceResponse response in perforce.StreamCommandAsync("sync", Array.Empty<string>(), new string[] {  }, null, typeof(SyncRecord), true, default))
+					{
+						PerforceError? error = response.Error;
+						if (error != null)
 						{
-							_logger.LogWarning("Unhandled command code {Code}", io.Command);
+							_logger.LogWarning("Perforce: {Message}", error.Data);
+							continue;
 						}
 
-						if (dataSize > trimDataSize)
+						PerforceIo? io = response.Io;
+						if (io != null)
 						{
-							_logger.LogInformation("Trimming working set after receiving {NumBytes:n0}mb...", dataSize / (1024 * 1024));
-							await _treeStore.WriteTreeAsync(refName, contents, false, cancellationToken);
-							GC.Collect();
-							trimDataSize = dataSize + TrimInterval;
-							_logger.LogInformation("Trimming complete. Next trim at {NextNumBytes:n0}mb.", trimDataSize / (1024 * 1024));
+							if (io.Command == PerforceIoCommand.Open)
+							{
+								Utf8String file = GetClientRelativePath(io.Payload, clientInfo.Client.Root);
+								handles[io.File] = await root.AddFileByPathAsync(file, FileEntryFlags.None, cancellationToken);
+							}
+							else if (io.Command == PerforceIoCommand.Write)
+							{
+								FileEntry entry = handles[io.File];
+								await entry.AppendAsync(io.Payload, Options.Chunking, cancellationToken);
+								dataSize += io.Payload.Length;
+							}
+							else if (io.Command == PerforceIoCommand.Close)
+							{
+								handles.Remove(io.File);
+							}
+							else if (io.Command == PerforceIoCommand.Unlink)
+							{
+								Utf8String file = GetClientRelativePath(io.Payload, clientInfo.Client.Root);
+								await root.DeleteFileByPathAsync(file, cancellationToken);
+							}
+							else
+							{
+								_logger.LogWarning("Unhandled command code {Code}", io.Command);
+							}
+
+							if (dataSize > trimDataSize)
+							{
+								_logger.LogInformation("Trimming working set after receiving {NumBytes:n0}mb...", dataSize / (1024 * 1024));
+								await _treeStore.WriteTreeAsync(refName, syncNode, false, cancellationToken);
+								GC.Collect();
+								trimDataSize = dataSize + TrimInterval;
+								_logger.LogInformation("Trimming complete. Next trim at {NextNumBytes:n0}mb.", trimDataSize / (1024 * 1024));
+							}
 						}
 					}
+
+					// Update the root sync node, and write it out
+					syncNode.Paths.Add(path);
+					await _treeStore.WriteTreeAsync(refName, syncNode, cancellationToken: cancellationToken);
+
+					// Remove the synced files
+					files.RemoveRange(idx, files.Count - idx);
 				}
 			}
 
@@ -969,9 +1114,48 @@ namespace Horde.Build.Perforce
 
 			// Return the new root object
 			_logger.LogInformation("Writing ref {RefId} for {StreamId} change {Change}", refName, stream.Id, change);
-			await _treeStore.WriteTreeAsync(refName, contents, true, cancellationToken);
+			await _treeStore.WriteTreeAsync(refName, syncNode, true, cancellationToken);
+
+			return syncNode;
+		}
+
+		static (int, Utf8String) GetSyncBatch(List<(Utf8String Path, long Size)> files, long maxSize)
+		{
+			int idx = files.Count - 1;
+			Utf8String path = files[idx].Path;
+			long size = files[idx].Size;
+
+			Utf8String prefix = path;
+			while (prefix.Length > 0)
+			{
+				// Get the length of the parent directory
+				int endIdx = prefix.Length - 1;
+				while (endIdx > 0 && path[endIdx - 1] != (byte)'/')
+				{
+					endIdx--;
+				}
+
+				// Get the parent directory prefix, ending with a slash
+				prefix = prefix.Substring(0, endIdx);
+
+				// Include all files with this prefix
+				int nextIdx = idx;
+				for (; nextIdx > 0 && files[nextIdx - 1].Path.StartsWith(prefix); nextIdx--)
+				{
+					long nextSize = size + files[nextIdx - 1].Size;
+					if (nextSize > maxSize)
+					{
+						return (idx, path);
+					}
+					size = nextSize;
+				}
+				idx = nextIdx;
+
+				// Update the sync path to be the prefixed directory
+				path = prefix + "...";
+			}
 
-			return contents;
+			return (idx, path);
 		}
 
 		async Task FlushWorkspaceAsync(ReplicationClient clientInfo, IPerforceConnection perforce, int change)
