Horde: Missing file.

diff --git a/Engine/Source/Programs/Horde/Horde.Build/Commands/Bundles/PerforceCommand.cs b/Engine/Source/Programs/Horde/Horde.Build/Commands/Bundles/PerforceCommand.cs
--- a/Engine/Source/Programs/Horde/Horde.Build/Commands/Bundles/PerforceCommand.cs
+++ b/Engine/Source/Programs/Horde/Horde.Build/Commands/Bundles/PerforceCommand.cs
@@ -86,10 +86,10 @@ namespace Horde.Build.Commands.Bundles
 			Dictionary<IStream, int> streamToFirstChange = new Dictionary<IStream, int>();
 			streamToFirstChange[stream] = Change;
 
-			DirectoryNode baseContents;
+			CommitNode baseContents;
 			if (BaseChange == 0)
 			{
-				baseContents = new DirectoryNode();
+				baseContents = new CommitNode(new DirectoryNode());
 			}
 			else
 			{
