Horde: Fix incorrect sync path for commit replication, and sort paths by folder first to improve batching.

diff --git a/Engine/Source/Programs/Horde/Horde.Build/Perforce/CommitService.cs b/Engine/Source/Programs/Horde/Horde.Build/Perforce/CommitService.cs
--- a/Engine/Source/Programs/Horde/Horde.Build/Perforce/CommitService.cs
+++ b/Engine/Source/Programs/Horde/Horde.Build/Perforce/CommitService.cs
@@ -927,6 +927,28 @@ namespace Horde.Build.Perforce
 			return await WriteCommitTreeAsync(stream, change, baseChange, baseContents, stream.ReplicationFilter, revisionsOnly, cancellationToken);
 		}
 
+		/// <summary>
+		/// Sorts paths by their folder, then by their filename.
+		/// </summary>
+		class FolderFirstSorter : IComparer<Utf8String>
+		{
+			public static FolderFirstSorter Instance { get; } = new FolderFirstSorter();
+
+			public int Compare(Utf8String x, Utf8String y)
+			{
+				Utf8String pathX = x.Substring(0, x.LastIndexOf('/') + 1);
+				Utf8String pathY = y.Substring(0, y.LastIndexOf('/') + 1);
+
+				int result = pathX.CompareTo(pathY);
+				if (result != 0)
+				{
+					return result;
+				}
+
+				return x.CompareTo(y);
+			}
+		}
+
 		/// <summary>
 		/// Replicates the contents of a stream to Horde storage, optionally using the given change as a starting point
 		/// </summary>
@@ -1035,7 +1057,7 @@ namespace Horde.Build.Perforce
 
 					files.Add((new Utf8String(path), response.Data.FileSize));
 				}
-				files.SortBy(x => x.Path);
+				files.SortBy(x => x.Path, FolderFirstSorter.Instance);
 
 				// Sync incrementally
 				while (files.Count > 0)
@@ -1053,7 +1075,7 @@ namespace Horde.Build.Perforce
 					_logger.LogInformation("Syncing {SyncPath}", syncPath);
 
 					Dictionary<int, FileEntry> handles = new Dictionary<int, FileEntry>();
-					await foreach (PerforceResponse response in perforce.StreamCommandAsync("sync", Array.Empty<string>(), new string[] {  }, null, typeof(SyncRecord), true, default))
+					await foreach (PerforceResponse response in perforce.StreamCommandAsync("sync", Array.Empty<string>(), new string[] { syncPath }, null, typeof(SyncRecord), true, default))
 					{
 						PerforceError? error = response.Error;
 						if (error != null)
