Horde: Include more useful stats for syncing operations.

diff --git a/Engine/Source/Programs/Horde/Horde.Build/Perforce/CommitService.cs b/Engine/Source/Programs/Horde/Horde.Build/Perforce/CommitService.cs
--- a/Engine/Source/Programs/Horde/Horde.Build/Perforce/CommitService.cs
+++ b/Engine/Source/Programs/Horde/Horde.Build/Perforce/CommitService.cs
@@ -1060,11 +1060,14 @@ namespace Horde.Build.Perforce
 				files.SortBy(x => x.Path, FolderFirstSorter.Instance);
 
 				// Sync incrementally
+				long totalSize = files.Sum(x => x.Size);
+				long syncedSize = 0;
 				while (files.Count > 0)
 				{
 					// Find the next path to sync
-					const long MaxBatchSize = 100 * 1024 * 1024;
-					(int idx, Utf8String path) = GetSyncBatch(files, MaxBatchSize);
+					const long MaxBatchSize = 1024 * 1024 * 1024;
+					(int idx, Utf8String path, long size) = GetSyncBatch(files, MaxBatchSize);
+					syncedSize += size;
 
 					const long TrimInterval = 1024 * 1024 * 256;
 
@@ -1072,7 +1075,8 @@ namespace Horde.Build.Perforce
 					long trimDataSize = TrimInterval;
 
 					string syncPath = $"//{clientInfo.Client.Name}/{path}@{change}";
-					_logger.LogInformation("Syncing {SyncPath}", syncPath);
+					double syncPct = (syncedSize * 100.0) / Math.Max(totalSize, 1L);
+					_logger.LogInformation("Syncing {StreamId} to {Change} [{SyncPct:n1}%]: {Path} ({Size:n0} bytes)", stream.Id, change, syncPct, path, size);
 
 					Dictionary<int, FileEntry> handles = new Dictionary<int, FileEntry>();
 					await foreach (PerforceResponse response in perforce.StreamCommandAsync("sync", Array.Empty<string>(), new string[] { syncPath }, null, typeof(SyncRecord), true, default))
@@ -1141,7 +1145,7 @@ namespace Horde.Build.Perforce
 			return syncNode;
 		}
 
-		static (int, Utf8String) GetSyncBatch(List<(Utf8String Path, long Size)> files, long maxSize)
+		static (int, Utf8String, long) GetSyncBatch(List<(Utf8String Path, long Size)> files, long maxSize)
 		{
 			int idx = files.Count - 1;
 			Utf8String path = files[idx].Path;
@@ -1162,22 +1166,23 @@ namespace Horde.Build.Perforce
 
 				// Include all files with this prefix
 				int nextIdx = idx;
+				long nextSize = size;
 				for (; nextIdx > 0 && files[nextIdx - 1].Path.StartsWith(prefix); nextIdx--)
 				{
-					long nextSize = size + files[nextIdx - 1].Size;
+					nextSize += files[nextIdx - 1].Size;
 					if (nextSize > maxSize)
 					{
-						return (idx, path);
+						return (idx, path, size);
 					}
-					size = nextSize;
 				}
+				size = nextSize;
 				idx = nextIdx;
 
 				// Update the sync path to be the prefixed directory
 				path = prefix + "...";
 			}
 
-			return (idx, path);
+			return (idx, path, size);
 		}
 
 		async Task FlushWorkspaceAsync(ReplicationClient clientInfo, IPerforceConnection perforce, int change)
