Horde: Try using client name when flushing to improve performance.

diff --git a/Engine/Source/Programs/Horde/Horde.Build/Perforce/CommitService.cs b/Engine/Source/Programs/Horde/Horde.Build/Perforce/CommitService.cs
--- a/Engine/Source/Programs/Horde/Horde.Build/Perforce/CommitService.cs
+++ b/Engine/Source/Programs/Horde/Horde.Build/Perforce/CommitService.cs
@@ -1192,13 +1192,13 @@ namespace Horde.Build.Perforce
 				clientInfo.Change = -1;
 				if (change == 0)
 				{
-					_logger.LogInformation("Flushing have table for {Client}", perforce.Settings.ClientName);
-					await perforce.SyncQuietAsync(SyncOptions.Force | SyncOptions.KeepWorkspaceFiles, -1, $"//...#0");
+					_logger.LogInformation("Flushing have table for {Client}", clientInfo.Client.Name);
+					await perforce.SyncQuietAsync(SyncOptions.Force | SyncOptions.KeepWorkspaceFiles, -1, $"//{clientInfo.Client.Name}/...#0");
 				}
 				else
 				{
-					_logger.LogInformation("Flushing have table for {Client} to change {Change}", perforce.Settings.ClientName, change);
-					await perforce.SyncQuietAsync(SyncOptions.Force | SyncOptions.KeepWorkspaceFiles, -1, $"//...@{change}");
+					_logger.LogInformation("Flushing have table for {Client} to change {Change}", clientInfo.Client.Name, change);
+					await perforce.SyncQuietAsync(SyncOptions.Force | SyncOptions.KeepWorkspaceFiles, -1, $"//{clientInfo.Client.Name}/...@{change}");
 				}
 				clientInfo.Change = change;
 			}
