Horde: Fix incorrect deserialization of root commit nodes.

diff --git a/Engine/Source/Programs/Horde/Horde.Build/Perforce/CommitService.cs b/Engine/Source/Programs/Horde/Horde.Build/Perforce/CommitService.cs
--- a/Engine/Source/Programs/Horde/Horde.Build/Perforce/CommitService.cs
+++ b/Engine/Source/Programs/Horde/Horde.Build/Perforce/CommitService.cs
@@ -117,7 +117,7 @@ namespace Horde.Build.Perforce
 			MemoryReader reader = new MemoryReader(data.AsSingleSegment());
 			Utf8String[] paths = reader.ReadVariableLengthArray(() => reader.ReadUtf8String());
 
-			return new CommitNode(paths, blob);
+			return new CommitNode(paths, children[0]);
 		}
 
 		/// <inheritdoc/>
