Horde: Prevent partial sync data when an exception occurs during replication.

diff --git a/Engine/Source/Programs/Horde/Horde.Build/Perforce/CommitService.cs b/Engine/Source/Programs/Horde/Horde.Build/Perforce/CommitService.cs
--- a/Engine/Source/Programs/Horde/Horde.Build/Perforce/CommitService.cs
+++ b/Engine/Source/Programs/Horde/Horde.Build/Perforce/CommitService.cs
@@ -779,10 +779,10 @@ namespace Horde.Build.Perforce
 					Task internalTask = Task.Run(() => UpdateStreamContentInternalAsync(stream, streamChanges, cancellationToken), cancellationToken);
 					while (!internalTask.IsCompleted)
 					{
-						Task delayTask = Task.Delay(TimeSpan.FromSeconds(15.0), cancellationToken);
+						Task delayTask = Task.Delay(TimeSpan.FromSeconds(60.0), cancellationToken);
 						if (await Task.WhenAny(internalTask, delayTask) == delayTask)
 						{
-							DateTime newTime = DateTime.UtcNow + TimeSpan.FromSeconds(30.0);
+							DateTime newTime = DateTime.UtcNow + TimeSpan.FromSeconds(90.0);
 							await _redisReservations.AddAsync(streamId, newTime.Ticks);
 							_logger.LogInformation("Extending reservation for content update of {StreamId} (elapsed: {Time}s)", streamId, (int)timer.Elapsed.TotalSeconds);
 						}
@@ -978,9 +978,11 @@ namespace Horde.Build.Perforce
 			string queryPath = $"//{clientInfo.Client.Name}/{filterOrDefault}";
 
 			RefName refName = GetRefName(stream.Id, change, filter, revisionsOnly);
+			RefName incRefName = new RefName($"{refName}_inc");
 
 			// Get the current sync state for this change
-			CommitNode? syncNode = await _treeStore.TryReadTreeAsync<CommitNode>(refName, cancellationToken);
+			CommitNode? syncNode = await _treeStore.TryReadTreeAsync<CommitNode>(incRefName, cancellationToken);
+			bool deleteIncRef = syncNode != null;
 			if (syncNode == null)
 			{
 				DirectoryNode contents = await baseContents.Contents.ExpandAsync(cancellationToken);
@@ -1064,6 +1066,13 @@ namespace Horde.Build.Perforce
 				long syncedSize = 0;
 				while (files.Count > 0)
 				{
+					// Save the incremental state
+					if (syncedSize > 0)
+					{
+						await _treeStore.WriteTreeAsync(incRefName, syncNode, cancellationToken: cancellationToken);
+						deleteIncRef = true;
+					}
+
 					// Find the next path to sync
 					const long MaxBatchSize = 1024 * 1024 * 1024;
 					(int idx, Utf8String path, long size) = GetSyncBatch(files, MaxBatchSize);
@@ -1127,11 +1136,8 @@ namespace Horde.Build.Perforce
 						}
 					}
 
-					// Update the root sync node, and write it out
+					// Update the root sync node
 					syncNode.Paths.Add(path);
-					await _treeStore.WriteTreeAsync(refName, syncNode, cancellationToken: cancellationToken);
-
-					// Remove the synced files
 					files.RemoveRange(idx, files.Count - idx);
 				}
 			}
@@ -1142,6 +1148,11 @@ namespace Horde.Build.Perforce
 			_logger.LogInformation("Writing ref {RefId} for {StreamId} change {Change}", refName, stream.Id, change);
 			await _treeStore.WriteTreeAsync(refName, syncNode, true, cancellationToken);
 
+			// Delete the incremental state
+			if (deleteIncRef)
+			{
+				await _treeStore.DeleteTreeAsync(incRefName, cancellationToken);
+			}
 			return syncNode;
 		}
 
