Update the replication version number to force a full snapshot.

diff --git a/Engine/Source/Programs/Horde/Horde.Build/Perforce/CommitService.cs b/Engine/Source/Programs/Horde/Horde.Build/Perforce/CommitService.cs
--- a/Engine/Source/Programs/Horde/Horde.Build/Perforce/CommitService.cs
+++ b/Engine/Source/Programs/Horde/Horde.Build/Perforce/CommitService.cs
@@ -874,7 +874,7 @@ namespace Horde.Build.Perforce
 		/// <returns></returns>
 		static RefName GetRefName(StreamId streamId, int change, string? filter, bool revisionsOnly)
 		{
-			StringBuilder builder = new StringBuilder($"v1/{streamId}/{change}");
+			StringBuilder builder = new StringBuilder($"v2/{streamId}/{change}");
 			if (filter != null)
 			{
 				builder.Append($"_filter_{IoHash.Compute(Encoding.UTF8.GetBytes(filter))}");
