Horde: Fix overlapping reservations for replication updates, causing multiple pods to replicate a stream at once.

diff --git a/Engine/Source/Programs/Horde/Horde.Build/Perforce/CommitService.cs b/Engine/Source/Programs/Horde/Horde.Build/Perforce/CommitService.cs
--- a/Engine/Source/Programs/Horde/Horde.Build/Perforce/CommitService.cs
+++ b/Engine/Source/Programs/Horde/Horde.Build/Perforce/CommitService.cs
@@ -629,6 +629,9 @@ namespace Horde.Build.Perforce
 
 		#region Content Updates
 
+		const double ReservationTimeSeconds = 60.0;
+		const double ExtendReservationAfterSeconds = 40.0;
+
 		RedisKey RedisBaseKey { get; } = new RedisKey("commits/");
 		RedisChannel<StreamId> RedisUpdateChannel { get; } = new RedisChannel<StreamId>("commits/streams");
 		bool _stopping;
@@ -711,7 +714,7 @@ namespace Horde.Build.Perforce
 					if (checkStreams.Count > 0)
 					{
 						// Try to start new background tasks
-						double newScore = (utcNow + TimeSpan.FromSeconds(30.0)).Ticks;
+						double newScore = (utcNow + TimeSpan.FromSeconds(ReservationTimeSeconds)).Ticks;
 						foreach (StreamId checkStream in checkStreams)
 						{
 							if (await _redisReservations.AddAsync(checkStream, newScore, When.NotExists))
@@ -779,10 +782,10 @@ namespace Horde.Build.Perforce
 					Task internalTask = Task.Run(() => UpdateStreamContentInternalAsync(stream, streamChanges, cancellationToken), cancellationToken);
 					while (!internalTask.IsCompleted)
 					{
-						Task delayTask = Task.Delay(TimeSpan.FromSeconds(60.0), cancellationToken);
+						Task delayTask = Task.Delay(TimeSpan.FromSeconds(ExtendReservationAfterSeconds), cancellationToken);
 						if (await Task.WhenAny(internalTask, delayTask) == delayTask)
 						{
-							DateTime newTime = DateTime.UtcNow + TimeSpan.FromSeconds(90.0);
+							DateTime newTime = DateTime.UtcNow + TimeSpan.FromSeconds(ReservationTimeSeconds);
 							await _redisReservations.AddAsync(streamId, newTime.Ticks);
 							_logger.LogInformation("Extending reservation for content update of {StreamId} (elapsed: {Time}s)", streamId, (int)timer.Elapsed.TotalSeconds);
 						}
@@ -970,7 +973,7 @@ namespace Horde.Build.Perforce
 			await FlushWorkspaceAsync(clientInfo, perforce, baseChange ?? 0);
 
 			// Apply all the updates
-			_logger.LogInformation("Updating client {Client} from changelist {BaseChange} to {Change}", clientInfo.Client.Name, baseChange ?? 0, change);
+			_logger.LogInformation("Syncing client {Client} from changelist {BaseChange} to {Change}", clientInfo.Client.Name, baseChange ?? 0, change);
 			clientInfo.Change = -1;
 
 			Utf8String clientRoot = new Utf8String(clientInfo.Client.Root);
