Horde: Additional logging to explain selection of sync filters.

diff --git a/Engine/Source/Programs/Horde/Horde.Build/Perforce/CommitService.cs b/Engine/Source/Programs/Horde/Horde.Build/Perforce/CommitService.cs
--- a/Engine/Source/Programs/Horde/Horde.Build/Perforce/CommitService.cs
+++ b/Engine/Source/Programs/Horde/Horde.Build/Perforce/CommitService.cs
@@ -1078,7 +1078,7 @@ namespace Horde.Build.Perforce
 
 					// Find the next path to sync
 					const long MaxBatchSize = 1024 * 1024 * 1024;
-					(int idx, Utf8String path, long size) = GetSyncBatch(files, MaxBatchSize);
+					(int idx, Utf8String path, long size) = GetSyncBatch(files, MaxBatchSize, _logger);
 					syncedSize += size;
 
 					const long TrimInterval = 1024 * 1024 * 256;
@@ -1087,8 +1087,8 @@ namespace Horde.Build.Perforce
 					long trimDataSize = TrimInterval;
 
 					string syncPath = $"//{clientInfo.Client.Name}/{path}@{change}";
-					double syncPct = (syncedSize * 100.0) / Math.Max(totalSize, 1L);
-					_logger.LogInformation("Syncing {StreamId} to {Change} [{SyncPct:n1}%]: {Path} ({Size:n0} bytes)", stream.Id, change, syncPct, path, size);
+					double syncPct = (totalSize == 0) ? 100.0 : (syncedSize * 100.0) / totalSize;
+					_logger.LogInformation("Syncing {StreamId} to {Change} [{SyncPct:n1}%]: {Path} ({Size:n1}mb)", stream.Id, change, syncPct, path, size / (1024.0 * 1024.0));
 
 					Dictionary<int, FileEntry> handles = new Dictionary<int, FileEntry>();
 					await foreach (PerforceResponse response in perforce.StreamCommandAsync("sync", Array.Empty<string>(), new string[] { syncPath }, null, typeof(SyncRecord), true, default))
@@ -1130,11 +1130,11 @@ namespace Horde.Build.Perforce
 
 							if (dataSize > trimDataSize)
 							{
-								_logger.LogInformation("Trimming working set after receiving {NumBytes:n0}mb...", dataSize / (1024 * 1024));
+								_logger.LogInformation("Trimming working set after receiving {Size:n1}mb...", dataSize / (1024 * 1024));
 								await _treeStore.WriteTreeAsync(refName, syncNode, false, cancellationToken);
 								GC.Collect();
 								trimDataSize = dataSize + TrimInterval;
-								_logger.LogInformation("Trimming complete. Next trim at {NextNumBytes:n0}mb.", trimDataSize / (1024 * 1024));
+								_logger.LogInformation("Trimming complete. Next trim at {Size:n1}mb.", trimDataSize / (1024 * 1024));
 							}
 						}
 					}
@@ -1159,7 +1159,7 @@ namespace Horde.Build.Perforce
 			return syncNode;
 		}
 
-		static (int, Utf8String, long) GetSyncBatch(List<(Utf8String Path, long Size)> files, long maxSize)
+		static (int, Utf8String, long) GetSyncBatch(List<(Utf8String Path, long Size)> files, long maxSize, ILogger logger)
 		{
 			int idx = files.Count - 1;
 			Utf8String path = files[idx].Path;
@@ -1181,13 +1181,15 @@ namespace Horde.Build.Perforce
 				// Include all files with this prefix
 				int nextIdx = idx;
 				long nextSize = size;
-				for (; nextIdx > 0 && files[nextIdx - 1].Path.StartsWith(prefix); nextIdx--)
+				while(nextIdx > 0 && files[nextIdx - 1].Path.StartsWith(prefix))
 				{
 					nextSize += files[nextIdx - 1].Size;
-					if (nextSize > maxSize)
-					{
-						return (idx, path, size);
-					}
+					nextIdx--;
+				}
+				if (nextSize > maxSize)
+				{
+					logger.LogInformation("Next filter {NextFilter} would exceed max size ({NextSize:n1}mb > {MaxSize:n1}mb); using {Filter} ({Size:n1}mb) instead.", prefix + "...", nextSize / (1024.0 * 1024.0), maxSize / (1024.0 * 1024.0), path, size / (1024.0 * 1024.0));
+					return (idx, path, size);
 				}
 				size = nextSize;
 				idx = nextIdx;
