Horde: Increase maximum sync batch size to 10gb.

diff --git a/Engine/Source/Programs/Horde/Horde.Build/Perforce/CommitService.cs b/Engine/Source/Programs/Horde/Horde.Build/Perforce/CommitService.cs
--- a/Engine/Source/Programs/Horde/Horde.Build/Perforce/CommitService.cs
+++ b/Engine/Source/Programs/Horde/Horde.Build/Perforce/CommitService.cs
@@ -1077,7 +1077,7 @@ namespace Horde.Build.Perforce
 					}
 
 					// Find the next path to sync
-					const long MaxBatchSize = 1024 * 1024 * 1024;
+					const long MaxBatchSize = 10L * 1024 * 1024 * 1024;
 					(int idx, Utf8String path, long size) = GetSyncBatch(files, MaxBatchSize, _logger);
 					syncedSize += size;
 
