Horde: Fix mismatched condition for stopping replication.

diff --git a/Engine/Source/Programs/Horde/Horde.Build/Perforce/CommitService.cs b/Engine/Source/Programs/Horde/Horde.Build/Perforce/CommitService.cs
--- a/Engine/Source/Programs/Horde/Horde.Build/Perforce/CommitService.cs
+++ b/Engine/Source/Programs/Horde/Horde.Build/Perforce/CommitService.cs
@@ -240,7 +240,7 @@ namespace Horde.Build.Perforce
 		/// <inheritdoc/>
 		public async Task StopAsync(CancellationToken cancellationToken)
 		{
-			if (Options.Content)
+			if (Options.Metadata)
 			{
 				await StopContentReplicationAsync();
 				await StopMetadataReplicationAsync();
