Horde: Create a separate service for replicating commit content to metadata.

diff --git a/Engine/Source/Programs/Horde/Horde.Build/Commands/Bundles/PerforceCommand.cs b/Engine/Source/Programs/Horde/Horde.Build/Commands/Bundles/PerforceCommand.cs
--- a/Engine/Source/Programs/Horde/Horde.Build/Commands/Bundles/PerforceCommand.cs
+++ b/Engine/Source/Programs/Horde/Horde.Build/Commands/Bundles/PerforceCommand.cs
@@ -73,6 +73,7 @@ namespace Horde.Build.Commands.Bundles
 			using ServiceProvider serviceProvider = Startup.CreateServiceProvider(_configuration, _loggerProvider);
 
 			CommitService commitService = serviceProvider.GetRequiredService<CommitService>();
+			ReplicationService replicationService = serviceProvider.GetRequiredService<ReplicationService>();
 			ICommitCollection commitCollection = serviceProvider.GetRequiredService<ICommitCollection>();
 			IStreamCollection streamCollection = serviceProvider.GetRequiredService<IStreamCollection>();
 			IStorageClient storageClient = serviceProvider.GetRequiredService<IStorageClient>();
@@ -86,14 +87,14 @@ namespace Horde.Build.Commands.Bundles
 			Dictionary<IStream, int> streamToFirstChange = new Dictionary<IStream, int>();
 			streamToFirstChange[stream] = Change;
 
-			CommitNode baseContents;
+			ReplicationNode baseContents;
 			if (BaseChange == 0)
 			{
-				baseContents = new CommitNode(new DirectoryNode());
+				baseContents = new ReplicationNode(new DirectoryNode());
 			}
 			else
 			{
-				baseContents = await commitService.ReadCommitTreeAsync(stream, BaseChange, Filter, RevisionsOnly, CancellationToken.None);
+				baseContents = await replicationService.ReadCommitTreeAsync(stream, BaseChange, Filter, RevisionsOnly, CancellationToken.None);
 			}
 
 			await foreach (NewCommit newCommit in commitService.FindCommitsForClusterAsync(stream.ClusterName, streamToFirstChange).Take(Count))
@@ -105,7 +106,7 @@ namespace Horde.Build.Commands.Bundles
 
 				if (Content)
 				{
-					baseContents = await commitService.WriteCommitTreeAsync(stream, newCommit.Change, BaseChange, baseContents, Filter, RevisionsOnly, CancellationToken.None);
+					baseContents = await replicationService.WriteCommitTreeAsync(stream, newCommit.Change, BaseChange, baseContents, Filter, RevisionsOnly, CancellationToken.None);
 					BaseChange = newCommit.Change;
 				}
 			}
diff --git a/Engine/Source/Programs/Horde/Horde.Build/Perforce/CommitService.cs b/Engine/Source/Programs/Horde/Horde.Build/Perforce/CommitService.cs
--- a/Engine/Source/Programs/Horde/Horde.Build/Perforce/CommitService.cs
+++ b/Engine/Source/Programs/Horde/Horde.Build/Perforce/CommitService.cs
@@ -42,102 +42,7 @@ namespace Horde.Build.Perforce
 		/// <summary>
 		/// Whether to mirror commit metadata to the database
 		/// </summary>
-		public bool Metadata { get; set; } = true;
-
-		/// <summary>
-		/// Whether to replicate content to Horde storage. Must also be enabled on a per-stream basis (see <see cref="IStream.ReplicationMode"/>).
-		/// </summary>
-		public bool Content { get; set; } = true;
-
-		/// <summary>
-		/// Options for how objects are packed together
-		/// </summary>
-		public BundleOptions Bundle { get; set; } = new BundleOptions();
-
-		/// <summary>
-		/// Options for how objects are sliced
-		/// </summary>
-		public ChunkingOptions Chunking { get; set; } = new ChunkingOptions();
-	}
-
-	/// <summary>
-	/// Root node for a commit snapshot
-	/// </summary>
-	[TreeSerializer(typeof(CommitNodeSerializer))]
-	public class CommitNode : TreeNode
-	{
-		/// <summary>
-		/// Paths that have been synced. Empty once complete.
-		/// </summary>
-		public List<Utf8String> Paths { get; }
-
-		/// <summary>
-		/// Contents of this snapshot
-		/// </summary>
-		public TreeNodeRef<DirectoryNode> Contents { get; }
-
-		/// <summary>
-		/// Constructor
-		/// </summary>
-		public CommitNode()
-			: this(new DirectoryNode())
-		{
-		}
-
-		/// <summary>
-		/// Constructor
-		/// </summary>
-		public CommitNode(DirectoryNode contents)
-		{
-			Paths = new List<Utf8String>();
-			Contents = new TreeNodeRef<DirectoryNode>(this, contents);
-		}
-
-		private CommitNode(IEnumerable<Utf8String> paths, ITreeBlob contents)
-		{
-			Paths = new List<Utf8String>(paths);
-			Contents = new TreeNodeRef<DirectoryNode>(this, contents);
-		}
-
-		internal async ValueTask<ITreeBlob> SerializeAsync(ITreeBlobWriter writer, CancellationToken cancellationToken)
-		{
-			List<ITreeBlob> references = new List<ITreeBlob>();
-			references.Add(await Contents.CollapseAsync(writer, cancellationToken));
-
-			ByteArrayBuilder builder = new ByteArrayBuilder();
-			builder.WriteVariableLengthArray(Paths, x => builder.WriteUtf8String(x));
-			return await writer.WriteBlobAsync(builder.AsSequence(), references, cancellationToken);
-		}
-
-		internal static async ValueTask<CommitNode> DeserializeAsync(ITreeBlob blob, CancellationToken cancellationToken)
-		{
-			IReadOnlyList<ITreeBlob> children = await blob.GetReferencesAsync(cancellationToken);
-			ReadOnlySequence<byte> data = await blob.GetDataAsync(cancellationToken);
-
-			MemoryReader reader = new MemoryReader(data.AsSingleSegment());
-			Utf8String[] paths = reader.ReadVariableLengthArray(() => reader.ReadUtf8String());
-
-			return new CommitNode(paths, children[0]);
-		}
-
-		/// <inheritdoc/>
-		public override IReadOnlyList<TreeNodeRef> GetReferences()
-		{
-			List<TreeNodeRef> refs = new List<TreeNodeRef>();
-			refs.Add(Contents);
-			return refs;
-		}
-	}
-
-	class CommitNodeSerializer : TreeNodeSerializer<CommitNode>
-	{
-		/// <inheritdoc/>
-		public override ValueTask<CommitNode> DeserializeAsync(ITreeBlob node, CancellationToken cancellationToken)
-			=> CommitNode.DeserializeAsync(node, cancellationToken);
-
-		/// <inheritdoc/>
-		public override ValueTask<ITreeBlob> SerializeAsync(ITreeBlobWriter writer, CommitNode node, CancellationToken cancellationToken)
-			=> node.SerializeAsync(writer, cancellationToken);
+		public bool Enable { get; set; } = true;
 	}
 
 	/// <summary>
@@ -163,29 +68,18 @@ namespace Horde.Build.Perforce
 		/// <summary>
 		/// A registered listener for new commits
 		/// </summary>
-		class ListenerInfo : IDisposable
+		sealed class ListenerInfo : IAsyncDisposable
 		{
-			readonly List<ListenerInfo> _listeners;
-			public Action<ICommit> _callback;
+			readonly CommitService _owner;
+			public Func<ICommit, Task> Callback { get; }
 
-			public ListenerInfo(List<ListenerInfo> listeners, Action<ICommit> callback)
+			public ListenerInfo(CommitService owner, Func<ICommit, Task> callback)
 			{
-				_listeners = listeners;
-				_callback = callback;
-
-				lock (listeners)
-				{
-					listeners.Add(this);
-				}
+				_owner = owner;
+				Callback = callback;
 			}
 
-			public void Dispose()
-			{
-				lock (_listeners)
-				{
-					_listeners.Remove(this);
-				}
-			}
+			public async ValueTask DisposeAsync() => await _owner.RemoveListenerAsync(this);
 		}
 
 		public CommitServiceOptions Options { get; }
@@ -198,27 +92,31 @@ namespace Horde.Build.Perforce
 		readonly IStreamCollection _streamCollection;
 		readonly IPerforceService _perforceService;
 		readonly IUserCollection _userCollection;
-		readonly ITreeStore _treeStore;
 		readonly ILogger<CommitService> _logger;
 
-		const int MaxBackgroundTasks = 2;
+		// Listeners
+		readonly object _lockObject = new object();
+		Task _listenerTask = Task.CompletedTask;
+		readonly List<ListenerInfo> _listeners = new List<ListenerInfo>();
+
+		// Metadata
+		readonly ITicker _updateMetadataTicker;
+
+		static RedisKey LastChangeKey(StreamId streamId) => $"commits/metadata/{streamId}";
 
 		/// <summary>
 		/// Constructor
 		/// </summary>
-		public CommitService(RedisService redisService, ICommitCollection commitCollection, IStreamCollection streamCollection, IPerforceService perforceService, IUserCollection userCollection, ITreeStore<CommitService> treeStore, IClock clock, IOptions<CommitServiceOptions> options, ILogger<CommitService> logger)
+		public CommitService(RedisService redisService, ICommitCollection commitCollection, IStreamCollection streamCollection, IPerforceService perforceService, IUserCollection userCollection, IClock clock, IOptions<CommitServiceOptions> options, ILogger<CommitService> logger)
 		{
 			Options = options.Value;
 
 			_redis = redisService.Database;
-			_redisDirtyStreams = new RedisSet<StreamId>(_redis, RedisBaseKey.Append("streams"));
-			_redisReservations = new RedisSortedSet<StreamId>(_redis, RedisBaseKey.Append("reservations"));
 
 			_commitCollection = commitCollection;
 			_streamCollection = streamCollection;
 			_perforceService = perforceService;
 			_userCollection = userCollection;
-			_treeStore = treeStore;
 			_logger = logger;
 			_updateMetadataTicker = clock.AddSharedTicker<CommitService>(TimeSpan.FromSeconds(30.0), UpdateMetadataAsync, logger);
 		}
@@ -229,87 +127,77 @@ namespace Horde.Build.Perforce
 		/// <inheritdoc/>
 		public async Task StartAsync(CancellationToken cancellationToken)
 		{
-			await StartNotificationsAsync();
-			if (Options.Metadata)
+			if (Options.Enable)
 			{
-				await StartMetadataReplicationAsync();
-				await StartContentReplicationAsync();
+				await _updateMetadataTicker.StartAsync();
 			}
 		}
 
 		/// <inheritdoc/>
 		public async Task StopAsync(CancellationToken cancellationToken)
 		{
-			if (Options.Metadata)
+			if (Options.Enable)
 			{
-				await StopContentReplicationAsync();
-				await StopMetadataReplicationAsync();
+				await _updateMetadataTicker.StopAsync();
 			}
-			await StopNotificationsAsync();
 		}
 
 		#region Notifications
 
-		readonly List<ListenerInfo> _listeners = new List<ListenerInfo>();
-		RedisChannel<CommitId> CommitNotifyChannel { get; } = new RedisChannel<CommitId>("commits/notify");
-		RedisChannelSubscription<CommitId>? _commitNotifySubscription;
-
-		async Task StartNotificationsAsync()
+		public IAsyncDisposable AddListener(Func<ICommit, Task> onAddCommit)
 		{
-			ISubscriber subscriber = _redis.Multiplexer.GetSubscriber();
-			_commitNotifySubscription = await subscriber.SubscribeAsync(CommitNotifyChannel, DispatchCommitNotification);
+			ListenerInfo listener = new ListenerInfo(this, onAddCommit);
+			lock (_lockObject)
+			{
+				_listeners.Add(listener);
+			}
+			return listener;
 		}
 
-		async Task StopNotificationsAsync()
+		Task RemoveListenerAsync(ListenerInfo listener)
 		{
-			if (_commitNotifySubscription != null)
+			Task waitTask;
+			lock (_lockObject)
 			{
-				await _commitNotifySubscription.DisposeAsync();
-				_commitNotifySubscription = null;
+				waitTask = _listenerTask.ContinueWith(x => RemoveListener(listener), TaskScheduler.Default);
+				_listenerTask = waitTask;
 			}
+			return waitTask;
 		}
 
-		async void DispatchCommitNotification(RedisChannel<CommitId> channel, CommitId commitId)
+		void RemoveListener(ListenerInfo listener)
 		{
-			ICommit? commit = await _commitCollection.GetCommitAsync(commitId);
-			if (commit != null)
+			lock (_lockObject)
 			{
-				lock (_listeners)
-				{
-					foreach (ListenerInfo listener in _listeners)
-					{
-						listener._callback(commit);
-					}
-				}
+				_listeners.Remove(listener);
 			}
 		}
 
-		async Task NotifyListeners(CommitId commitId)
+		private Task NotifyListenersAsync(ICommit commit)
 		{
-			await _redis.PublishAsync(CommitNotifyChannel, commitId);
+			Task waitTask;
+			lock (_lockObject)
+			{
+				waitTask = _listenerTask.ContinueWith(x => NotifyAllListenersAsync(commit), TaskScheduler.Default);
+				_listenerTask = waitTask;
+			}
+			return waitTask;
 		}
 
-		public IDisposable AddListener(Action<ICommit> onAddCommit)
+		async Task NotifyAllListenersAsync(ICommit commit)
 		{
-			return new ListenerInfo(_listeners, onAddCommit);
+			List<ListenerInfo> listeners;
+			lock (_lockObject)
+			{
+				listeners = new List<ListenerInfo>(_listeners);
+			}
+			await Task.WhenAll(listeners.Select(x => x.Callback(commit)));
 		}
 
 		#endregion
 
 		#region Metadata updates
 
-		readonly ITicker _updateMetadataTicker;
-
-		Task StartMetadataReplicationAsync()
-		{
-			return _updateMetadataTicker.StartAsync();
-		}
-
-		Task StopMetadataReplicationAsync()
-		{
-			return _updateMetadataTicker.StopAsync();
-		}
-
 		/// <summary>
 		/// Polls Perforce for submitted changes
 		/// </summary>
@@ -340,41 +228,31 @@ namespace Horde.Build.Perforce
 		async Task UpdateMetadataForClusterAsync(string clusterName, IEnumerable<IStream> streams)
 		{
 			// Find the minimum changelist number to query
-			Dictionary<IStream, int> streamToFirstChange = new Dictionary<IStream, int>();
+			Dictionary<IStream, int> streamToNextChange = new Dictionary<IStream, int>();
 			foreach (IStream stream in streams)
 			{
-				RedisList<int> changes = RedisStreamChanges(stream.Id);
+				RedisValue lastChange = await _redis.StringGetAsync(LastChangeKey(stream.Id));
 
-				int firstChange = await changes.GetByIndexAsync(-1);
-				if (firstChange != 0)
+				int nextChange = 0;
+				if (lastChange.HasValue && lastChange.IsInteger)
 				{
-					firstChange++;
+					nextChange = (int)lastChange + 1;
 				}
-
-				streamToFirstChange[stream] = firstChange;
+				
+				streamToNextChange[stream] = nextChange;
 			}
 
 			// Update the database with all the commits
 			HashSet<StreamId> newStreamIds = new HashSet<StreamId>();
-			await foreach (NewCommit newCommit in FindCommitsForClusterAsync(clusterName, streamToFirstChange))
+			await foreach (NewCommit newCommit in FindCommitsForClusterAsync(clusterName, streamToNextChange))
 			{
 				ICommit commit = await _commitCollection.AddOrReplaceAsync(newCommit);
 				_logger.LogInformation("Replicated CL {Change}: {Description} as {CommitId}", commit.Change, StringUtils.Truncate(commit.Description, 40), commit.Id);
 
-				await NotifyListeners(commit.Id);
-
-				RedisList<int> streamCommitsKey = RedisStreamChanges(commit.StreamId);
-				await streamCommitsKey.RightPushAsync(commit.Change);
+				await NotifyListenersAsync(commit);
 
 				newStreamIds.Add(commit.StreamId);
 			}
-
-			// Signal to any listeners that we have new data to process
-			foreach (StreamId newStreamId in newStreamIds)
-			{
-				await _redisDirtyStreams.AddAsync(newStreamId);
-				await _redis.PublishAsync(RedisUpdateChannel, newStreamId);
-			}
 		}
 
 		/// <summary>
@@ -626,688 +504,5 @@ namespace Horde.Build.Perforce
 		}
 
 		#endregion
-
-		#region Content Updates
-
-		const double ReservationTimeSeconds = 60.0;
-		const double ExtendReservationAfterSeconds = 40.0;
-
-		RedisKey RedisBaseKey { get; } = new RedisKey("commits/");
-		RedisChannel<StreamId> RedisUpdateChannel { get; } = new RedisChannel<StreamId>("commits/streams");
-		bool _stopping;
-		readonly AsyncEvent _updateStreamsEvent = new AsyncEvent();
-		readonly RedisSet<StreamId> _redisDirtyStreams;
-		readonly RedisSortedSet<StreamId> _redisReservations;
-		RedisChannelSubscription<StreamId>? _redisUpdateSubscription;
-		Task? _streamUpdateTask;
-
-		RedisList<int> RedisStreamChanges(StreamId streamId) => new RedisList<int>(_redis, RedisBaseKey.Append($"stream/{streamId}/changes"));
-
-		async Task StartContentReplicationAsync()
-		{
-			_stopping = false;
-			_updateStreamsEvent.Reset();
-			_redisUpdateSubscription = await _redis.Multiplexer.GetSubscriber().SubscribeAsync(RedisUpdateChannel, (_, _) => _updateStreamsEvent.Pulse());
-			_streamUpdateTask = Task.Run(() => UpdateContentAsync());
-		}
-
-		async Task StopContentReplicationAsync()
-		{
-			_stopping = true;
-			_updateStreamsEvent.Latch();
-			if (_redisUpdateSubscription != null)
-			{
-				await _redisUpdateSubscription.DisposeAsync();
-				_redisUpdateSubscription = null;
-			}
-			await _streamUpdateTask!;
-		}
-
-		async Task UpdateContentAsync()
-		{
-			List<(StreamId, Task)> backgroundTasks = new List<(StreamId, Task)>(MaxBackgroundTasks);
-			while (!_stopping || backgroundTasks.Count > 0)
-			{
-				try
-				{
-					await UpdateContentInternalAsync(backgroundTasks);
-				}
-				catch (Exception ex)
-				{
-					_logger.LogError(ex, "Exception when updating commit service");
-				}
-			}
-		}
-
-		async Task UpdateContentInternalAsync(List<(StreamId, Task)> backgroundTasks)
-		{
-			// Remove any complete background tasks
-			for (int idx = backgroundTasks.Count - 1; idx >= 0; idx--)
-			{
-				(_, Task backgroundTask) = backgroundTasks[idx];
-				if (backgroundTask.IsCompleted)
-				{
-					if (backgroundTask.Exception != null)
-					{
-						_logger.LogError(backgroundTask.Exception, "Update background task faulted");
-					}
-					backgroundTasks.RemoveAt(idx);
-				}
-			}
-
-			// Create a list of events to wait for
-			List<Task> waitTasks = new List<Task>(backgroundTasks.Select(x => x.Item2));
-			if (!_stopping)
-			{
-				// If we have spare slots for executing background tasks, check if there are any dirty streams
-				Task newStreamTask = _updateStreamsEvent.Task;
-				if (backgroundTasks.Count < MaxBackgroundTasks)
-				{
-					// Expire any reservations that are no longer valid
-					DateTime utcNow = DateTime.UtcNow;
-					await _redisReservations.RemoveRangeByScoreAsync(Double.NegativeInfinity, utcNow.Ticks);
-
-					// Find the streams that we can wait for
-					HashSet<StreamId> checkStreams = new HashSet<StreamId>(await _redisDirtyStreams.MembersAsync());
-					checkStreams.ExceptWith(backgroundTasks.Select(x => x.Item1));
-
-					if (checkStreams.Count > 0)
-					{
-						// Try to start new background tasks
-						double newScore = (utcNow + TimeSpan.FromSeconds(ReservationTimeSeconds)).Ticks;
-						foreach (StreamId checkStream in checkStreams)
-						{
-							if (await _redisReservations.AddAsync(checkStream, newScore, When.NotExists))
-							{
-								Task newTask = Task.Run(() => UpdateStreamContentAsync(checkStream, CancellationToken.None));
-								backgroundTasks.Add((checkStream, newTask));
-
-								if (backgroundTasks.Count == MaxBackgroundTasks)
-								{
-									break;
-								}
-							}
-						}
-						checkStreams.ExceptWith(backgroundTasks.Select(x => x.Item1));
-
-						// If we still have spare tasks, check how long we should wait for a reservation to expire
-						if (backgroundTasks.Count < MaxBackgroundTasks)
-						{
-							long waitTime = Int64.MaxValue;
-							await foreach (SortedSetEntry<StreamId> entry in _redisReservations.ScanAsync())
-							{
-								if (checkStreams.Contains(entry.Element))
-								{
-									waitTime = Math.Min(waitTime, (long)entry.Score - utcNow.Ticks);
-								}
-							}
-							if (waitTime < Int32.MaxValue)
-							{
-								waitTasks.Add(Task.Delay(new TimeSpan(Math.Max(waitTime, 0))));
-							}
-						}
-					}
-				}
-
-				// If we still have bandwidth to process more tasks, wait for new streams to become available
-				if (backgroundTasks.Count < MaxBackgroundTasks)
-				{
-					waitTasks.Add(newStreamTask);
-				}
-			}
-
-			// Wait until any task has completed
-			if (waitTasks.Count > 0)
-			{
-				await Task.WhenAny(waitTasks);
-			}
-		}
-
-		async Task UpdateStreamContentAsync(StreamId streamId, CancellationToken cancellationToken)
-		{
-			Stopwatch timer = Stopwatch.StartNew();
-
-			RedisList<int> streamChanges = RedisStreamChanges(streamId);
-			for (; ; )
-			{
-				IStream? stream = await _streamCollection.GetAsync(streamId);
-				if (stream == null || stream.ReplicationMode == ContentReplicationMode.None)
-				{
-					// Remove all but the last item
-					await streamChanges.TrimAsync(0, -2);
-				}
-				else
-				{
-					// Update the stream, updating the reservation every 30 seconds
-					Task internalTask = Task.Run(() => UpdateStreamContentInternalAsync(stream, streamChanges, cancellationToken), cancellationToken);
-					while (!internalTask.IsCompleted)
-					{
-						Task delayTask = Task.Delay(TimeSpan.FromSeconds(ExtendReservationAfterSeconds), cancellationToken);
-						if (await Task.WhenAny(internalTask, delayTask) == delayTask)
-						{
-							DateTime newTime = DateTime.UtcNow + TimeSpan.FromSeconds(ReservationTimeSeconds);
-							await _redisReservations.AddAsync(streamId, newTime.Ticks);
-							_logger.LogInformation("Extending reservation for content update of {StreamId} (elapsed: {Time}s)", streamId, (int)timer.Elapsed.TotalSeconds);
-						}
-					}
-
-					// Log any error during the update
-					if (internalTask.IsFaulted)
-					{
-						_logger.LogError(internalTask.Exception, "Exception while updating stream content for {StreamId}", streamId);
-					}
-					else
-					{
-						_logger.LogInformation("Finished update for {StreamId}", streamId);
-					}
-				}
-
-				// Remove this stream from the dirty list if it's empty
-				ITransaction transaction = _redis.CreateTransaction();
-				transaction.AddCondition(Condition.ListLengthLessThan(streamChanges.Key, 2));
-				_ = transaction.With(_redisDirtyStreams).RemoveAsync(streamId);
-				if (await transaction.ExecuteAsync())
-				{
-					break;
-				}
-			}
-			await _redisReservations.RemoveAsync(streamId);
-		}
-
-		async Task UpdateStreamContentInternalAsync(IStream stream, RedisList<int> changes, CancellationToken cancellationToken)
-		{
-			int prevChange = 0;
-			CommitNode prevContents = new CommitNode();
-			for (; ; )
-			{
-				// Get the first two changes to be mirrored. The first one should already exist, unless it's the start of replication for this stream.
-				int[] values = await changes.RangeAsync(0, 1);
-				if(values.Length == 0)
-				{
-					break;
-				}
-
-				// Get or add the tree for the first change.
-				if (prevChange != values[0])
-				{
-					RefName prevRefName = GetRefName(stream, values[0]);
-
-					CommitNode? contents = await _treeStore.TryReadTreeAsync<CommitNode>(prevRefName, cancellationToken);
-					if (contents == null)
-					{
-						_logger.LogInformation("No content for CL {Change}; creating full snapshot", values[0]);
-						prevContents = await WriteCommitTreeAsync(stream, values[0], 0, new CommitNode(), cancellationToken);
-					}
-					else
-					{
-						_logger.LogInformation("Reading existing commit tree for CL {Change} from ref {RefName}", values[0], prevRefName);
-						prevContents = contents;
-					}
-
-					prevChange = values[0];
-				}
-				else if (values.Length == 2)
-				{
-					// Perform a snapshot of the new change, then remove it from the list
-					prevContents = await WriteCommitTreeAsync(stream, values[1], prevChange, prevContents, cancellationToken);
-					prevChange = values[1];
-
-					// Remove the first item from the list
-					ITransaction transaction = _redis.CreateTransaction();
-					transaction.AddCondition(Condition.ListIndexEqual(changes.Key, 0, values[0]));
-					transaction.AddCondition(Condition.ListIndexEqual(changes.Key, 1, values[1]));
-					_ = transaction.With(changes).LeftPopAsync();
-					await transaction.ExecuteAsync();
-				}
-				else
-				{
-					// Nothing to do
-					break;
-				}
-			}
-		}
-
-		/// <summary>
-		/// Gets the ref for a particular change
-		/// </summary>
-		/// <param name="streamId"></param>
-		/// <param name="change"></param>
-		/// <param name="filter"></param>
-		/// <param name="revisionsOnly"></param>
-		/// <returns></returns>
-		static RefName GetRefName(StreamId streamId, int change, string? filter, bool revisionsOnly)
-		{
-			StringBuilder builder = new StringBuilder($"v2/{streamId}/{change}");
-			if (filter != null)
-			{
-				builder.Append($"_filter_{IoHash.Compute(Encoding.UTF8.GetBytes(filter))}");
-			}
-			if (revisionsOnly)
-			{
-				builder.Append("_revs");
-			}
-			return new RefName(builder.ToString());
-		}
-
-		/// <summary>
-		/// Gets the ref for a particular change
-		/// </summary>
-		/// <param name="stream"></param>
-		/// <param name="change"></param>
-		/// <returns></returns>
-		static RefName GetRefName(IStream stream, int change)
-		{
-			return GetRefName(stream.Id, change, stream.ReplicationFilter, stream.ReplicationMode == ContentReplicationMode.RevisionsOnly);
-		}
-
-		/// <summary>
-		/// Reads a tree from the given stream
-		/// </summary>
-		/// <param name="stream"></param>
-		/// <param name="change"></param>
-		/// <param name="filter"></param>
-		/// <param name="revisionsOnly"></param>
-		/// <param name="cancellationToken"></param>
-		/// <returns></returns>
-		public async Task<CommitNode> ReadCommitTreeAsync(IStream stream, int change, string? filter, bool revisionsOnly, CancellationToken cancellationToken)
-		{
-			RefName name = GetRefName(stream.Id, change, filter, revisionsOnly);
-			return await _treeStore.ReadTreeAsync<CommitNode>(name, cancellationToken);
-		}
-
-		/// <summary>
-		/// Replicates the contents of a stream to Horde storage, optionally using the given change as a starting point
-		/// </summary>
-		/// <param name="stream">The stream to replicate</param>
-		/// <param name="change">Commit to store the tree ref</param>
-		/// <param name="baseChange">The base change to update from</param>
-		/// <param name="baseContents">Initial contents of the tree at baseChange</param>
-		/// <param name="cancellationToken">Cancellation token for the operation</param>
-		/// <returns>Root tree object</returns>
-		public async Task<CommitNode> WriteCommitTreeAsync(IStream stream, int change, int baseChange, CommitNode baseContents, CancellationToken cancellationToken)
-		{
-			bool revisionsOnly = stream.ReplicationMode == ContentReplicationMode.RevisionsOnly;
-			return await WriteCommitTreeAsync(stream, change, baseChange, baseContents, stream.ReplicationFilter, revisionsOnly, cancellationToken);
-		}
-
-		/// <summary>
-		/// Sorts paths by their folder, then by their filename.
-		/// </summary>
-		class FolderFirstSorter : IComparer<Utf8String>
-		{
-			public static FolderFirstSorter Instance { get; } = new FolderFirstSorter();
-
-			public int Compare(Utf8String x, Utf8String y)
-			{
-				Utf8String pathX = x.Substring(0, x.LastIndexOf('/') + 1);
-				Utf8String pathY = y.Substring(0, y.LastIndexOf('/') + 1);
-
-				int result = pathX.CompareTo(pathY);
-				if (result != 0)
-				{
-					return result;
-				}
-
-				return x.CompareTo(y);
-			}
-		}
-
-		/// <summary>
-		/// Replicates the contents of a stream to Horde storage, optionally using the given change as a starting point
-		/// </summary>
-		/// <param name="stream">The stream to replicate</param>
-		/// <param name="change">Commit to store the tree ref</param>
-		/// <param name="baseChange">The base change to update from</param>
-		/// <param name="baseContents">Initial contents of the tree at baseChange</param>
-		/// <param name="filter"></param>
-		/// <param name="revisionsOnly"></param>
-		/// <param name="cancellationToken">Cancellation token for the operation</param>
-		/// <returns>Root tree object</returns>
-		public async Task<CommitNode> WriteCommitTreeAsync(IStream stream, int change, int? baseChange, CommitNode baseContents, string? filter, bool revisionsOnly, CancellationToken cancellationToken)
-		{
-			// Create a client to replicate from this stream
-			ReplicationClient clientInfo = await FindOrAddReplicationClientAsync(stream);
-
-			// Connect to the server and flush the workspace
-			using IPerforceConnection perforce = await PerforceConnection.CreateAsync(clientInfo.Settings, _logger);
-			await FlushWorkspaceAsync(clientInfo, perforce, baseChange ?? 0);
-
-			// Apply all the updates
-			_logger.LogInformation("Syncing client {Client} from changelist {BaseChange} to {Change}", clientInfo.Client.Name, baseChange ?? 0, change);
-			clientInfo.Change = -1;
-
-			Utf8String clientRoot = new Utf8String(clientInfo.Client.Root);
-			string filterOrDefault = filter ?? "...";
-			string queryPath = $"//{clientInfo.Client.Name}/{filterOrDefault}";
-
-			RefName refName = GetRefName(stream.Id, change, filter, revisionsOnly);
-			RefName incRefName = new RefName($"{refName}_inc");
-
-			// Get the current sync state for this change
-			CommitNode? syncNode = await _treeStore.TryReadTreeAsync<CommitNode>(incRefName, cancellationToken);
-			bool deleteIncRef = syncNode != null;
-			if (syncNode == null)
-			{
-				DirectoryNode contents = await baseContents.Contents.ExpandAsync(cancellationToken);
-				syncNode = new CommitNode(contents);
-			}
-
-			// Get the contents of the current tree
-			DirectoryNode root = await syncNode.Contents.ExpandAsync(cancellationToken);
-
-			if (revisionsOnly)
-			{
-				int count = 0;
-				await foreach (PerforceResponse<SyncRecord> record in perforce.StreamCommandAsync<SyncRecord>("sync", new[] { "-k", $"{queryPath}@{change}" }, null, default, cancellationToken))
-				{
-					PerforceError? error = record.Error;
-					if (error != null && error.Generic == PerforceGenericCode.Empty)
-					{
-						continue;
-					}
-
-					SyncRecord syncRecord = record.Data;
-					if (!syncRecord.Path.StartsWith(clientRoot))
-					{
-						throw new ArgumentException($"Unable to make path {clientInfo.Client.Root} relative to client root {clientInfo.Client.Root}");
-					}
-
-					Utf8String path = syncRecord.Path.Substring(clientRoot.Length);
-					byte[] data = Encoding.UTF8.GetBytes($"{syncRecord.DepotFile}#{syncRecord.Revision}");
-
-					using ReadOnlyMemoryStream dataStream = new ReadOnlyMemoryStream(data);
-
-					FileEntry entry = await root.AddFileByPathAsync(path, FileEntryFlags.PerforceDepotPathAndRevision, cancellationToken);
-					await entry.AppendAsync(dataStream, Options.Chunking, cancellationToken);
-
-					if (++count > 1000)
-					{
-						await _treeStore.WriteTreeAsync(refName, syncNode, false, cancellationToken);
-						count = 0;
-					}
-				}
-			}
-			else
-			{
-				// Replay the files that have already been synced
-				foreach (Utf8String path in syncNode.Paths)
-				{
-					string flushPath = $"//{clientInfo.Client.Name}/{path}@{change}";
-					_logger.LogInformation("Flushing {FlushPath}", flushPath);
-					await perforce.SyncQuietAsync(SyncOptions.Force | SyncOptions.KeepWorkspaceFiles, -1, flushPath, cancellationToken);
-				}
-
-				// Do a sync preview to find everything that's left, and sort the remaining list of paths
-				List<(Utf8String Path, long Size)> files = new List<(Utf8String, long)>();
-				await foreach (PerforceResponse<SyncRecord> response in perforce.StreamCommandAsync<SyncRecord>("sync", new[] { "-n" }, new string[] { $"{queryPath}@{change}" }, null, cancellationToken))
-				{
-					PerforceError? error = response.Error;
-					if (error != null)
-					{
-						_logger.LogWarning("Perforce: {Message}", error.Data);
-						continue;
-					}
-
-					Utf8String file = response.Data.Path.Clone();
-					if (!file.StartsWith(clientRoot, Utf8StringComparer.Ordinal))
-					{
-						throw new ArgumentException($"Unable to make path {file} relative to client root {clientRoot}");
-					}
-
-					byte[] path = new byte[file.Length - clientRoot.Length];
-					for (int idx = clientRoot.Length; idx < file.Length; idx++)
-					{
-						path[idx - clientRoot.Length] = (file[idx] == '\\') ? (byte)'/' : file[idx];
-					}
-
-					files.Add((new Utf8String(path), response.Data.FileSize));
-				}
-				files.SortBy(x => x.Path, FolderFirstSorter.Instance);
-
-				// Sync incrementally
-				long totalSize = files.Sum(x => x.Size);
-				long syncedSize = 0;
-				while (files.Count > 0)
-				{
-					// Save the incremental state
-					if (syncedSize > 0)
-					{
-						await _treeStore.WriteTreeAsync(incRefName, syncNode, cancellationToken: cancellationToken);
-						deleteIncRef = true;
-					}
-
-					// Find the next path to sync
-					const long MaxBatchSize = 10L * 1024 * 1024 * 1024;
-					(int idx, Utf8String path, long size) = GetSyncBatch(files, MaxBatchSize, _logger);
-					syncedSize += size;
-
-					const long TrimInterval = 1024 * 1024 * 256;
-
-					long dataSize = 0;
-					long trimDataSize = TrimInterval;
-
-					string syncPath = $"//{clientInfo.Client.Name}/{path}@{change}";
-					double syncPct = (totalSize == 0) ? 100.0 : (syncedSize * 100.0) / totalSize;
-					_logger.LogInformation("Syncing {StreamId} to {Change} [{SyncPct:n1}%]: {Path} ({Size:n1}mb)", stream.Id, change, syncPct, path, size / (1024.0 * 1024.0));
-
-					Dictionary<int, FileEntry> handles = new Dictionary<int, FileEntry>();
-					await foreach (PerforceResponse response in perforce.StreamCommandAsync("sync", Array.Empty<string>(), new string[] { syncPath }, null, typeof(SyncRecord), true, default))
-					{
-						PerforceError? error = response.Error;
-						if (error != null)
-						{
-							_logger.LogWarning("Perforce: {Message}", error.Data);
-							continue;
-						}
-
-						PerforceIo? io = response.Io;
-						if (io != null)
-						{
-							if (io.Command == PerforceIoCommand.Open)
-							{
-								Utf8String file = GetClientRelativePath(io.Payload, clientInfo.Client.Root);
-								handles[io.File] = await root.AddFileByPathAsync(file, FileEntryFlags.None, cancellationToken);
-							}
-							else if (io.Command == PerforceIoCommand.Write)
-							{
-								FileEntry entry = handles[io.File];
-								await entry.AppendAsync(io.Payload, Options.Chunking, cancellationToken);
-								dataSize += io.Payload.Length;
-							}
-							else if (io.Command == PerforceIoCommand.Close)
-							{
-								handles.Remove(io.File);
-							}
-							else if (io.Command == PerforceIoCommand.Unlink)
-							{
-								Utf8String file = GetClientRelativePath(io.Payload, clientInfo.Client.Root);
-								await root.DeleteFileByPathAsync(file, cancellationToken);
-							}
-							else
-							{
-								_logger.LogWarning("Unhandled command code {Code}", io.Command);
-							}
-
-							if (dataSize > trimDataSize)
-							{
-								_logger.LogInformation("Trimming working set after receiving {Size:n1}mb...", dataSize / (1024 * 1024));
-								await _treeStore.WriteTreeAsync(refName, syncNode, false, cancellationToken);
-								GC.Collect();
-								trimDataSize = dataSize + TrimInterval;
-								_logger.LogInformation("Trimming complete. Next trim at {Size:n1}mb.", trimDataSize / (1024 * 1024));
-							}
-						}
-					}
-
-					// Update the root sync node
-					syncNode.Paths.Add(path);
-					files.RemoveRange(idx, files.Count - idx);
-				}
-			}
-
-			clientInfo.Change = change;
-
-			// Return the new root object
-			_logger.LogInformation("Writing ref {RefId} for {StreamId} change {Change}", refName, stream.Id, change);
-			await _treeStore.WriteTreeAsync(refName, syncNode, true, cancellationToken);
-
-			// Delete the incremental state
-			if (deleteIncRef)
-			{
-				await _treeStore.DeleteTreeAsync(incRefName, cancellationToken);
-			}
-			return syncNode;
-		}
-
-		static (int, Utf8String, long) GetSyncBatch(List<(Utf8String Path, long Size)> files, long maxSize, ILogger logger)
-		{
-			int idx = files.Count - 1;
-			Utf8String path = files[idx].Path;
-			long size = files[idx].Size;
-
-			Utf8String prefix = path;
-			while (prefix.Length > 0)
-			{
-				// Get the length of the parent directory
-				int endIdx = prefix.Length - 1;
-				while (endIdx > 0 && path[endIdx - 1] != (byte)'/')
-				{
-					endIdx--;
-				}
-
-				// Get the parent directory prefix, ending with a slash
-				prefix = prefix.Substring(0, endIdx);
-
-				// Include all files with this prefix
-				int nextIdx = idx;
-				long nextSize = size;
-				while(nextIdx > 0 && files[nextIdx - 1].Path.StartsWith(prefix))
-				{
-					nextSize += files[nextIdx - 1].Size;
-					nextIdx--;
-				}
-				if (nextSize > maxSize)
-				{
-					logger.LogInformation("Next filter {NextFilter} would exceed max size ({NextSize:n1}mb > {MaxSize:n1}mb); using {Filter} ({Size:n1}mb) instead.", prefix + "...", nextSize / (1024.0 * 1024.0), maxSize / (1024.0 * 1024.0), path, size / (1024.0 * 1024.0));
-					return (idx, path, size);
-				}
-				size = nextSize;
-				idx = nextIdx;
-
-				// Update the sync path to be the prefixed directory
-				path = prefix + "...";
-			}
-
-			return (idx, path, size);
-		}
-
-		async Task FlushWorkspaceAsync(ReplicationClient clientInfo, IPerforceConnection perforce, int change)
-		{
-			if (clientInfo.Change != change)
-			{
-				clientInfo.Change = -1;
-				if (change == 0)
-				{
-					_logger.LogInformation("Flushing have table for {Client}", clientInfo.Client.Name);
-					await perforce.SyncQuietAsync(SyncOptions.Force | SyncOptions.KeepWorkspaceFiles, -1, $"//{clientInfo.Client.Name}/...#0");
-				}
-				else
-				{
-					_logger.LogInformation("Flushing have table for {Client} to change {Change}", clientInfo.Client.Name, change);
-					await perforce.SyncQuietAsync(SyncOptions.Force | SyncOptions.KeepWorkspaceFiles, -1, $"//{clientInfo.Client.Name}/...@{change}");
-				}
-				clientInfo.Change = change;
-			}
-		}
-
-		static Utf8String GetClientRelativePath(ReadOnlyMemory<byte> data, Utf8String clientRoot)
-		{
-			int length = data.Span.IndexOf((byte)0);
-			if (length != -1)
-			{
-				data = data.Slice(0, length);
-			}
-
-			Utf8String path = new Utf8String(data);
-			if (!path.StartsWith(clientRoot, Utf8StringComparer.Ordinal))
-			{
-				throw new ArgumentException($"Unable to make path {path} relative to client root {clientRoot}");
-			}
-
-			return path.Substring(clientRoot.Length).Clone();
-		}
-
-		class ReplicationClient
-		{
-			public PerforceSettings Settings { get; }
-			public string ClusterName { get; }
-			public InfoRecord ServerInfo { get; }
-			public ClientRecord Client { get; }
-			public int Change { get; set; }
-
-			public ReplicationClient(PerforceSettings settings, string clusterName, InfoRecord serverInfo, ClientRecord client, int change)
-			{
-				Settings = settings;
-				ClusterName = clusterName;
-				ServerInfo = serverInfo;
-				Client = client;
-				Change = change;
-			}
-		}
-
-		readonly Dictionary<StreamId, ReplicationClient> _cachedPerforceClients = new Dictionary<StreamId, ReplicationClient>();
-
-		async Task<ReplicationClient?> FindReplicationClientAsync(IStream stream)
-		{
-			ReplicationClient? clientInfo;
-			if (_cachedPerforceClients.TryGetValue(stream.Id, out clientInfo))
-			{
-				if (!String.Equals(clientInfo.ClusterName, stream.ClusterName, StringComparison.Ordinal) && String.Equals(clientInfo.Client.Stream, stream.Name, StringComparison.Ordinal))
-				{
-					PerforceSettings serverSettings = new PerforceSettings(clientInfo.Settings);
-					serverSettings.ClientName = null;
-
-					using IPerforceConnection perforce = await PerforceConnection.CreateAsync(_logger);
-					await perforce.DeleteClientAsync(DeleteClientOptions.None, clientInfo.Client.Name);
-
-					_cachedPerforceClients.Remove(stream.Id);
-					clientInfo = null;
-				}
-			}
-			return clientInfo;
-		}
-
-		async Task<ReplicationClient> FindOrAddReplicationClientAsync(IStream stream)
-		{
-			ReplicationClient? clientInfo = await FindReplicationClientAsync(stream);
-			if (clientInfo == null)
-			{
-				using IPerforceConnection? perforce = await _perforceService.GetServiceUserConnection(stream.ClusterName);
-				if (perforce == null)
-				{
-					throw new PerforceException($"Unable to create connection to Perforce server");
-				}
-
-				InfoRecord serverInfo = await perforce.GetInfoAsync(InfoOptions.ShortOutput);
-
-				ClientRecord newClient = new ClientRecord($"Horde.Build_{serverInfo.ClientHost}_{stream.Id}", perforce.Settings.UserName, "/p4/");
-				newClient.Description = "Created to mirror Perforce content to Horde Storage";
-				newClient.Owner = perforce.Settings.UserName;
-				newClient.Host = serverInfo.ClientHost;
-				newClient.Stream = stream.Name;
-				await perforce.CreateClientAsync(newClient);
-				_logger.LogInformation("Created client {ClientName} for {StreamName}", newClient.Name, stream.Name);
-
-				PerforceSettings settings = new PerforceSettings(perforce.Settings);
-				settings.ClientName = newClient.Name;
-				settings.PreferNativeClient = true;
-
-				clientInfo = new ReplicationClient(settings, stream.ClusterName, serverInfo, newClient, -1);
-				_cachedPerforceClients.Add(stream.Id, clientInfo);
-			}
-			return clientInfo;
-		}
-
-		#endregion
 	}
 }
diff --git a/Engine/Source/Programs/Horde/Horde.Build/Perforce/ICommitService.cs b/Engine/Source/Programs/Horde/Horde.Build/Perforce/ICommitService.cs
--- a/Engine/Source/Programs/Horde/Horde.Build/Perforce/ICommitService.cs
+++ b/Engine/Source/Programs/Horde/Horde.Build/Perforce/ICommitService.cs
@@ -1,6 +1,7 @@
 // Copyright Epic Games, Inc. All Rights Reserved.
 
 using System;
+using System.Threading.Tasks;
 using EpicGames.Horde.Storage;
 using Horde.Build.Streams;
 using Horde.Build.Utilities;
@@ -17,6 +18,6 @@ namespace Horde.Build.Perforce
 		/// </summary>
 		/// <param name="onAddCommit">Callback for a new commit being added</param>
 		/// <returns>Disposable handler.</returns>
-		IDisposable AddListener(Action<ICommit> onAddCommit);
+		IAsyncDisposable AddListener(Func<ICommit, Task> onAddCommit);
 	}
 }
diff --git a/Engine/Source/Programs/Horde/Horde.Build/Perforce/ReplicationService.cs b/Engine/Source/Programs/Horde/Horde.Build/Perforce/ReplicationService.cs
new file mode 100644
--- /dev/null
+++ b/Engine/Source/Programs/Horde/Horde.Build/Perforce/ReplicationService.cs
@@ -0,0 +1,902 @@
+// Copyright Epic Games, Inc. All Rights Reserved.
+
+using System;
+using System.Buffers;
+using System.Collections.Generic;
+using System.Diagnostics;
+using System.Globalization;
+using System.IO;
+using System.Linq;
+using System.Text;
+using System.Text.RegularExpressions;
+using System.Threading;
+using System.Threading.Tasks;
+using EpicGames.Core;
+using EpicGames.Horde.Storage;
+using EpicGames.Horde.Storage.Bundles;
+using EpicGames.Horde.Storage.Nodes;
+using EpicGames.Perforce;
+using EpicGames.Redis;
+using EpicGames.Serialization;
+using Horde.Build.Server;
+using Horde.Build.Streams;
+using Horde.Build.Users;
+using Horde.Build.Utilities;
+using HordeCommon;
+using Microsoft.Extensions.Caching.Memory;
+using Microsoft.Extensions.Hosting;
+using Microsoft.Extensions.Logging;
+using Microsoft.Extensions.Options;
+using StackExchange.Redis;
+
+namespace Horde.Build.Perforce
+{
+	using CommitId = ObjectId<ICommit>;
+	using StreamId = StringId<IStream>;
+
+	/// <summary>
+	/// Options for the commit service
+	/// </summary>
+	public class ReplicationServiceOptions
+	{
+		/// <summary>
+		/// Whether to enable replication. Must also be enabled on a per-stream basis (see <see cref="IStream.ReplicationMode"/>).
+		/// </summary>
+		public bool Enable { get; set; } = true;
+
+		/// <summary>
+		/// Options for how objects are packed together
+		/// </summary>
+		public BundleOptions Bundle { get; set; } = new BundleOptions();
+
+		/// <summary>
+		/// Options for how objects are sliced
+		/// </summary>
+		public ChunkingOptions Chunking { get; set; } = new ChunkingOptions();
+	}
+
+	/// <summary>
+	/// Root node for a commit snapshot
+	/// </summary>
+	[TreeSerializer(typeof(ReplicationNodeSerializer))]
+	public class ReplicationNode : TreeNode
+	{
+		/// <summary>
+		/// Paths that have been synced. Empty once complete.
+		/// </summary>
+		public List<Utf8String> Paths { get; }
+
+		/// <summary>
+		/// Contents of this snapshot
+		/// </summary>
+		public TreeNodeRef<DirectoryNode> Contents { get; }
+
+		/// <summary>
+		/// Constructor
+		/// </summary>
+		public ReplicationNode()
+			: this(new DirectoryNode())
+		{
+		}
+
+		/// <summary>
+		/// Constructor
+		/// </summary>
+		public ReplicationNode(DirectoryNode contents)
+		{
+			Paths = new List<Utf8String>();
+			Contents = new TreeNodeRef<DirectoryNode>(this, contents);
+		}
+
+		private ReplicationNode(IEnumerable<Utf8String> paths, ITreeBlob contents)
+		{
+			Paths = new List<Utf8String>(paths);
+			Contents = new TreeNodeRef<DirectoryNode>(this, contents);
+		}
+
+		internal async ValueTask<ITreeBlob> SerializeAsync(ITreeBlobWriter writer, CancellationToken cancellationToken)
+		{
+			List<ITreeBlob> references = new List<ITreeBlob>();
+			references.Add(await Contents.CollapseAsync(writer, cancellationToken));
+
+			ByteArrayBuilder builder = new ByteArrayBuilder();
+			builder.WriteVariableLengthArray(Paths, x => builder.WriteUtf8String(x));
+			return await writer.WriteBlobAsync(builder.AsSequence(), references, cancellationToken);
+		}
+
+		internal static async ValueTask<ReplicationNode> DeserializeAsync(ITreeBlob blob, CancellationToken cancellationToken)
+		{
+			IReadOnlyList<ITreeBlob> children = await blob.GetReferencesAsync(cancellationToken);
+			ReadOnlySequence<byte> data = await blob.GetDataAsync(cancellationToken);
+
+			MemoryReader reader = new MemoryReader(data.AsSingleSegment());
+			Utf8String[] paths = reader.ReadVariableLengthArray(() => reader.ReadUtf8String());
+
+			return new ReplicationNode(paths, children[0]);
+		}
+
+		/// <inheritdoc/>
+		public override IReadOnlyList<TreeNodeRef> GetReferences()
+		{
+			List<TreeNodeRef> refs = new List<TreeNodeRef>();
+			refs.Add(Contents);
+			return refs;
+		}
+	}
+
+	class ReplicationNodeSerializer : TreeNodeSerializer<ReplicationNode>
+	{
+		/// <inheritdoc/>
+		public override ValueTask<ReplicationNode> DeserializeAsync(ITreeBlob node, CancellationToken cancellationToken)
+			=> ReplicationNode.DeserializeAsync(node, cancellationToken);
+
+		/// <inheritdoc/>
+		public override ValueTask<ITreeBlob> SerializeAsync(ITreeBlobWriter writer, ReplicationNode node, CancellationToken cancellationToken)
+			=> node.SerializeAsync(writer, cancellationToken);
+	}
+
+	/// <summary>
+	/// Service which replicates content from Perforce
+	/// </summary>
+	class ReplicationService : IHostedService, IAsyncDisposable
+	{
+		/// <summary>
+		/// Metadata about a stream that needs to have commits mirrored
+		/// </summary>
+		class StreamInfo
+		{
+			public IStream Stream { get; set; }
+			public ViewMap View { get; }
+
+			public StreamInfo(IStream stream, ViewMap view)
+			{
+				Stream = stream;
+				View = view;
+			}
+		}
+
+		public ReplicationServiceOptions Options { get; }
+
+		// Redis
+		readonly IDatabase _redis;
+
+		// Collections
+		readonly CommitService _commitService;
+		readonly IAsyncDisposable _commitListener;
+
+		readonly IStreamCollection _streamCollection;
+		readonly IPerforceService _perforceService;
+		readonly ITreeStore _treeStore;
+		readonly ILogger _logger;
+
+		const int MaxBackgroundTasks = 2;
+
+		const double ReservationTimeSeconds = 60.0;
+		const double ExtendReservationAfterSeconds = 40.0;
+
+		RedisKey RedisBaseKey { get; } = new RedisKey("commits/replication/");
+		RedisChannel<StreamId> RedisUpdateChannel { get; } = new RedisChannel<StreamId>("commits/replication/streams");
+		bool _stopping;
+		readonly AsyncEvent _updateStreamsEvent = new AsyncEvent();
+		readonly RedisSet<StreamId> _redisDirtyStreams;
+		readonly RedisSortedSet<StreamId> _redisReservations;
+		RedisChannelSubscription<StreamId>? _redisUpdateSubscription;
+		Task? _streamUpdateTask;
+
+		/// <summary>
+		/// Constructor
+		/// </summary>
+		public ReplicationService(RedisService redisService, CommitService commitService, IStreamCollection streamCollection, IPerforceService perforceService, ITreeStore<ReplicationService> treeStore, IClock clock, IOptions<ReplicationServiceOptions> options, ILogger<ReplicationService> logger)
+		{
+			Options = options.Value;
+
+			_redis = redisService.Database;
+			_redisDirtyStreams = new RedisSet<StreamId>(_redis, RedisBaseKey.Append("streams"));
+			_redisReservations = new RedisSortedSet<StreamId>(_redis, RedisBaseKey.Append("reservations"));
+
+			_commitService = commitService;
+			_commitListener = _commitService.AddListener(OnCommitAdded);
+
+			_streamCollection = streamCollection;
+			_perforceService = perforceService;
+			_treeStore = treeStore;
+			_logger = logger;
+		}
+
+		/// <inheritdoc/>
+		public async ValueTask DisposeAsync()
+		{
+			await _commitListener.DisposeAsync();
+		}
+
+		/// <inheritdoc/>
+		public async Task StartAsync(CancellationToken cancellationToken)
+		{
+			if (Options.Enable)
+			{
+				_stopping = false;
+				_updateStreamsEvent.Reset();
+				_redisUpdateSubscription = await _redis.Multiplexer.GetSubscriber().SubscribeAsync(RedisUpdateChannel, (_, _) => _updateStreamsEvent.Pulse());
+				_streamUpdateTask = Task.Run(() => UpdateContentAsync(), cancellationToken);
+			}
+		}
+
+		/// <inheritdoc/>
+		public async Task StopAsync(CancellationToken cancellationToken)
+		{
+			if (Options.Enable)
+			{
+				_stopping = true;
+				_updateStreamsEvent.Latch();
+				if (_redisUpdateSubscription != null)
+				{
+					await _redisUpdateSubscription.DisposeAsync();
+					_redisUpdateSubscription = null;
+				}
+				await _streamUpdateTask!;
+				await _commitListener.DisposeAsync();
+			}
+		}
+
+		RedisList<int> RedisStreamChanges(StreamId streamId) => new RedisList<int>(_redis, RedisBaseKey.Append($"stream/{streamId}/changes"));
+
+		async Task OnCommitAdded(ICommit commit)
+		{
+			// Post the new commit key
+			RedisList<int> streamCommitsKey = RedisStreamChanges(commit.StreamId);
+			await streamCommitsKey.RightPushAsync(commit.Change);
+
+			// Signal to any listeners that we have new data to process
+			await _redisDirtyStreams.AddAsync(commit.StreamId);
+			await _redis.PublishAsync(RedisUpdateChannel, commit.StreamId);
+		}
+
+		async Task UpdateContentAsync()
+		{
+			List<(StreamId, Task)> backgroundTasks = new List<(StreamId, Task)>(MaxBackgroundTasks);
+			while (!_stopping || backgroundTasks.Count > 0)
+			{
+				try
+				{
+					await UpdateContentInternalAsync(backgroundTasks);
+				}
+				catch (Exception ex)
+				{
+					_logger.LogError(ex, "Exception when updating commit service");
+				}
+			}
+		}
+
+		async Task UpdateContentInternalAsync(List<(StreamId, Task)> backgroundTasks)
+		{
+			// Remove any complete background tasks
+			for (int idx = backgroundTasks.Count - 1; idx >= 0; idx--)
+			{
+				(_, Task backgroundTask) = backgroundTasks[idx];
+				if (backgroundTask.IsCompleted)
+				{
+					if (backgroundTask.Exception != null)
+					{
+						_logger.LogError(backgroundTask.Exception, "Update background task faulted");
+					}
+					backgroundTasks.RemoveAt(idx);
+				}
+			}
+
+			// Create a list of events to wait for
+			List<Task> waitTasks = new List<Task>(backgroundTasks.Select(x => x.Item2));
+			if (!_stopping)
+			{
+				// If we have spare slots for executing background tasks, check if there are any dirty streams
+				Task newStreamTask = _updateStreamsEvent.Task;
+				if (backgroundTasks.Count < MaxBackgroundTasks)
+				{
+					// Expire any reservations that are no longer valid
+					DateTime utcNow = DateTime.UtcNow;
+					await _redisReservations.RemoveRangeByScoreAsync(Double.NegativeInfinity, utcNow.Ticks);
+
+					// Find the streams that we can wait for
+					HashSet<StreamId> checkStreams = new HashSet<StreamId>(await _redisDirtyStreams.MembersAsync());
+					checkStreams.ExceptWith(backgroundTasks.Select(x => x.Item1));
+
+					if (checkStreams.Count > 0)
+					{
+						// Try to start new background tasks
+						double newScore = (utcNow + TimeSpan.FromSeconds(ReservationTimeSeconds)).Ticks;
+						foreach (StreamId checkStream in checkStreams)
+						{
+							if (await _redisReservations.AddAsync(checkStream, newScore, When.NotExists))
+							{
+								Task newTask = Task.Run(() => UpdateStreamContentAsync(checkStream, CancellationToken.None));
+								backgroundTasks.Add((checkStream, newTask));
+
+								if (backgroundTasks.Count == MaxBackgroundTasks)
+								{
+									break;
+								}
+							}
+						}
+						checkStreams.ExceptWith(backgroundTasks.Select(x => x.Item1));
+
+						// If we still have spare tasks, check how long we should wait for a reservation to expire
+						if (backgroundTasks.Count < MaxBackgroundTasks)
+						{
+							long waitTime = Int64.MaxValue;
+							await foreach (SortedSetEntry<StreamId> entry in _redisReservations.ScanAsync())
+							{
+								if (checkStreams.Contains(entry.Element))
+								{
+									waitTime = Math.Min(waitTime, (long)entry.Score - utcNow.Ticks);
+								}
+							}
+							if (waitTime < Int32.MaxValue)
+							{
+								waitTasks.Add(Task.Delay(new TimeSpan(Math.Max(waitTime, 0))));
+							}
+						}
+					}
+				}
+
+				// If we still have bandwidth to process more tasks, wait for new streams to become available
+				if (backgroundTasks.Count < MaxBackgroundTasks)
+				{
+					waitTasks.Add(newStreamTask);
+				}
+			}
+
+			// Wait until any task has completed
+			if (waitTasks.Count > 0)
+			{
+				await Task.WhenAny(waitTasks);
+			}
+		}
+
+		async Task UpdateStreamContentAsync(StreamId streamId, CancellationToken cancellationToken)
+		{
+			Stopwatch timer = Stopwatch.StartNew();
+
+			RedisList<int> streamChanges = RedisStreamChanges(streamId);
+			for (; ; )
+			{
+				IStream? stream = await _streamCollection.GetAsync(streamId);
+				if (stream == null || stream.ReplicationMode == ContentReplicationMode.None)
+				{
+					// Remove all but the last item
+					await streamChanges.TrimAsync(0, -2);
+				}
+				else
+				{
+					// Update the stream, updating the reservation every 30 seconds
+					Task internalTask = Task.Run(() => UpdateStreamContentInternalAsync(stream, streamChanges, cancellationToken), cancellationToken);
+					while (!internalTask.IsCompleted)
+					{
+						Task delayTask = Task.Delay(TimeSpan.FromSeconds(ExtendReservationAfterSeconds), cancellationToken);
+						if (await Task.WhenAny(internalTask, delayTask) == delayTask)
+						{
+							DateTime newTime = DateTime.UtcNow + TimeSpan.FromSeconds(ReservationTimeSeconds);
+							await _redisReservations.AddAsync(streamId, newTime.Ticks);
+							_logger.LogInformation("Extending reservation for content update of {StreamId} (elapsed: {Time}s)", streamId, (int)timer.Elapsed.TotalSeconds);
+						}
+					}
+
+					// Log any error during the update
+					if (internalTask.IsFaulted)
+					{
+						_logger.LogError(internalTask.Exception, "Exception while updating stream content for {StreamId}", streamId);
+					}
+					else
+					{
+						_logger.LogInformation("Finished update for {StreamId}", streamId);
+					}
+				}
+
+				// Remove this stream from the dirty list if it's empty
+				ITransaction transaction = _redis.CreateTransaction();
+				transaction.AddCondition(Condition.ListLengthLessThan(streamChanges.Key, 2));
+				_ = transaction.With(_redisDirtyStreams).RemoveAsync(streamId);
+				if (await transaction.ExecuteAsync())
+				{
+					break;
+				}
+			}
+			await _redisReservations.RemoveAsync(streamId);
+		}
+
+		async Task UpdateStreamContentInternalAsync(IStream stream, RedisList<int> changes, CancellationToken cancellationToken)
+		{
+			int prevChange = 0;
+			ReplicationNode prevContents = new ReplicationNode();
+			for (; ; )
+			{
+				// Get the first two changes to be mirrored. The first one should already exist, unless it's the start of replication for this stream.
+				int[] values = await changes.RangeAsync(0, 1);
+				if (values.Length == 0)
+				{
+					break;
+				}
+
+				// Get or add the tree for the first change.
+				if (prevChange != values[0])
+				{
+					RefName prevRefName = GetRefName(stream, values[0]);
+
+					ReplicationNode? contents = await _treeStore.TryReadTreeAsync<ReplicationNode>(prevRefName, cancellationToken);
+					if (contents == null)
+					{
+						_logger.LogInformation("No content for CL {Change}; creating full snapshot", values[0]);
+						prevContents = await WriteCommitTreeAsync(stream, values[0], 0, new ReplicationNode(), cancellationToken);
+					}
+					else
+					{
+						_logger.LogInformation("Reading existing commit tree for CL {Change} from ref {RefName}", values[0], prevRefName);
+						prevContents = contents;
+					}
+
+					prevChange = values[0];
+				}
+				else if (values.Length == 2)
+				{
+					// Perform a snapshot of the new change, then remove it from the list
+					prevContents = await WriteCommitTreeAsync(stream, values[1], prevChange, prevContents, cancellationToken);
+					prevChange = values[1];
+
+					// Remove the first item from the list
+					ITransaction transaction = _redis.CreateTransaction();
+					transaction.AddCondition(Condition.ListIndexEqual(changes.Key, 0, values[0]));
+					transaction.AddCondition(Condition.ListIndexEqual(changes.Key, 1, values[1]));
+					_ = transaction.With(changes).LeftPopAsync();
+					await transaction.ExecuteAsync();
+				}
+				else
+				{
+					// Nothing to do
+					break;
+				}
+			}
+		}
+
+		/// <summary>
+		/// Gets the ref for a particular change
+		/// </summary>
+		/// <param name="streamId"></param>
+		/// <param name="change"></param>
+		/// <param name="filter"></param>
+		/// <param name="revisionsOnly"></param>
+		/// <returns></returns>
+		static RefName GetRefName(StreamId streamId, int change, string? filter, bool revisionsOnly)
+		{
+			StringBuilder builder = new StringBuilder($"v2/{streamId}/{change}");
+			if (filter != null)
+			{
+				builder.Append($"_filter_{IoHash.Compute(Encoding.UTF8.GetBytes(filter))}");
+			}
+			if (revisionsOnly)
+			{
+				builder.Append("_revs");
+			}
+			return new RefName(builder.ToString());
+		}
+
+		/// <summary>
+		/// Gets the ref for a particular change
+		/// </summary>
+		/// <param name="stream"></param>
+		/// <param name="change"></param>
+		/// <returns></returns>
+		static RefName GetRefName(IStream stream, int change)
+		{
+			return GetRefName(stream.Id, change, stream.ReplicationFilter, stream.ReplicationMode == ContentReplicationMode.RevisionsOnly);
+		}
+
+		/// <summary>
+		/// Reads a tree from the given stream
+		/// </summary>
+		/// <param name="stream"></param>
+		/// <param name="change"></param>
+		/// <param name="filter"></param>
+		/// <param name="revisionsOnly"></param>
+		/// <param name="cancellationToken"></param>
+		/// <returns></returns>
+		public async Task<ReplicationNode> ReadCommitTreeAsync(IStream stream, int change, string? filter, bool revisionsOnly, CancellationToken cancellationToken)
+		{
+			RefName name = GetRefName(stream.Id, change, filter, revisionsOnly);
+			return await _treeStore.ReadTreeAsync<ReplicationNode>(name, cancellationToken);
+		}
+
+		/// <summary>
+		/// Replicates the contents of a stream to Horde storage, optionally using the given change as a starting point
+		/// </summary>
+		/// <param name="stream">The stream to replicate</param>
+		/// <param name="change">Commit to store the tree ref</param>
+		/// <param name="baseChange">The base change to update from</param>
+		/// <param name="baseContents">Initial contents of the tree at baseChange</param>
+		/// <param name="cancellationToken">Cancellation token for the operation</param>
+		/// <returns>Root tree object</returns>
+		public async Task<ReplicationNode> WriteCommitTreeAsync(IStream stream, int change, int baseChange, ReplicationNode baseContents, CancellationToken cancellationToken)
+		{
+			bool revisionsOnly = stream.ReplicationMode == ContentReplicationMode.RevisionsOnly;
+			return await WriteCommitTreeAsync(stream, change, baseChange, baseContents, stream.ReplicationFilter, revisionsOnly, cancellationToken);
+		}
+
+		/// <summary>
+		/// Sorts paths by their folder, then by their filename.
+		/// </summary>
+		class FolderFirstSorter : IComparer<Utf8String>
+		{
+			public static FolderFirstSorter Instance { get; } = new FolderFirstSorter();
+
+			public int Compare(Utf8String x, Utf8String y)
+			{
+				Utf8String pathX = x.Substring(0, x.LastIndexOf('/') + 1);
+				Utf8String pathY = y.Substring(0, y.LastIndexOf('/') + 1);
+
+				int result = pathX.CompareTo(pathY);
+				if (result != 0)
+				{
+					return result;
+				}
+
+				return x.CompareTo(y);
+			}
+		}
+
+		/// <summary>
+		/// Replicates the contents of a stream to Horde storage, optionally using the given change as a starting point
+		/// </summary>
+		/// <param name="stream">The stream to replicate</param>
+		/// <param name="change">Commit to store the tree ref</param>
+		/// <param name="baseChange">The base change to update from</param>
+		/// <param name="baseContents">Initial contents of the tree at baseChange</param>
+		/// <param name="filter"></param>
+		/// <param name="revisionsOnly"></param>
+		/// <param name="cancellationToken">Cancellation token for the operation</param>
+		/// <returns>Root tree object</returns>
+		public async Task<ReplicationNode> WriteCommitTreeAsync(IStream stream, int change, int? baseChange, ReplicationNode baseContents, string? filter, bool revisionsOnly, CancellationToken cancellationToken)
+		{
+			// Create a client to replicate from this stream
+			ReplicationClient clientInfo = await FindOrAddReplicationClientAsync(stream);
+
+			// Connect to the server and flush the workspace
+			using IPerforceConnection perforce = await PerforceConnection.CreateAsync(clientInfo.Settings, _logger);
+			await FlushWorkspaceAsync(clientInfo, perforce, baseChange ?? 0);
+
+			// Apply all the updates
+			_logger.LogInformation("Syncing client {Client} from changelist {BaseChange} to {Change}", clientInfo.Client.Name, baseChange ?? 0, change);
+			clientInfo.Change = -1;
+
+			Utf8String clientRoot = new Utf8String(clientInfo.Client.Root);
+			string filterOrDefault = filter ?? "...";
+			string queryPath = $"//{clientInfo.Client.Name}/{filterOrDefault}";
+
+			RefName refName = GetRefName(stream.Id, change, filter, revisionsOnly);
+			RefName incRefName = new RefName($"{refName}_inc");
+
+			// Get the current sync state for this change
+			ReplicationNode? syncNode = await _treeStore.TryReadTreeAsync<ReplicationNode>(incRefName, cancellationToken);
+			bool deleteIncRef = syncNode != null;
+			if (syncNode == null)
+			{
+				DirectoryNode contents = await baseContents.Contents.ExpandAsync(cancellationToken);
+				syncNode = new ReplicationNode(contents);
+			}
+
+			// Get the contents of the current tree
+			DirectoryNode root = await syncNode.Contents.ExpandAsync(cancellationToken);
+
+			if (revisionsOnly)
+			{
+				int count = 0;
+				await foreach (PerforceResponse<SyncRecord> record in perforce.StreamCommandAsync<SyncRecord>("sync", new[] { "-k", $"{queryPath}@{change}" }, null, default, cancellationToken))
+				{
+					PerforceError? error = record.Error;
+					if (error != null && error.Generic == PerforceGenericCode.Empty)
+					{
+						continue;
+					}
+
+					SyncRecord syncRecord = record.Data;
+					if (!syncRecord.Path.StartsWith(clientRoot))
+					{
+						throw new ArgumentException($"Unable to make path {clientInfo.Client.Root} relative to client root {clientInfo.Client.Root}");
+					}
+
+					Utf8String path = syncRecord.Path.Substring(clientRoot.Length);
+					byte[] data = Encoding.UTF8.GetBytes($"{syncRecord.DepotFile}#{syncRecord.Revision}");
+
+					using ReadOnlyMemoryStream dataStream = new ReadOnlyMemoryStream(data);
+
+					FileEntry entry = await root.AddFileByPathAsync(path, FileEntryFlags.PerforceDepotPathAndRevision, cancellationToken);
+					await entry.AppendAsync(dataStream, Options.Chunking, cancellationToken);
+
+					if (++count > 1000)
+					{
+						await _treeStore.WriteTreeAsync(refName, syncNode, false, cancellationToken);
+						count = 0;
+					}
+				}
+			}
+			else
+			{
+				// Replay the files that have already been synced
+				foreach (Utf8String path in syncNode.Paths)
+				{
+					string flushPath = $"//{clientInfo.Client.Name}/{path}@{change}";
+					_logger.LogInformation("Flushing {FlushPath}", flushPath);
+					await perforce.SyncQuietAsync(SyncOptions.Force | SyncOptions.KeepWorkspaceFiles, -1, flushPath, cancellationToken);
+				}
+
+				// Do a sync preview to find everything that's left, and sort the remaining list of paths
+				List<(Utf8String Path, long Size)> files = new List<(Utf8String, long)>();
+				await foreach (PerforceResponse<SyncRecord> response in perforce.StreamCommandAsync<SyncRecord>("sync", new[] { "-n" }, new string[] { $"{queryPath}@{change}" }, null, cancellationToken))
+				{
+					PerforceError? error = response.Error;
+					if (error != null)
+					{
+						_logger.LogWarning("Perforce: {Message}", error.Data);
+						continue;
+					}
+
+					Utf8String file = response.Data.Path.Clone();
+					if (!file.StartsWith(clientRoot, Utf8StringComparer.Ordinal))
+					{
+						throw new ArgumentException($"Unable to make path {file} relative to client root {clientRoot}");
+					}
+
+					byte[] path = new byte[file.Length - clientRoot.Length];
+					for (int idx = clientRoot.Length; idx < file.Length; idx++)
+					{
+						path[idx - clientRoot.Length] = (file[idx] == '\\') ? (byte)'/' : file[idx];
+					}
+
+					files.Add((new Utf8String(path), response.Data.FileSize));
+				}
+				files.SortBy(x => x.Path, FolderFirstSorter.Instance);
+
+				// Output some stats for the sync
+				long totalSize = files.Sum(x => x.Size);
+				_logger.LogInformation("Total sync size: {Size:n1}mb", totalSize / (1024.0 * 1024.0));
+
+				// Sync incrementally
+				long syncedSize = 0;
+				while (files.Count > 0)
+				{
+					// Save the incremental state
+					if (syncedSize > 0)
+					{
+						await _treeStore.WriteTreeAsync(incRefName, syncNode, cancellationToken: cancellationToken);
+						deleteIncRef = true;
+					}
+
+					// Find the next path to sync
+					const long MaxBatchSize = 10L * 1024 * 1024 * 1024;
+					(int idx, Utf8String path, long size) = GetSyncBatch(files, MaxBatchSize, _logger);
+					syncedSize += size;
+
+					const long TrimInterval = 1024 * 1024 * 256;
+
+					long dataSize = 0;
+					long trimDataSize = TrimInterval;
+
+					string syncPath = $"//{clientInfo.Client.Name}/{path}@{change}";
+					double syncPct = (totalSize == 0) ? 100.0 : (syncedSize * 100.0) / totalSize;
+					_logger.LogInformation("Syncing {StreamId} to {Change} [{SyncPct:n1}%]: {Path} ({Size:n1}mb)", stream.Id, change, syncPct, path, size / (1024.0 * 1024.0));
+
+					Dictionary<int, FileEntry> handles = new Dictionary<int, FileEntry>();
+					await foreach (PerforceResponse response in perforce.StreamCommandAsync("sync", Array.Empty<string>(), new string[] { syncPath }, null, typeof(SyncRecord), true, default))
+					{
+						PerforceError? error = response.Error;
+						if (error != null)
+						{
+							_logger.LogWarning("Perforce: {Message}", error.Data);
+							continue;
+						}
+
+						PerforceIo? io = response.Io;
+						if (io != null)
+						{
+							if (io.Command == PerforceIoCommand.Open)
+							{
+								Utf8String file = GetClientRelativePath(io.Payload, clientInfo.Client.Root);
+								handles[io.File] = await root.AddFileByPathAsync(file, FileEntryFlags.None, cancellationToken);
+							}
+							else if (io.Command == PerforceIoCommand.Write)
+							{
+								FileEntry entry = handles[io.File];
+								await entry.AppendAsync(io.Payload, Options.Chunking, cancellationToken);
+								dataSize += io.Payload.Length;
+							}
+							else if (io.Command == PerforceIoCommand.Close)
+							{
+								handles.Remove(io.File);
+							}
+							else if (io.Command == PerforceIoCommand.Unlink)
+							{
+								Utf8String file = GetClientRelativePath(io.Payload, clientInfo.Client.Root);
+								await root.DeleteFileByPathAsync(file, cancellationToken);
+							}
+							else
+							{
+								_logger.LogWarning("Unhandled command code {Code}", io.Command);
+							}
+
+							if (dataSize > trimDataSize)
+							{
+								_logger.LogInformation("Trimming working set after receiving {Size:n1}mb...", dataSize / (1024 * 1024));
+								await _treeStore.WriteTreeAsync(refName, syncNode, false, cancellationToken);
+								GC.Collect();
+								trimDataSize = dataSize + TrimInterval;
+								_logger.LogInformation("Trimming complete. Next trim at {Size:n1}mb.", trimDataSize / (1024 * 1024));
+							}
+						}
+					}
+
+					// Update the root sync node
+					syncNode.Paths.Add(path);
+					files.RemoveRange(idx, files.Count - idx);
+				}
+			}
+
+			clientInfo.Change = change;
+
+			// Return the new root object
+			_logger.LogInformation("Writing ref {RefId} for {StreamId} change {Change}", refName, stream.Id, change);
+			await _treeStore.WriteTreeAsync(refName, syncNode, true, cancellationToken);
+
+			// Delete the incremental state
+			if (deleteIncRef)
+			{
+				await _treeStore.DeleteTreeAsync(incRefName, cancellationToken);
+			}
+			return syncNode;
+		}
+
+		static (int, Utf8String, long) GetSyncBatch(List<(Utf8String Path, long Size)> files, long maxSize, ILogger logger)
+		{
+			int idx = files.Count - 1;
+			Utf8String path = files[idx].Path;
+			long size = files[idx].Size;
+
+			Utf8String prefix = path;
+			while (prefix.Length > 0)
+			{
+				// Get the length of the parent directory
+				int endIdx = prefix.Length - 1;
+				while (endIdx > 0 && path[endIdx - 1] != (byte)'/')
+				{
+					endIdx--;
+				}
+
+				// Get the parent directory prefix, ending with a slash
+				prefix = prefix.Substring(0, endIdx);
+
+				// Include all files with this prefix
+				int nextIdx = idx;
+				long nextSize = size;
+				while (nextIdx > 0 && files[nextIdx - 1].Path.StartsWith(prefix))
+				{
+					nextSize += files[nextIdx - 1].Size;
+					nextIdx--;
+				}
+				if (nextSize > maxSize)
+				{
+					logger.LogInformation("Next filter {NextFilter} would exceed max size ({NextSize:n1}mb > {MaxSize:n1}mb); using {Filter} ({Size:n1}mb) instead.", prefix + "...", nextSize / (1024.0 * 1024.0), maxSize / (1024.0 * 1024.0), path, size / (1024.0 * 1024.0));
+					return (idx, path, size);
+				}
+				size = nextSize;
+				idx = nextIdx;
+
+				// Update the sync path to be the prefixed directory
+				path = prefix + "...";
+			}
+
+			return (idx, path, size);
+		}
+
+		async Task FlushWorkspaceAsync(ReplicationClient clientInfo, IPerforceConnection perforce, int change)
+		{
+			if (clientInfo.Change != change)
+			{
+				clientInfo.Change = -1;
+				if (change == 0)
+				{
+					_logger.LogInformation("Flushing have table for {Client}", clientInfo.Client.Name);
+					await perforce.SyncQuietAsync(SyncOptions.Force | SyncOptions.KeepWorkspaceFiles, -1, $"//{clientInfo.Client.Name}/...#0");
+				}
+				else
+				{
+					_logger.LogInformation("Flushing have table for {Client} to change {Change}", clientInfo.Client.Name, change);
+					await perforce.SyncQuietAsync(SyncOptions.Force | SyncOptions.KeepWorkspaceFiles, -1, $"//{clientInfo.Client.Name}/...@{change}");
+				}
+				clientInfo.Change = change;
+			}
+		}
+
+		static Utf8String GetClientRelativePath(ReadOnlyMemory<byte> data, Utf8String clientRoot)
+		{
+			int length = data.Span.IndexOf((byte)0);
+			if (length != -1)
+			{
+				data = data.Slice(0, length);
+			}
+
+			Utf8String path = new Utf8String(data);
+			if (!path.StartsWith(clientRoot, Utf8StringComparer.Ordinal))
+			{
+				throw new ArgumentException($"Unable to make path {path} relative to client root {clientRoot}");
+			}
+
+			return path.Substring(clientRoot.Length).Clone();
+		}
+
+		class ReplicationClient
+		{
+			public PerforceSettings Settings { get; }
+			public string ClusterName { get; }
+			public InfoRecord ServerInfo { get; }
+			public ClientRecord Client { get; }
+			public int Change { get; set; }
+
+			public ReplicationClient(PerforceSettings settings, string clusterName, InfoRecord serverInfo, ClientRecord client, int change)
+			{
+				Settings = settings;
+				ClusterName = clusterName;
+				ServerInfo = serverInfo;
+				Client = client;
+				Change = change;
+			}
+		}
+
+		readonly Dictionary<StreamId, ReplicationClient> _cachedPerforceClients = new Dictionary<StreamId, ReplicationClient>();
+
+		async Task<ReplicationClient?> FindReplicationClientAsync(IStream stream)
+		{
+			ReplicationClient? clientInfo;
+			if (_cachedPerforceClients.TryGetValue(stream.Id, out clientInfo))
+			{
+				if (!String.Equals(clientInfo.ClusterName, stream.ClusterName, StringComparison.Ordinal) && String.Equals(clientInfo.Client.Stream, stream.Name, StringComparison.Ordinal))
+				{
+					PerforceSettings serverSettings = new PerforceSettings(clientInfo.Settings);
+					serverSettings.ClientName = null;
+
+					using IPerforceConnection perforce = await PerforceConnection.CreateAsync(_logger);
+					await perforce.DeleteClientAsync(DeleteClientOptions.None, clientInfo.Client.Name);
+
+					_cachedPerforceClients.Remove(stream.Id);
+					clientInfo = null;
+				}
+			}
+			return clientInfo;
+		}
+
+		async Task<ReplicationClient> FindOrAddReplicationClientAsync(IStream stream)
+		{
+			ReplicationClient? clientInfo = await FindReplicationClientAsync(stream);
+			if (clientInfo == null)
+			{
+				using IPerforceConnection? perforce = await _perforceService.GetServiceUserConnection(stream.ClusterName);
+				if (perforce == null)
+				{
+					throw new PerforceException($"Unable to create connection to Perforce server");
+				}
+
+				InfoRecord serverInfo = await perforce.GetInfoAsync(InfoOptions.ShortOutput);
+
+				ClientRecord newClient = new ClientRecord($"Horde.Build_{serverInfo.ClientHost}_{stream.Id}", perforce.Settings.UserName, "/p4/");
+				newClient.Description = "Created to mirror Perforce content to Horde Storage";
+				newClient.Owner = perforce.Settings.UserName;
+				newClient.Host = serverInfo.ClientHost;
+				newClient.Stream = stream.Name;
+				await perforce.CreateClientAsync(newClient);
+				_logger.LogInformation("Created client {ClientName} for {StreamName}", newClient.Name, stream.Name);
+
+				PerforceSettings settings = new PerforceSettings(perforce.Settings);
+				settings.ClientName = newClient.Name;
+				settings.PreferNativeClient = true;
+
+				clientInfo = new ReplicationClient(settings, stream.ClusterName, serverInfo, newClient, -1);
+				_cachedPerforceClients.Add(stream.Id, clientInfo);
+			}
+			return clientInfo;
+		}
+	}
+}
diff --git a/Engine/Source/Programs/Horde/Horde.Build/Startup.cs b/Engine/Source/Programs/Horde/Horde.Build/Startup.cs
--- a/Engine/Source/Programs/Horde/Horde.Build/Startup.cs
+++ b/Engine/Source/Programs/Horde/Horde.Build/Startup.cs
@@ -307,8 +307,9 @@ namespace Horde.Build
 			configSection.Bind(settings);
 
 			settings.Validate();
-			
-			services.Configure<CommitServiceOptions>(configSection.GetSection("Replication"));
+
+			services.Configure<CommitServiceOptions>(configSection.GetSection("Commits"));
+			services.Configure<ReplicationServiceOptions>(configSection.GetSection("Replication"));
 
 			if (settings.GlobalThreadPoolMinSize != null)
 			{
@@ -440,6 +441,7 @@ namespace Horde.Build
 			services.AddSingleton<PerforceLoadBalancer>();
 			services.AddSingleton<PoolService>();
 			services.AddSingleton<ProjectService>();
+			services.AddSingleton<ReplicationService>();
 			services.AddSingleton<ScheduleService>();
 			services.AddSingleton<SlackNotificationSink>();
 			services.AddSingleton<IAvatarService, SlackNotificationSink>(sp => sp.GetRequiredService<SlackNotificationSink>());
@@ -461,7 +463,7 @@ namespace Horde.Build
 			// Storage providers
 			services.AddSingleton<IStorageBackend<PersistentLogStorage>>(sp => CreateStorageBackend(sp, settings.LogStorage).ForType<PersistentLogStorage>());
 			services.AddSingleton<IStorageBackend<ArtifactCollection>>(sp => CreateStorageBackend(sp, settings.ArtifactStorage).ForType<ArtifactCollection>());
-			services.AddSingleton<ITreeStore<CommitService>>(sp => CreateTreeStore(sp, settings.CommitStorage).ForType<CommitService>());
+			services.AddSingleton<ITreeStore<ReplicationService>>(sp => CreateTreeStore(sp, settings.CommitStorage).ForType<ReplicationService>());
 
 			services.AddHordeStorage(settings => configSection.GetSection("Storage").Bind(settings));
 
@@ -613,6 +615,7 @@ namespace Horde.Build
 				services.AddHostedService<IssueReportService>();
 				services.AddHostedService(provider => (LogFileService)provider.GetRequiredService<ILogFileService>());
 				services.AddHostedService(provider => (NotificationService)provider.GetRequiredService<INotificationService>());
+				services.AddHostedService(provider => provider.GetRequiredService<ReplicationService>());
 				if (!settings.DisableSchedules)
 				{
 					services.AddHostedService(provider => provider.GetRequiredService<ScheduleService>());
