Horde: Respect case sensitivity of server when grouping files for content syncing.

diff --git a/Engine/Source/Programs/Horde/Horde.Build/Perforce/ReplicationService.cs b/Engine/Source/Programs/Horde/Horde.Build/Perforce/ReplicationService.cs
--- a/Engine/Source/Programs/Horde/Horde.Build/Perforce/ReplicationService.cs
+++ b/Engine/Source/Programs/Horde/Horde.Build/Perforce/ReplicationService.cs
@@ -465,7 +465,7 @@ namespace Horde.Build.Perforce
 		/// <returns></returns>
 		static RefName GetRefName(StreamId streamId, int change, string? filter, bool revisionsOnly)
 		{
-			StringBuilder builder = new StringBuilder($"v2/{streamId}/{change}");
+			StringBuilder builder = new StringBuilder($"v3/{streamId}/{change}");
 			if (filter != null)
 			{
 				builder.Append($"_filter_{IoHash.Compute(Encoding.UTF8.GetBytes(filter))}");
@@ -523,20 +523,25 @@ namespace Horde.Build.Perforce
 		/// </summary>
 		class FolderFirstSorter : IComparer<Utf8String>
 		{
-			public static FolderFirstSorter Instance { get; } = new FolderFirstSorter();
+			readonly Utf8StringComparer _comparer;
+
+			public FolderFirstSorter(Utf8StringComparer comparer)
+			{
+				_comparer = comparer;
+			}
 
 			public int Compare(Utf8String x, Utf8String y)
 			{
 				Utf8String pathX = x.Substring(0, x.LastIndexOf('/') + 1);
 				Utf8String pathY = y.Substring(0, y.LastIndexOf('/') + 1);
 
-				int result = pathX.CompareTo(pathY);
+				int result = _comparer.Compare(pathX, pathY);
 				if (result != 0)
 				{
 					return result;
 				}
 
-				return x.CompareTo(y);
+				return _comparer.Compare(x, y);
 			}
 		}
 
@@ -650,7 +655,7 @@ namespace Horde.Build.Perforce
 
 					files.Add((new Utf8String(path), response.Data.FileSize));
 				}
-				files.SortBy(x => x.Path, FolderFirstSorter.Instance);
+				files.SortBy(x => x.Path, new FolderFirstSorter(clientInfo.ServerInfo.Utf8PathComparer));
 
 				// Output some stats for the sync
 				long totalSize = files.Sum(x => x.Size);
@@ -669,7 +674,7 @@ namespace Horde.Build.Perforce
 
 					// Find the next path to sync
 					const long MaxBatchSize = 10L * 1024 * 1024 * 1024;
-					(int idx, Utf8String path, long size) = GetSyncBatch(files, MaxBatchSize, _logger);
+					(int idx, Utf8String path, long size) = GetSyncBatch(files, MaxBatchSize, clientInfo.ServerInfo.Utf8PathComparer, _logger);
 					syncedSize += size;
 
 					const long TrimInterval = 1024 * 1024 * 256;
@@ -750,7 +755,7 @@ namespace Horde.Build.Perforce
 			return syncNode;
 		}
 
-		static (int, Utf8String, long) GetSyncBatch(List<(Utf8String Path, long Size)> files, long maxSize, ILogger logger)
+		static (int, Utf8String, long) GetSyncBatch(List<(Utf8String Path, long Size)> files, long maxSize, Utf8StringComparer comparer, ILogger logger)
 		{
 			int idx = files.Count - 1;
 			Utf8String path = files[idx].Path;
@@ -772,7 +777,7 @@ namespace Horde.Build.Perforce
 				// Include all files with this prefix
 				int nextIdx = idx;
 				long nextSize = size;
-				while (nextIdx > 0 && files[nextIdx - 1].Path.StartsWith(prefix))
+				while (nextIdx > 0 && files[nextIdx - 1].Path.StartsWith(prefix, comparer))
 				{
 					nextSize += files[nextIdx - 1].Size;
 					nextIdx--;
