Horde: Add stats for time spent processing sync output to log.

diff --git a/Engine/Source/Programs/Horde/Horde.Build/Perforce/ReplicationService.cs b/Engine/Source/Programs/Horde/Horde.Build/Perforce/ReplicationService.cs
--- a/Engine/Source/Programs/Horde/Horde.Build/Perforce/ReplicationService.cs
+++ b/Engine/Source/Programs/Horde/Horde.Build/Perforce/ReplicationService.cs
@@ -686,6 +686,9 @@ namespace Horde.Build.Perforce
 					double syncPct = (totalSize == 0) ? 100.0 : (syncedSize * 100.0) / totalSize;
 					_logger.LogInformation("Syncing {StreamId} to {Change} [{SyncPct:n1}%]: {Path} ({Size:n1}mb)", stream.Id, change, syncPct, path, size / (1024.0 * 1024.0));
 
+					Stopwatch SyncTimer = Stopwatch.StartNew();
+					Stopwatch ProcessTimer = new Stopwatch();
+
 					Dictionary<int, FileEntry> handles = new Dictionary<int, FileEntry>();
 					await foreach (PerforceResponse response in perforce.StreamCommandAsync("sync", Array.Empty<string>(), new string[] { syncPath }, null, typeof(SyncRecord), true, default))
 					{
@@ -696,6 +699,8 @@ namespace Horde.Build.Perforce
 							continue;
 						}
 
+						ProcessTimer.Start();
+
 						PerforceIo? io = response.Io;
 						if (io != null)
 						{
@@ -733,7 +738,10 @@ namespace Horde.Build.Perforce
 								_logger.LogInformation("Trimming complete. Next trim at {Size:n1}mb.", trimDataSize / (1024 * 1024));
 							}
 						}
+
+						ProcessTimer.Stop();
 					}
+					_logger.LogInformation("Completed batch in {TimeSeconds:n1}s ({ProcessTimeSeconds:n1}s processing)", SyncTimer.Elapsed.TotalSeconds, ProcessTimer.Elapsed.TotalSeconds);
 
 					// Update the root sync node
 					syncNode.Paths.Add(path);
