Horde: Additional stats timers for replication.

diff --git a/Engine/Source/Programs/Horde/Horde.Build/Perforce/ReplicationService.cs b/Engine/Source/Programs/Horde/Horde.Build/Perforce/ReplicationService.cs
--- a/Engine/Source/Programs/Horde/Horde.Build/Perforce/ReplicationService.cs
+++ b/Engine/Source/Programs/Horde/Horde.Build/Perforce/ReplicationService.cs
@@ -686,8 +686,10 @@ namespace Horde.Build.Perforce
 					double syncPct = (totalSize == 0) ? 100.0 : (syncedSize * 100.0) / totalSize;
 					_logger.LogInformation("Syncing {StreamId} to {Change} [{SyncPct:n1}%]: {Path} ({Size:n1}mb)", stream.Id, change, syncPct, path, size / (1024.0 * 1024.0));
 
-					Stopwatch SyncTimer = Stopwatch.StartNew();
-					Stopwatch ProcessTimer = new Stopwatch();
+					Stopwatch syncTimer = Stopwatch.StartNew();
+					Stopwatch processTimer = new Stopwatch();
+					Stopwatch trimTimer = new Stopwatch();
+					Stopwatch gcTimer = new Stopwatch();
 
 					Dictionary<int, FileEntry> handles = new Dictionary<int, FileEntry>();
 					await foreach (PerforceResponse response in perforce.StreamCommandAsync("sync", Array.Empty<string>(), new string[] { syncPath }, null, typeof(SyncRecord), true, default))
@@ -699,7 +701,7 @@ namespace Horde.Build.Perforce
 							continue;
 						}
 
-						ProcessTimer.Start();
+						processTimer.Start();
 
 						PerforceIo? io = response.Io;
 						if (io != null)
@@ -732,16 +734,20 @@ namespace Horde.Build.Perforce
 							if (dataSize > trimDataSize)
 							{
 								_logger.LogInformation("Trimming working set after receiving {Size:n1}mb...", dataSize / (1024 * 1024));
+								trimTimer.Start();
 								await _treeStore.WriteTreeAsync(refName, syncNode, false, cancellationToken);
+								trimTimer.Stop();
+								gcTimer.Start();
 								GC.Collect();
+								gcTimer.Stop();
 								trimDataSize = dataSize + TrimInterval;
 								_logger.LogInformation("Trimming complete. Next trim at {Size:n1}mb.", trimDataSize / (1024 * 1024));
 							}
 						}
 
-						ProcessTimer.Stop();
+						processTimer.Stop();
 					}
-					_logger.LogInformation("Completed batch in {TimeSeconds:n1}s ({ProcessTimeSeconds:n1}s processing)", SyncTimer.Elapsed.TotalSeconds, ProcessTimer.Elapsed.TotalSeconds);
+					_logger.LogInformation("Completed batch in {TimeSeconds:n1}s ({ProcessTimeSeconds:n1}s processing, {TrimTimeSeconds:n1}s trimming, {GcTimeSeconds:n1}s gc)", syncTimer.Elapsed.TotalSeconds, processTimer.Elapsed.TotalSeconds, trimTimer.Elapsed.TotalSeconds, gcTimer.Elapsed.TotalSeconds);
 
 					// Update the root sync node
 					syncNode.Paths.Add(path);
