Horde Storage - Make sure to re-trigger ref cleanups on each poll (but not start new ones if one is pending). As some namespaces take much longer then others to run this makes sure they will all get triggered without having to wait for the slowest one.

diff --git a/Engine/Source/Programs/Horde/HordeStorage/Horde.Storage/Implementation/GC/RefCleanupService.cs b/Engine/Source/Programs/Horde/HordeStorage/Horde.Storage/Implementation/GC/RefCleanupService.cs
--- a/Engine/Source/Programs/Horde/HordeStorage/Horde.Storage/Implementation/GC/RefCleanupService.cs
+++ b/Engine/Source/Programs/Horde/HordeStorage/Horde.Storage/Implementation/GC/RefCleanupService.cs
@@ -22,6 +22,7 @@ namespace Horde.Storage.Implementation
 
         public IRefsStore Refs { get; }
         public IRefCleanup RefCleanup { get; }
+        public Dictionary<NamespaceId, Task> RunningCleanupTasks { get; } = new Dictionary<NamespaceId, Task>();
     }
 
     public class RefCleanupService : PollingService<RefCleanupState>
@@ -60,14 +61,21 @@ namespace Horde.Storage.Implementation
                     return false;
                 }
                 List<NamespaceId>? namespaces = await _referencesStore.GetNamespaces().ToListAsync(cancellationToken);
-                List<Task> cleanupTasks = new List<Task>();
 
                 foreach(NamespaceId ns in namespaces)
                 {
-                    cleanupTasks.Add(DoCleanup(ns, state, cancellationToken));
+                    if (state.RunningCleanupTasks.TryGetValue(ns, out Task? runningTask))
+                    {
+                        if (!runningTask.IsCompleted)
+                        {
+                            continue;
+                        }
+                        await runningTask;
+                        state.RunningCleanupTasks.Remove(ns);
+                    }
+                    state.RunningCleanupTasks.Add(ns, DoCleanup(ns, state, cancellationToken));
                 }
 
-                await Task.WhenAll(cleanupTasks);
                 // Do not run the cleanup of legacy namespaces as this can take a very long time and we do not care about these anymore
                 /*
                  await foreach (NamespaceId ns in state.Refs.GetNamespaces().WithCancellation(cancellationToken))
