Horde: Update to use TypeScript 4.7.4 (from 4.1.3)
#skipci
#rnx

diff --git a/Engine/Source/Programs/Horde/HordeDashboard/package.json b/Engine/Source/Programs/Horde/HordeDashboard/package.json
--- a/Engine/Source/Programs/Horde/HordeDashboard/package.json
+++ b/Engine/Source/Programs/Horde/HordeDashboard/package.json
@@ -11,7 +11,7 @@
     "@fluentui/react-charting": "^5.2.0",
     "crypto-js": "^4.0.0",
     "formik": "^2.1.1",
-    "idb": "^6.0.0",
+    "idb": "^7.0.2",
     "markdown-to-jsx": "^6.10.3",
     "mobx": "^5.15.1",
     "mobx-react-lite": "^1.5.2",
@@ -56,6 +56,6 @@
     "@types/react-highlighter": "^0.3.4",
     "@types/react-router-dom": "^5.1.6",
     "rimraf": "^3.0.0",
-    "typescript": "^4.1.3"
+    "typescript": "^4.7.4"
   }
 }
