Horde Storage - Added cached content id store and fixed issue were we were resolving content ids twice when resolving references. This should saves us between 1 and 4ms per ref GET.

diff --git a/Engine/Source/Programs/Horde/HordeStorage/Horde.Storage/Horde.Storage.Startup.cs b/Engine/Source/Programs/Horde/HordeStorage/Horde.Storage/Horde.Storage.Startup.cs
--- a/Engine/Source/Programs/Horde/HordeStorage/Horde.Storage/Horde.Storage.Startup.cs
+++ b/Engine/Source/Programs/Horde/HordeStorage/Horde.Storage/Horde.Storage.Startup.cs
@@ -262,20 +262,28 @@ namespace Horde.Storage
         private object ContentIdStoreFactory(IServiceProvider provider)
         {
             HordeStorageSettings settings = provider.GetService<IOptionsMonitor<HordeStorageSettings>>()!.CurrentValue!;
-            switch (settings.ContentIdStoreImplementation)
+            IContentIdStore store = settings.ContentIdStoreImplementation switch
             {
-                case HordeStorageSettings.ContentIdStoreImplementations.Memory:
-                    return ActivatorUtilities.CreateInstance<MemoryContentIdStore>(provider);
-                case HordeStorageSettings.ContentIdStoreImplementations.Scylla:
-                    return ActivatorUtilities.CreateInstance<ScyllaContentIdStore>(provider);
-                case HordeStorageSettings.ContentIdStoreImplementations.Mongo:
-                    return ActivatorUtilities.CreateInstance<MongoContentIdStore>(provider);
-                case HordeStorageSettings.ContentIdStoreImplementations.Cache:
-                    return ActivatorUtilities.CreateInstance<CacheContentIdStore>(provider);
-                default:
-                    throw new NotImplementedException($"Unknown content id store implementation: {settings.ContentIdStoreImplementation}");
+                HordeStorageSettings.ContentIdStoreImplementations.Memory => ActivatorUtilities
+                    .CreateInstance<MemoryContentIdStore>(provider),
+                HordeStorageSettings.ContentIdStoreImplementations.Scylla => ActivatorUtilities
+                    .CreateInstance<ScyllaContentIdStore>(provider),
+                HordeStorageSettings.ContentIdStoreImplementations.Mongo => ActivatorUtilities
+                    .CreateInstance<MongoContentIdStore>(provider),
+                HordeStorageSettings.ContentIdStoreImplementations.Cache => ActivatorUtilities
+                    .CreateInstance<CacheContentIdStore>(provider),
+                _ => throw new NotImplementedException(
+                    $"Unknown content id store implementation: {settings.ContentIdStoreImplementation}")
+            };
+
+            MemoryCacheContentIdSettings memoryCacheSettings = provider.GetService<IOptionsMonitor<MemoryCacheContentIdSettings>>()!.CurrentValue;
+
+            if (memoryCacheSettings.Enabled)
+            {
+                store = ActivatorUtilities.CreateInstance<MemoryCachedContentIdStore>(provider, store);
             }
-        }
+			return store;
+		}
 
         private IScyllaSessionManager ScyllaFactory(IServiceProvider provider)
         {
diff --git a/Engine/Source/Programs/Horde/HordeStorage/Horde.Storage/Implementation/IReferenceResolver.cs b/Engine/Source/Programs/Horde/HordeStorage/Horde.Storage/Implementation/IReferenceResolver.cs
--- a/Engine/Source/Programs/Horde/HordeStorage/Horde.Storage/Implementation/IReferenceResolver.cs
+++ b/Engine/Source/Programs/Horde/HordeStorage/Horde.Storage/Implementation/IReferenceResolver.cs
@@ -51,10 +51,12 @@ namespace Horde.Storage.Implementation
     public class ContentIdAttachment : Attachment
     {
         public ContentId Identifier { get; }
+        public BlobIdentifier[] ReferencedBlobs { get; }
 
-        public ContentIdAttachment(ContentId contentId)
+        public ContentIdAttachment(ContentId contentId, BlobIdentifier[] referencedBlobs)
         {
             Identifier = contentId;
+            ReferencedBlobs = referencedBlobs;
         }
 
         public override IoHash AsIoHash()
@@ -83,7 +85,6 @@ namespace Horde.Storage.Implementation
 
         public async IAsyncEnumerable<Attachment> GetAttachments(NamespaceId ns, CbObject cb)
         {
-            // TODO: This is cacheable and we should store the result somewhere
             Queue<CbObject> objectsToVisit = new Queue<CbObject>();
             objectsToVisit.Enqueue(cb);
             List<BlobIdentifier> unresolvedBlobReferences = new List<BlobIdentifier>();
@@ -105,27 +106,27 @@ namespace Horde.Storage.Implementation
 
                         if (field.IsBinaryAttachment())
                         {
-                            bool isContentId = false;
+                            BlobIdentifier[]? resolvedBlobs = null;
+                            bool wasContentId = false;
                             try
                             {
-                                // TODO: Having to do a proper resolve here is not good as its quite expensive
-                                // it would be much better if the attachment type indicated if it was a content id or not
-                                Task<BlobIdentifier[]?> resolveContentId = _contentIdStore.Resolve(ns, contentId, mustBeContentId: true);
+                                Task<(ContentId, BlobIdentifier[]?)> resolveContentId = ResolveContentId(ns, contentId);
                                 resolveContentId.Wait();
-                                isContentId = resolveContentId.Result != null;
+                                resolvedBlobs = resolveContentId.Result.Item2;
+                                wasContentId = !(resolvedBlobs is { Length: 1 } && resolvedBlobs[0].Equals(blobIdentifier));
                             }
                             catch (InvalidContentIdException)
                             {
-                                isContentId = false;
+                                resolvedBlobs = null;
                             }
                             catch (BlobNotFoundException e)
                             {
                                 unresolvedBlobReferences.Add(e.Blob);
                             }
                             
-                            if (isContentId)
+                            if (wasContentId && resolvedBlobs != null)
                             {
-                                attachments.Add(new ContentIdAttachment(contentId));
+                                attachments.Add(new ContentIdAttachment(contentId, resolvedBlobs));
                             }
                             else
                             {
@@ -189,7 +190,6 @@ namespace Horde.Storage.Implementation
 
         public async IAsyncEnumerable<BlobIdentifier> GetReferencedBlobs(NamespaceId ns, CbObject cb)
         {
-            List<Task<(ContentId, BlobIdentifier[]?)>> pendingContentIdResolves = new();
             List<Task<(BlobIdentifier, bool)>> pendingBlobExistsChecks = new();
             List<ContentId> unresolvedContentIdReferences = new List<ContentId>();
             List<BlobIdentifier> unresolvedBlobReferences = new List<BlobIdentifier>();
@@ -204,7 +204,19 @@ namespace Horde.Storage.Implementation
                 else if (attachment is ContentIdAttachment contentIdAttachment)
                 {
                     // If we find a content id we resolve that into the actual blobs it references
-                    pendingContentIdResolves.Add(ResolveContentId(ns, contentIdAttachment.Identifier));
+                    
+                    foreach (BlobIdentifier b in contentIdAttachment.ReferencedBlobs)
+                    {
+                        (BlobIdentifier _, bool exists) = await CheckBlobExists(ns, b);
+                        if (!exists)
+                        {
+                            unresolvedContentIdReferences.Add(contentIdAttachment.Identifier);
+                            continue;
+                        }
+
+                        yield return b;
+                    }
+
                 }
                 else if (attachment is ObjectAttachment objectAttachment)
                 {
@@ -217,23 +229,6 @@ namespace Horde.Storage.Implementation
                 }
             }
 
-            // Return the results of all the content id resolves
-            foreach (Task<(ContentId, BlobIdentifier[]?)> pendingContentIdResolveTask in pendingContentIdResolves)
-            {
-                (ContentId contentId, BlobIdentifier[]? resolvedBlobs) =  await pendingContentIdResolveTask;
-                if (resolvedBlobs != null)
-                {
-                    foreach (BlobIdentifier b in resolvedBlobs)
-                    {
-                        pendingBlobExistsChecks.Add(CheckBlobExists(ns, b));
-                    }
-                }
-                else
-                {
-                    unresolvedContentIdReferences.Add(contentId);
-                }
-            }
-
             // return any verified blobs
             foreach (Task<(BlobIdentifier, bool)> pendingBlobExistsTask in pendingBlobExistsChecks)
             {
diff --git a/Engine/Source/Programs/Horde/HordeStorage/Horde.Storage/Settings.cs b/Engine/Source/Programs/Horde/HordeStorage/Horde.Storage/Settings.cs
--- a/Engine/Source/Programs/Horde/HordeStorage/Horde.Storage/Settings.cs
+++ b/Engine/Source/Programs/Horde/HordeStorage/Horde.Storage/Settings.cs
@@ -197,7 +197,15 @@ namespace Horde.Storage
         [Required] public string ConnectionString { get; set; } = "";
     }
 
-    public class MemoryCacheBlobSettings : MemoryCacheOptions
+    public class MemoryCacheContentIdSettings : MemoryCacheOptions
+    {
+        public bool Enabled { get; set; } = true;
+
+        public bool EnableSlidingExpiry { get; set; } = true;
+        public int SlidingExpirationMinutes { get; set; } = 120;
+    }
+
+	public class MemoryCacheBlobSettings : MemoryCacheOptions
     {
         public bool Enabled { get; set; } = true;
 
diff --git a/Engine/Source/Programs/Horde/HordeStorage/Horde.Storage/Tests/Functional/Objects/ReferencesTests.cs b/Engine/Source/Programs/Horde/HordeStorage/Horde.Storage/Tests/Functional/Objects/ReferencesTests.cs
--- a/Engine/Source/Programs/Horde/HordeStorage/Horde.Storage/Tests/Functional/Objects/ReferencesTests.cs
+++ b/Engine/Source/Programs/Horde/HordeStorage/Horde.Storage/Tests/Functional/Objects/ReferencesTests.cs
@@ -54,6 +54,11 @@ namespace Horde.Storage.FunctionalTests.References
 
             IContentIdStore contentIdStore = provider.GetService<IContentIdStore>()!;
             //verify we are using the expected store
+            if (contentIdStore is MemoryCachedContentIdStore memoryStore)
+            {
+                memoryStore.Clear();
+                contentIdStore = memoryStore.GetUnderlyingContentIdStore();
+            }
             Assert.IsTrue(contentIdStore.GetType() == typeof(ScyllaContentIdStore));
 
             IReplicationLog replicationLog = provider.GetService<IReplicationLog>()!;
@@ -107,6 +112,11 @@ namespace Horde.Storage.FunctionalTests.References
             Assert.IsTrue(referencesStore.GetType() == typeof(MongoReferencesStore));
 
             IContentIdStore contentIdStore = provider.GetService<IContentIdStore>()!;
+            if (contentIdStore is MemoryCachedContentIdStore memoryStore)
+            {
+                memoryStore.Clear();
+                contentIdStore = memoryStore.GetUnderlyingContentIdStore();
+            }
             //verify we are using the expected store
             Assert.IsTrue(contentIdStore.GetType() == typeof(MongoContentIdStore));
 
@@ -143,6 +153,11 @@ namespace Horde.Storage.FunctionalTests.References
             Assert.IsTrue(referencesStore.GetType() == typeof(MemoryReferencesStore));
 
             IContentIdStore contentIdStore = provider.GetService<IContentIdStore>()!;
+            if (contentIdStore is MemoryCachedContentIdStore memoryStore)
+            {
+                memoryStore.Clear();
+                contentIdStore = memoryStore.GetUnderlyingContentIdStore();
+            }
             //verify we are using the expected store
             Assert.IsTrue(contentIdStore.GetType() == typeof(MemoryContentIdStore));
 
@@ -1221,7 +1236,7 @@ namespace Horde.Storage.FunctionalTests.References
             byte[] compressedBuffer = bufferUtils.CompressContent(OoodleCompressorMethod.Mermaid, OoodleCompressionLevel.VeryFast, blobData);
             BlobIdentifier compressedHash = BlobIdentifier.FromBlob(compressedBuffer);
 
-			CbObject cbObject = CbObject.Build(writer => writer.WriteBinaryAttachment("Attachment", uncompressedHash.AsIoHash()));
+            CbObject cbObject = CbObject.Build(writer => writer.WriteBinaryAttachment("Attachment", uncompressedHash.AsIoHash()));
             byte[] cbObjectData = cbObject.GetView().ToArray();
             BlobIdentifier cbObjectHash = BlobIdentifier.FromBlob(cbObjectData);
 
@@ -1234,8 +1249,8 @@ namespace Horde.Storage.FunctionalTests.References
                 result.EnsureSuccessStatusCode();
             }
 
-			{
-				// upload ref
+            {
+                // upload ref
                 using HttpContent requestContent = new ByteArrayContent(cbObjectData);
                 requestContent.Headers.ContentType = new MediaTypeHeaderValue(CustomMediaTypeNames.UnrealCompactBinary);
                 requestContent.Headers.Add(CommonHeaders.HashHeaderName, cbObjectHash.ToString());
@@ -1256,7 +1271,7 @@ namespace Horde.Storage.FunctionalTests.References
                 }
             }
 
-			{
+            {
                 // verify we can fetch the blob properly
                 HttpResponseMessage getResponse = await Client.GetAsync(new Uri($"api/v1/refs/{TestNamespace}/bucket/{key}.uecb", UriKind.Relative));
                 getResponse.EnsureSuccessStatusCode();
@@ -1285,7 +1300,7 @@ namespace Horde.Storage.FunctionalTests.References
                 Assert.AreEqual(0, response.Missing.Count);
             }
 
-			{
+            {
                 // delete the blob referenced by the compressed buffer
                 await Service.DeleteObject(TestNamespace, compressedHash);
             }
@@ -1310,218 +1325,217 @@ namespace Horde.Storage.FunctionalTests.References
 
                 ExistCheckMultipleRefsResponse response = await existsResponse.Content.ReadAsAsync<ExistCheckMultipleRefsResponse>();
                 Assert.AreEqual(1, response.Missing.Count);
-				Assert.AreEqual(key, response.Missing[0].Key);
-				Assert.AreEqual("bucket", response.Missing[0].Bucket.ToString());
+                Assert.AreEqual(key, response.Missing[0].Key);
+                Assert.AreEqual("bucket", response.Missing[0].Bucket.ToString());
             }
         }
 
-
-		[TestMethod]
-		public async Task GetMissingCompressedBufferAttachment()
-		{
+        [TestMethod]
+        public async Task GetMissingCompressedBufferAttachment()
+        {
             CbObject cbObjectAttachment = CbObject.Build(writer => writer.WriteString("ValueField", "This field has a value"));
             byte[] cbAttachmentData = cbObjectAttachment.GetView().ToArray();
             BlobIdentifier cbAttachmentHash = BlobIdentifier.FromBlob(cbAttachmentData);
-			IoHashKey key = IoHashKey.FromName("compressedAttachedObject");
-
-			CbObject cbObject = CbObject.Build(writer => writer.WriteObjectAttachment("Attachment", cbAttachmentHash.AsIoHash()));
-			byte[] cbObjectData = cbObject.GetView().ToArray();
-			BlobIdentifier cbObjectHash = BlobIdentifier.FromBlob(cbObjectData);
-
-			{
-				// upload compressed blob
-				using HttpContent requestContent = new ByteArrayContent(cbAttachmentData);
-				requestContent.Headers.ContentType = new MediaTypeHeaderValue(CustomMediaTypeNames.UnrealCompactBinary);
-
-				HttpResponseMessage result = await Client!.PutAsync(new Uri($"api/v1/objects/{TestNamespace}/{cbAttachmentHash}", UriKind.Relative), requestContent);
-				result.EnsureSuccessStatusCode();
-			}
-
-			{
-				// upload ref
-				using HttpContent requestContent = new ByteArrayContent(cbObjectData);
-				requestContent.Headers.ContentType = new MediaTypeHeaderValue(CustomMediaTypeNames.UnrealCompactBinary);
-				requestContent.Headers.Add(CommonHeaders.HashHeaderName, cbObjectHash.ToString());
-
-				HttpResponseMessage result = await Client!.PutAsync(new Uri($"api/v1/refs/{TestNamespace}/bucket/{key}.uecb", UriKind.Relative), requestContent);
-				result.EnsureSuccessStatusCode();
-
-				{
-					Assert.AreEqual(CustomMediaTypeNames.UnrealCompactBinary, result!.Content.Headers.ContentType!.MediaType);
-					// check that no blobs are missing
-					await using MemoryStream ms = new MemoryStream();
-					await result.Content.CopyToAsync(ms);
-					byte[] roundTrippedBuffer = ms.ToArray();
-					CbObject cb = new CbObject(roundTrippedBuffer);
-					CbField needsField = cb["needs"];
-					List<BlobIdentifier> missingBlobs = needsField.AsArray().Select(field => BlobIdentifier.FromIoHash(field.AsHash())).ToList();
-					Assert.AreEqual(0, missingBlobs.Count);
-				}
-			}
-
-			{
-				// verify we can fetch the blob properly
-				HttpResponseMessage getResponse = await Client.GetAsync(new Uri($"api/v1/refs/{TestNamespace}/bucket/{key}.uecb", UriKind.Relative));
-				getResponse.EnsureSuccessStatusCode();
-				await using MemoryStream ms = new MemoryStream();
-				await getResponse.Content.CopyToAsync(ms);
-
-				byte[] roundTrippedBuffer = ms.ToArray();
-
-				CollectionAssert.AreEqual(cbObjectData, roundTrippedBuffer);
-				Assert.AreEqual(cbObjectHash, BlobIdentifier.FromBlob(roundTrippedBuffer));
-			}
-
-			{
-				// verify that head checks find the object
-				using HttpRequestMessage headRequest = new HttpRequestMessage(HttpMethod.Head, new Uri($"api/v1/refs/{TestNamespace}/bucket/{key}", UriKind.Relative));
-				HttpResponseMessage getResponse = await Client.SendAsync(headRequest);
-				getResponse.EnsureSuccessStatusCode();
-			}
-
-			{
-				// verify that the exists check doesnt find any issue
-				HttpResponseMessage existsResponse = await Client.GetAsync(new Uri($"api/v1/refs/{TestNamespace}/exists?names=bucket.{key}", UriKind.Relative));
-				existsResponse.EnsureSuccessStatusCode();
-
-				ExistCheckMultipleRefsResponse response = await existsResponse.Content.ReadAsAsync<ExistCheckMultipleRefsResponse>();
-				Assert.AreEqual(0, response.Missing.Count);
-			}
-
-			{
-				// delete the blob referenced by the compressed buffer
-				await Service.DeleteObject(TestNamespace, cbAttachmentHash);
-			}
-
-			{
-				// the compact binary object still exists, so this returns a success (trying to resolve the attachment will fail)
-				HttpResponseMessage getResponse = await Client.GetAsync(new Uri($"api/v1/refs/{TestNamespace}/bucket/{key}.uecb", UriKind.Relative));
-				getResponse.EnsureSuccessStatusCode();
-			}
-
-			{
-				// verify that head checks fail
-				using HttpRequestMessage headRequest = new HttpRequestMessage(HttpMethod.Head, new Uri($"api/v1/refs/{TestNamespace}/bucket/{key}", UriKind.Relative));
-				HttpResponseMessage getResponse = await Client.SendAsync(headRequest);
-				Assert.AreEqual(HttpStatusCode.NotFound, getResponse.StatusCode);
-			}
-
-			{
-				// verify that the exists check correctly returns a missing blob
-				HttpResponseMessage existsResponse = await Client.GetAsync(new Uri($"api/v1/refs/{TestNamespace}/exists?names=bucket.{key}", UriKind.Relative));
-				existsResponse.EnsureSuccessStatusCode();
-
-				ExistCheckMultipleRefsResponse response = await existsResponse.Content.ReadAsAsync<ExistCheckMultipleRefsResponse>();
-				Assert.AreEqual(1, response.Missing.Count);
-				Assert.AreEqual(key, response.Missing[0].Key);
-				Assert.AreEqual("bucket", response.Missing[0].Bucket.ToString());
-			}
-		}
-
-		[TestMethod]
-		public async Task GetMissingBlobRecord()
-		{
-			string blobContents = "This is a blob";
-			byte[] blobData = Encoding.ASCII.GetBytes(blobContents);
-			BlobIdentifier blobHash = BlobIdentifier.FromBlob(blobData);
-			IoHashKey key = IoHashKey.FromName("newReferenceObject");
-
-			using HttpContent requestContent = new ByteArrayContent(blobData);
-			requestContent.Headers.ContentType = new MediaTypeHeaderValue(MediaTypeNames.Application.Octet);
-			requestContent.Headers.Add(CommonHeaders.HashHeaderName, blobHash.ToString());
-
-			HttpResponseMessage result = await Client!.PutAsync(new Uri($"api/v1/refs/{TestNamespace}/bucket/{key}.uecb", UriKind.Relative), requestContent);
-			result.EnsureSuccessStatusCode();
-
-			{
-				Assert.AreEqual(CustomMediaTypeNames.UnrealCompactBinary, result!.Content.Headers.ContentType!.MediaType);
-				// check that no blobs are missing
-				await using MemoryStream ms = new MemoryStream();
-				await result.Content.CopyToAsync(ms);
-				byte[] roundTrippedBuffer = ms.ToArray();
-				CbObject cb = new CbObject(roundTrippedBuffer);
-				CbField needsField = cb["needs"];
-				List<BlobIdentifier> missingBlobs = needsField.AsArray().Select(field => BlobIdentifier.FromIoHash(field.AsHash())).ToList();
-				Assert.AreEqual(0, missingBlobs.Count);
-			}
-
-			{
-				// verify we can fetch the blob properly
-				HttpResponseMessage getResponse = await Client.GetAsync(new Uri($"api/v1/refs/{TestNamespace}/bucket/{key}.raw", UriKind.Relative));
-				getResponse.EnsureSuccessStatusCode();
-				await using MemoryStream ms = new MemoryStream();
-				await getResponse.Content.CopyToAsync(ms);
-
-				byte[] roundTrippedBuffer = ms.ToArray();
-
-				CollectionAssert.AreEqual(blobData, roundTrippedBuffer);
-				Assert.AreEqual(blobHash, BlobIdentifier.FromBlob(roundTrippedBuffer));
-			}
-
-			{
-				// verify that head checks find the object
-				using HttpRequestMessage headRequest = new HttpRequestMessage(HttpMethod.Head, new Uri($"api/v1/refs/{TestNamespace}/bucket/{key}", UriKind.Relative));
-				HttpResponseMessage getResponse = await Client.SendAsync(headRequest);
-				getResponse.EnsureSuccessStatusCode();
-			}
-
-			{
-				// verify that the exists check doesnt find any issue
-				HttpResponseMessage existsResponse = await Client.GetAsync(new Uri($"api/v1/refs/{TestNamespace}/exists?names=bucket.{key}", UriKind.Relative));
-				existsResponse.EnsureSuccessStatusCode();
-
-				ExistCheckMultipleRefsResponse response = await existsResponse.Content.ReadAsAsync<ExistCheckMultipleRefsResponse>();
-				Assert.AreEqual(0, response.Missing.Count);
-			}
-
-			{
-				// delete the blob 
-				await Service.DeleteObject(TestNamespace, blobHash);
-			}
-
-			{
-				// the compact binary object still exists, so this returns a success (trying to resolve the attachment will fail)
-				HttpResponseMessage getResponse = await Client.GetAsync(new Uri($"api/v1/refs/{TestNamespace}/bucket/{key}.uecb", UriKind.Relative));
-				getResponse.EnsureSuccessStatusCode();
-			}
-
-			{
-				// the compact binary object still exists, so this returns a success (trying to resolve the attachment will fail)
-				HttpResponseMessage getResponse = await Client.GetAsync(new Uri($"api/v1/refs/{TestNamespace}/bucket/{key}.json", UriKind.Relative));
-				getResponse.EnsureSuccessStatusCode();
-			}
-
-			{
-				// we should now see a 404 as the blob is missing
-				HttpResponseMessage getResponse = await Client.GetAsync(new Uri($"api/v1/refs/{TestNamespace}/bucket/{key}.raw", UriKind.Relative));
-				Assert.AreEqual(HttpStatusCode.NotFound, getResponse.StatusCode);
-				Assert.AreEqual("application/problem+json", getResponse.Content.Headers.ContentType!.MediaType);
-				string s = await getResponse.Content.ReadAsStringAsync();
-				ProblemDetails? problem = JsonConvert.DeserializeObject<ProblemDetails>(s);
-				Assert.IsNotNull(problem);
-				Assert.AreEqual($"Object {blobHash} in {TestNamespace} not found", problem.Title);
-			}
-
-			{
-				// verify that head checks fail
-				using HttpRequestMessage headRequest = new HttpRequestMessage(HttpMethod.Head, new Uri($"api/v1/refs/{TestNamespace}/bucket/{key}", UriKind.Relative));
-				HttpResponseMessage getResponse = await Client.SendAsync(headRequest);
-				Assert.AreEqual(HttpStatusCode.NotFound, getResponse.StatusCode);
-			}
-
-			{
-				// verify that the exists check correctly returns a missing blob
-				HttpResponseMessage existsResponse = await Client.GetAsync(new Uri($"api/v1/refs/{TestNamespace}/exists?names=bucket.{key}", UriKind.Relative));
-				existsResponse.EnsureSuccessStatusCode();
-
-				ExistCheckMultipleRefsResponse response = await existsResponse.Content.ReadAsAsync<ExistCheckMultipleRefsResponse>();
-				Assert.AreEqual(1, response.Missing.Count);
-				Assert.AreEqual(key, response.Missing[0].Key);
-				Assert.AreEqual("bucket", response.Missing[0].Bucket.ToString());
-			}
-		}
-
-		[TestMethod]
+            IoHashKey key = IoHashKey.FromName("compressedAttachedObject");
+
+            CbObject cbObject = CbObject.Build(writer => writer.WriteObjectAttachment("Attachment", cbAttachmentHash.AsIoHash()));
+            byte[] cbObjectData = cbObject.GetView().ToArray();
+            BlobIdentifier cbObjectHash = BlobIdentifier.FromBlob(cbObjectData);
+
+            {
+                // upload compressed blob
+                using HttpContent requestContent = new ByteArrayContent(cbAttachmentData);
+                requestContent.Headers.ContentType = new MediaTypeHeaderValue(CustomMediaTypeNames.UnrealCompactBinary);
+
+                HttpResponseMessage result = await Client!.PutAsync(new Uri($"api/v1/objects/{TestNamespace}/{cbAttachmentHash}", UriKind.Relative), requestContent);
+                result.EnsureSuccessStatusCode();
+            }
+
+            {
+                // upload ref
+                using HttpContent requestContent = new ByteArrayContent(cbObjectData);
+                requestContent.Headers.ContentType = new MediaTypeHeaderValue(CustomMediaTypeNames.UnrealCompactBinary);
+                requestContent.Headers.Add(CommonHeaders.HashHeaderName, cbObjectHash.ToString());
+
+                HttpResponseMessage result = await Client!.PutAsync(new Uri($"api/v1/refs/{TestNamespace}/bucket/{key}.uecb", UriKind.Relative), requestContent);
+                result.EnsureSuccessStatusCode();
+
+                {
+                    Assert.AreEqual(CustomMediaTypeNames.UnrealCompactBinary, result!.Content.Headers.ContentType!.MediaType);
+                    // check that no blobs are missing
+                    await using MemoryStream ms = new MemoryStream();
+                    await result.Content.CopyToAsync(ms);
+                    byte[] roundTrippedBuffer = ms.ToArray();
+                    CbObject cb = new CbObject(roundTrippedBuffer);
+                    CbField needsField = cb["needs"];
+                    List<BlobIdentifier> missingBlobs = needsField.AsArray().Select(field => BlobIdentifier.FromIoHash(field.AsHash())).ToList();
+                    Assert.AreEqual(0, missingBlobs.Count);
+                }
+            }
+
+            {
+                // verify we can fetch the blob properly
+                HttpResponseMessage getResponse = await Client.GetAsync(new Uri($"api/v1/refs/{TestNamespace}/bucket/{key}.uecb", UriKind.Relative));
+                getResponse.EnsureSuccessStatusCode();
+                await using MemoryStream ms = new MemoryStream();
+                await getResponse.Content.CopyToAsync(ms);
+
+                byte[] roundTrippedBuffer = ms.ToArray();
+
+                CollectionAssert.AreEqual(cbObjectData, roundTrippedBuffer);
+                Assert.AreEqual(cbObjectHash, BlobIdentifier.FromBlob(roundTrippedBuffer));
+            }
+
+            {
+                // verify that head checks find the object
+                using HttpRequestMessage headRequest = new HttpRequestMessage(HttpMethod.Head, new Uri($"api/v1/refs/{TestNamespace}/bucket/{key}", UriKind.Relative));
+                HttpResponseMessage getResponse = await Client.SendAsync(headRequest);
+                getResponse.EnsureSuccessStatusCode();
+            }
+
+            {
+                // verify that the exists check doesnt find any issue
+                HttpResponseMessage existsResponse = await Client.GetAsync(new Uri($"api/v1/refs/{TestNamespace}/exists?names=bucket.{key}", UriKind.Relative));
+                existsResponse.EnsureSuccessStatusCode();
+
+                ExistCheckMultipleRefsResponse response = await existsResponse.Content.ReadAsAsync<ExistCheckMultipleRefsResponse>();
+                Assert.AreEqual(0, response.Missing.Count);
+            }
+
+            {
+                // delete the blob referenced by the compressed buffer
+                await Service.DeleteObject(TestNamespace, cbAttachmentHash);
+            }
+
+            {
+                // the compact binary object still exists, so this returns a success (trying to resolve the attachment will fail)
+                HttpResponseMessage getResponse = await Client.GetAsync(new Uri($"api/v1/refs/{TestNamespace}/bucket/{key}.uecb", UriKind.Relative));
+                getResponse.EnsureSuccessStatusCode();
+            }
+
+            {
+                // verify that head checks fail
+                using HttpRequestMessage headRequest = new HttpRequestMessage(HttpMethod.Head, new Uri($"api/v1/refs/{TestNamespace}/bucket/{key}", UriKind.Relative));
+                HttpResponseMessage getResponse = await Client.SendAsync(headRequest);
+                Assert.AreEqual(HttpStatusCode.NotFound, getResponse.StatusCode);
+            }
+
+            {
+                // verify that the exists check correctly returns a missing blob
+                HttpResponseMessage existsResponse = await Client.GetAsync(new Uri($"api/v1/refs/{TestNamespace}/exists?names=bucket.{key}", UriKind.Relative));
+                existsResponse.EnsureSuccessStatusCode();
+
+                ExistCheckMultipleRefsResponse response = await existsResponse.Content.ReadAsAsync<ExistCheckMultipleRefsResponse>();
+                Assert.AreEqual(1, response.Missing.Count);
+                Assert.AreEqual(key, response.Missing[0].Key);
+                Assert.AreEqual("bucket", response.Missing[0].Bucket.ToString());
+            }
+        }
+
+        [TestMethod]
+        public async Task GetMissingBlobRecord()
+        {
+            string blobContents = "This is a blob";
+            byte[] blobData = Encoding.ASCII.GetBytes(blobContents);
+            BlobIdentifier blobHash = BlobIdentifier.FromBlob(blobData);
+            IoHashKey key = IoHashKey.FromName("newReferenceObject");
+
+            using HttpContent requestContent = new ByteArrayContent(blobData);
+            requestContent.Headers.ContentType = new MediaTypeHeaderValue(MediaTypeNames.Application.Octet);
+            requestContent.Headers.Add(CommonHeaders.HashHeaderName, blobHash.ToString());
+
+            HttpResponseMessage result = await Client!.PutAsync(new Uri($"api/v1/refs/{TestNamespace}/bucket/{key}.uecb", UriKind.Relative), requestContent);
+            result.EnsureSuccessStatusCode();
+
+            {
+                Assert.AreEqual(CustomMediaTypeNames.UnrealCompactBinary, result!.Content.Headers.ContentType!.MediaType);
+                // check that no blobs are missing
+                await using MemoryStream ms = new MemoryStream();
+                await result.Content.CopyToAsync(ms);
+                byte[] roundTrippedBuffer = ms.ToArray();
+                CbObject cb = new CbObject(roundTrippedBuffer);
+                CbField needsField = cb["needs"];
+                List<BlobIdentifier> missingBlobs = needsField.AsArray().Select(field => BlobIdentifier.FromIoHash(field.AsHash())).ToList();
+                Assert.AreEqual(0, missingBlobs.Count);
+            }
+
+            {
+                // verify we can fetch the blob properly
+                HttpResponseMessage getResponse = await Client.GetAsync(new Uri($"api/v1/refs/{TestNamespace}/bucket/{key}.raw", UriKind.Relative));
+                getResponse.EnsureSuccessStatusCode();
+                await using MemoryStream ms = new MemoryStream();
+                await getResponse.Content.CopyToAsync(ms);
+
+                byte[] roundTrippedBuffer = ms.ToArray();
+
+                CollectionAssert.AreEqual(blobData, roundTrippedBuffer);
+                Assert.AreEqual(blobHash, BlobIdentifier.FromBlob(roundTrippedBuffer));
+            }
+
+            {
+                // verify that head checks find the object
+                using HttpRequestMessage headRequest = new HttpRequestMessage(HttpMethod.Head, new Uri($"api/v1/refs/{TestNamespace}/bucket/{key}", UriKind.Relative));
+                HttpResponseMessage getResponse = await Client.SendAsync(headRequest);
+                getResponse.EnsureSuccessStatusCode();
+            }
+
+            {
+                // verify that the exists check doesnt find any issue
+                HttpResponseMessage existsResponse = await Client.GetAsync(new Uri($"api/v1/refs/{TestNamespace}/exists?names=bucket.{key}", UriKind.Relative));
+                existsResponse.EnsureSuccessStatusCode();
+
+                ExistCheckMultipleRefsResponse response = await existsResponse.Content.ReadAsAsync<ExistCheckMultipleRefsResponse>();
+                Assert.AreEqual(0, response.Missing.Count);
+            }
+
+            {
+                // delete the blob 
+                await Service.DeleteObject(TestNamespace, blobHash);
+            }
+
+            {
+                // the compact binary object still exists, so this returns a success (trying to resolve the attachment will fail)
+                HttpResponseMessage getResponse = await Client.GetAsync(new Uri($"api/v1/refs/{TestNamespace}/bucket/{key}.uecb", UriKind.Relative));
+                getResponse.EnsureSuccessStatusCode();
+            }
+
+            {
+                // the compact binary object still exists, so this returns a success (trying to resolve the attachment will fail)
+                HttpResponseMessage getResponse = await Client.GetAsync(new Uri($"api/v1/refs/{TestNamespace}/bucket/{key}.json", UriKind.Relative));
+                getResponse.EnsureSuccessStatusCode();
+            }
+
+            {
+                // we should now see a 404 as the blob is missing
+                HttpResponseMessage getResponse = await Client.GetAsync(new Uri($"api/v1/refs/{TestNamespace}/bucket/{key}.raw", UriKind.Relative));
+                Assert.AreEqual(HttpStatusCode.NotFound, getResponse.StatusCode);
+                Assert.AreEqual("application/problem+json", getResponse.Content.Headers.ContentType!.MediaType);
+                string s = await getResponse.Content.ReadAsStringAsync();
+                ProblemDetails? problem = JsonConvert.DeserializeObject<ProblemDetails>(s);
+                Assert.IsNotNull(problem);
+                Assert.AreEqual($"Object {blobHash} in {TestNamespace} not found", problem.Title);
+            }
+
+            {
+                // verify that head checks fail
+                using HttpRequestMessage headRequest = new HttpRequestMessage(HttpMethod.Head, new Uri($"api/v1/refs/{TestNamespace}/bucket/{key}", UriKind.Relative));
+                HttpResponseMessage getResponse = await Client.SendAsync(headRequest);
+                Assert.AreEqual(HttpStatusCode.NotFound, getResponse.StatusCode);
+            }
+
+            {
+                // verify that the exists check correctly returns a missing blob
+                HttpResponseMessage existsResponse = await Client.GetAsync(new Uri($"api/v1/refs/{TestNamespace}/exists?names=bucket.{key}", UriKind.Relative));
+                existsResponse.EnsureSuccessStatusCode();
+
+                ExistCheckMultipleRefsResponse response = await existsResponse.Content.ReadAsAsync<ExistCheckMultipleRefsResponse>();
+                Assert.AreEqual(1, response.Missing.Count);
+                Assert.AreEqual(key, response.Missing[0].Key);
+                Assert.AreEqual("bucket", response.Missing[0].Bucket.ToString());
+            }
+        }
+
+        [TestMethod]
         public async Task DeleteObject()
         {
             const string ObjectContents = "This is treated as a opaque blob";
