Horde: Add tracing to storage backends

diff --git a/Engine/Source/Programs/Horde/Horde.Build/Storage/Backends/AwsStorageBackend.cs b/Engine/Source/Programs/Horde/Horde.Build/Storage/Backends/AwsStorageBackend.cs
--- a/Engine/Source/Programs/Horde/Horde.Build/Storage/Backends/AwsStorageBackend.cs
+++ b/Engine/Source/Programs/Horde/Horde.Build/Storage/Backends/AwsStorageBackend.cs
@@ -15,6 +15,8 @@ using Amazon.S3.Model;
 using Horde.Build.Utilities;
 using Microsoft.Extensions.Configuration;
 using Microsoft.Extensions.Logging;
+using OpenTracing;
+using OpenTracing.Util;
 
 namespace Horde.Build.Storage.Backends
 {
@@ -257,6 +259,9 @@ namespace Horde.Build.Storage.Backends
 		/// <inheritdoc/>
 		public async Task<Stream?> ReadAsync(string path, CancellationToken cancellationToken)
 		{
+			using IScope scope = GlobalTracer.Instance.BuildSpan("AwsStorageBackend.ReadAsync").StartActive();
+			scope.Span.SetTag("Path", path);
+			
 			string fullPath = GetFullPath(path);
 
 			IDisposable? semaLock = null;
@@ -320,6 +325,9 @@ namespace Horde.Build.Storage.Backends
 		/// <inheritdoc/>
 		public async Task WriteInternalAsync(string fullPath, Stream stream, CancellationToken cancellationToken)
 		{
+			using IScope scope = GlobalTracer.Instance.BuildSpan("AwsStorageBackend.WriteInternalAsync").StartActive();
+			scope.Span.SetTag("Path", fullPath);
+			
 			const int MinPartSize = 5 * 1024 * 1024;
 
 			long streamLen = stream.Length;
@@ -426,6 +434,9 @@ namespace Horde.Build.Storage.Backends
 		/// <inheritdoc/>
 		public async Task DeleteAsync(string path, CancellationToken cancellationToken)
 		{
+			using IScope scope = GlobalTracer.Instance.BuildSpan("AwsStorageBackend.DeleteAsync").StartActive();
+			scope.Span.SetTag("Path", path);
+			
 			DeleteObjectRequest newDeleteRequest = new DeleteObjectRequest();
 			newDeleteRequest.BucketName = _options.AwsBucketName;
 			newDeleteRequest.Key = GetFullPath(path);
@@ -435,6 +446,9 @@ namespace Horde.Build.Storage.Backends
 		/// <inheritdoc/>
 		public async Task<bool> ExistsAsync(string path, CancellationToken cancellationToken)
 		{
+			using IScope scope = GlobalTracer.Instance.BuildSpan("AwsStorageBackend.ExistsAsync").StartActive();
+			scope.Span.SetTag("Path", path);
+			
 			try
 			{
 				GetObjectMetadataRequest request = new GetObjectMetadataRequest();
diff --git a/Engine/Source/Programs/Horde/Horde.Build/Storage/Backends/FileSystemStorageBackend.cs b/Engine/Source/Programs/Horde/Horde.Build/Storage/Backends/FileSystemStorageBackend.cs
--- a/Engine/Source/Programs/Horde/Horde.Build/Storage/Backends/FileSystemStorageBackend.cs
+++ b/Engine/Source/Programs/Horde/Horde.Build/Storage/Backends/FileSystemStorageBackend.cs
@@ -5,6 +5,8 @@ using System.IO;
 using System.Threading;
 using System.Threading.Tasks;
 using EpicGames.Core;
+using OpenTracing;
+using OpenTracing.Util;
 
 namespace Horde.Build.Storage.Backends
 {
@@ -53,6 +55,9 @@ namespace Horde.Build.Storage.Backends
 		/// <inheritdoc/>
 		public Task<Stream?> ReadAsync(string path, CancellationToken cancellationToken)
 		{
+			using IScope scope = GlobalTracer.Instance.BuildSpan("FileSystemStorageBackend.ReadAsync").StartActive();
+			scope.Span.SetTag("Path", path);
+			
 			FileReference location = FileReference.Combine(_baseDir, path);
 			if (!FileReference.Exists(location))
 			{
@@ -76,6 +81,9 @@ namespace Horde.Build.Storage.Backends
 		/// <inheritdoc/>
 		public async Task WriteAsync(string path, Stream stream, CancellationToken cancellationToken)
 		{
+			using IScope scope = GlobalTracer.Instance.BuildSpan("FileSystemStorageBackend.WriteAsync").StartActive();
+			scope.Span.SetTag("Path", path);
+			
 			FileReference finalLocation = FileReference.Combine(_baseDir, path);
 			if (!FileReference.Exists(finalLocation))
 			{
@@ -112,6 +120,9 @@ namespace Horde.Build.Storage.Backends
 		/// <inheritdoc/>
 		public Task<bool> ExistsAsync(string path, CancellationToken cancellationToken)
 		{
+			using IScope scope = GlobalTracer.Instance.BuildSpan("FileSystemStorageBackend.ExistsAsync").StartActive();
+			scope.Span.SetTag("Path", path);
+			
 			FileReference location = FileReference.Combine(_baseDir, path);
 			return Task.FromResult(FileReference.Exists(location));
 		}
@@ -119,6 +130,9 @@ namespace Horde.Build.Storage.Backends
 		/// <inheritdoc/>
 		public Task DeleteAsync(string path, CancellationToken cancellationToken)
 		{
+			using IScope scope = GlobalTracer.Instance.BuildSpan("FileSystemStorageBackend.DeleteAsync").StartActive();
+			scope.Span.SetTag("Path", path);
+			
 			FileReference location = FileReference.Combine(_baseDir, path);
 			FileReference.Delete(location);
 			return Task.CompletedTask;
diff --git a/Engine/Source/Programs/Horde/Horde.Build/Storage/Backends/RelayStorageBackend.cs b/Engine/Source/Programs/Horde/Horde.Build/Storage/Backends/RelayStorageBackend.cs
--- a/Engine/Source/Programs/Horde/Horde.Build/Storage/Backends/RelayStorageBackend.cs
+++ b/Engine/Source/Programs/Horde/Horde.Build/Storage/Backends/RelayStorageBackend.cs
@@ -6,6 +6,8 @@ using System.Net.Http;
 using System.Net.Http.Headers;
 using System.Threading;
 using System.Threading.Tasks;
+using OpenTracing;
+using OpenTracing.Util;
 
 namespace Horde.Build.Storage.Backends
 {
@@ -77,6 +79,9 @@ namespace Horde.Build.Storage.Backends
 		/// <inheritdoc/>
 		public async Task<Stream?> ReadAsync(string path, CancellationToken cancellationToken)
 		{
+			using IScope scope = GlobalTracer.Instance.BuildSpan("RelayStorageBackend.ReadAsync").StartActive();
+			scope.Span.SetTag("Path", path);
+			
 			Stream? localResult = await _localStorage.ReadAsync(path, cancellationToken);
 			if (localResult != null)
 			{
@@ -100,6 +105,9 @@ namespace Horde.Build.Storage.Backends
 		/// <inheritdoc/>
 		public Task WriteAsync(string path, Stream stream, CancellationToken cancellationToken)
 		{
+			using IScope scope = GlobalTracer.Instance.BuildSpan("RelayStorageBackend.WriteAsync").StartActive();
+			scope.Span.SetTag("Path", path);
+
 			return _localStorage.WriteAsync(path, stream, cancellationToken);
 		}
 
