Horde Storage - Missing file when adding content id store cache

diff --git a/Engine/Source/Programs/Horde/HordeStorage/Horde.Storage/Implementation/ContentId/MemoryCachedContentIdStore.cs b/Engine/Source/Programs/Horde/HordeStorage/Horde.Storage/Implementation/ContentId/MemoryCachedContentIdStore.cs
new file mode 100644
--- /dev/null
+++ b/Engine/Source/Programs/Horde/HordeStorage/Horde.Storage/Implementation/ContentId/MemoryCachedContentIdStore.cs
@@ -0,0 +1,92 @@
+// Copyright Epic Games, Inc. All Rights Reserved.
+
+using System;
+using System.Collections.Concurrent;
+using System.Threading.Tasks;
+using EpicGames.Horde.Storage;
+using Jupiter.Implementation;
+using Microsoft.Extensions.Caching.Memory;
+using Microsoft.Extensions.Options;
+
+namespace Horde.Storage.Implementation
+{
+    public class MemoryCachedContentIdStore : IContentIdStore
+    {
+        private readonly IContentIdStore _actualContentIdStore;
+        private readonly ConcurrentDictionary<NamespaceId, MemoryCache> _contentIdCaches = new ConcurrentDictionary<NamespaceId, MemoryCache>();
+        private readonly IOptionsMonitor<MemoryCacheContentIdSettings> _options;
+
+        public MemoryCachedContentIdStore(IContentIdStore actualContentIdStore, IOptionsMonitor<MemoryCacheContentIdSettings> options)
+        {
+            _actualContentIdStore = actualContentIdStore;
+            _options = options;
+        }
+
+        public IContentIdStore GetUnderlyingContentIdStore()
+        {
+            return _actualContentIdStore;
+        }
+
+        public void Clear()
+        {
+            _contentIdCaches.Clear();
+        }
+
+        public async Task<BlobIdentifier[]?> Resolve(NamespaceId ns, ContentId contentId, bool mustBeContentId)
+        {
+			MemoryCache cache = GetCacheForNamespace(ns);
+            if (cache.TryGetValue(contentId, out CachedContentIdEntry cachedResult))
+            {
+                return cachedResult.ReferencedBlobs;
+            }
+
+            BlobIdentifier[]? referencedBlobs = await _actualContentIdStore.Resolve(ns, contentId, mustBeContentId);
+            // we do not cache unresolvable content ids
+            if (referencedBlobs == null)
+            {
+                return null;
+            }
+
+            AddCacheEntry(ns, contentId, referencedBlobs);
+
+            return referencedBlobs;
+        }
+
+		public Task Put(NamespaceId ns, ContentId contentId, BlobIdentifier blobIdentifier, int contentWeight)
+        {
+            Task actualPutTask = _actualContentIdStore.Put(ns, contentId, blobIdentifier, contentWeight);
+			AddCacheEntry(ns, contentId, new BlobIdentifier[] {blobIdentifier});
+            return actualPutTask;
+        }
+
+        private void AddCacheEntry(NamespaceId ns, ContentId contentId, BlobIdentifier[] blobs)
+        {
+            MemoryCache cache = GetCacheForNamespace(ns);
+
+			CachedContentIdEntry cachedEntry = new CachedContentIdEntry
+            {
+                ReferencedBlobs = blobs
+			};
+            using ICacheEntry entry = cache.CreateEntry(contentId);
+            entry.Value = cachedEntry;
+            entry.Size = cachedEntry.GetSize();
+
+            if (_options.CurrentValue.EnableSlidingExpiry)
+            {
+                entry.SlidingExpiration = TimeSpan.FromMinutes(_options.CurrentValue.SlidingExpirationMinutes);
+            }
+		}
+
+        private MemoryCache GetCacheForNamespace(NamespaceId ns)
+        {
+            return _contentIdCaches.GetOrAdd(ns, id => new MemoryCache(_options.CurrentValue));
+        }
+    }
+
+    class CachedContentIdEntry
+    {
+        public int GetSize() { return ReferencedBlobs.Length * 20; }
+
+        public BlobIdentifier[] ReferencedBlobs { get; set; } = Array.Empty<BlobIdentifier>();
+    }
+}
