Horde Storage - Handle lastAccessTime not being set when enumerating all records

diff --git a/Engine/Source/Programs/Horde/HordeStorage/Horde.Storage/Implementation/Objects/ScyllaReferencesStore.cs b/Engine/Source/Programs/Horde/HordeStorage/Horde.Storage/Implementation/Objects/ScyllaReferencesStore.cs
--- a/Engine/Source/Programs/Horde/HordeStorage/Horde.Storage/Implementation/Objects/ScyllaReferencesStore.cs
+++ b/Engine/Source/Programs/Horde/HordeStorage/Horde.Storage/Implementation/Objects/ScyllaReferencesStore.cs
@@ -137,7 +137,7 @@ namespace Horde.Storage.Implementation
 
                 string bucket = row.GetValue<string>("bucket");
                 string name = row.GetValue<string>("name");
-                DateTime lastAccessTime = row.GetValue<DateTime>("last_access_time");
+                DateTime? lastAccessTime = row.GetValue<DateTime?>("last_access_time");
 
                 // skip any names that are not conformant to io hash
                 if (name.Length != 40)
@@ -145,7 +145,9 @@ namespace Horde.Storage.Implementation
                     continue;
                 }
 
-                yield return (new BucketId(bucket), new IoHashKey(name), lastAccessTime);
+                // if last access time is missing we treat it as being very old
+                lastAccessTime ??= DateTime.MinValue;
+                yield return (new BucketId(bucket), new IoHashKey(name), lastAccessTime.Value);
             }
         }
 
