Horde: Fix build error.

diff --git a/Engine/Source/Programs/Horde/Horde.Agent/Commands/Bundles/CreateCommand.cs b/Engine/Source/Programs/Horde/Horde.Agent/Commands/Bundles/CreateCommand.cs
--- a/Engine/Source/Programs/Horde/Horde.Agent/Commands/Bundles/CreateCommand.cs
+++ b/Engine/Source/Programs/Horde/Horde.Agent/Commands/Bundles/CreateCommand.cs
@@ -140,11 +140,12 @@ namespace Horde.Agent.Commands.Bundles
 		{
 			using IMemoryCache cache = new MemoryCache(new MemoryCacheOptions());
 			using ITreeStore store = CreateTreeStore(logger, cache);
+			ITreeWriter writer = store.CreateTreeWriter(RefName);
 
 			DirectoryNode node = new DirectoryNode();
-			await node.CopyFromDirectoryAsync(InputDir.ToDirectoryInfo(), new ChunkingOptions(), logger, CancellationToken.None);
-			await store.WriteTreeAsync(RefName, node);
+			await node.CopyFromDirectoryAsync(InputDir.ToDirectoryInfo(), new ChunkingOptions(), writer, logger, CancellationToken.None);
 
+			await writer.FlushAsync(node, CancellationToken.None);
 			return 0;
 		}
 	}
