Horde: Improve fleet manager for ephemeral/terminating EC2 instances

* Specifies which pool to join for the user-data/cloud-init script to pick up
* Allows both name and ARN based instance profiles to be specified
* Properly encodes user-data as base64
* Passes in credentials in constructor for easier testing (e.g read credentials from default $HOME/.aws location)
* Upgrade of NuGet packages required to get the new InstanceMetadataOptionsRequest

diff --git a/Engine/Source/Programs/Horde/Horde.Build/Agents/Fleet/Providers/AwsTerminatingFleetManager.cs b/Engine/Source/Programs/Horde/Horde.Build/Agents/Fleet/Providers/AwsTerminatingFleetManager.cs
--- a/Engine/Source/Programs/Horde/Horde.Build/Agents/Fleet/Providers/AwsTerminatingFleetManager.cs
+++ b/Engine/Source/Programs/Horde/Horde.Build/Agents/Fleet/Providers/AwsTerminatingFleetManager.cs
@@ -8,6 +8,7 @@ using System.Threading.Tasks;
 using Amazon;
 using Amazon.EC2;
 using Amazon.EC2.Model;
+using Amazon.Extensions.NETCore.Setup;
 using Horde.Build.Agents.Pools;
 using Horde.Build.Auditing;
 using Microsoft.Extensions.Logging;
@@ -103,6 +104,22 @@ namespace Horde.Build.Agents.Fleet.Providers
 			int index = new Random().Next(SubnetIds.Count);
 			return SubnetIds[index];
 		}
+
+		/// <summary>
+		/// Get user-data as base64
+		/// Required when interacting with the EC2 API.
+		/// </summary>
+		/// <returns>User-data encoded as Base64</returns>
+		public string? GetUserDataAsBase64()
+		{
+			if (UserData == null)
+			{
+				return UserData;
+			}
+			
+			byte[] plainTextBytes = System.Text.Encoding.UTF8.GetBytes(UserData);
+			return Convert.ToBase64String(plainTextBytes);
+		}
 	}
 	
 	/// <summary>
@@ -125,7 +142,7 @@ namespace Horde.Build.Agents.Fleet.Providers
 		/// <summary>
 		/// Constructor
 		/// </summary>
-		public AwsFleetManager(IAgentCollection agentCollection, AwsInstanceLaunchSettings launchSettings, ILogger<AwsFleetManager> logger)
+		public AwsFleetManager(AWSOptions awsOptions, IAgentCollection agentCollection, AwsInstanceLaunchSettings launchSettings, ILogger<AwsFleetManager> logger)
 		{
 			_agentCollection = agentCollection;
 			_launchSettings = launchSettings;
@@ -133,7 +150,7 @@ namespace Horde.Build.Agents.Fleet.Providers
 
 			AmazonEC2Config config = new () { RegionEndpoint = RegionEndpoint.GetBySystemName(_launchSettings.Region) };
 			logger.LogInformation("Initializing AWS fleet manager for region {Region}", config.RegionEndpoint);
-			_client = new AmazonEC2Client(config);
+			_client = new AmazonEC2Client(awsOptions.Credentials, config);
 		}
 
 		/// <inheritdoc/>
@@ -151,6 +168,8 @@ namespace Horde.Build.Agents.Fleet.Providers
 			scope.Span.SetTag("count", count);
 
 			List<Tag> tags = _launchSettings.Tags.Select(x => new Tag(x.Key, x.Value)).ToList();
+			tags.Add(new Tag("Horde:RequestedPools", pool.Name));
+			tags.Add(new Tag("Horde:Timestamp:Ec2InstanceLaunchRequested", Convert.ToString(DateTimeOffset.UtcNow.ToUnixTimeMilliseconds())));
 			
 			RunInstancesRequest request = new ()
 			{
@@ -161,13 +180,21 @@ namespace Horde.Build.Agents.Fleet.Providers
 				SubnetId = _launchSettings.GetRandomSubnetId(), // Pick randomly to spread evenly among available subnets 
 				MinCount = count,
 				MaxCount = count,
-				UserData = _launchSettings.UserData,
-				InstanceInitiatedShutdownBehavior = ShutdownBehavior.Terminate
+				UserData = _launchSettings.GetUserDataAsBase64(),
+				InstanceInitiatedShutdownBehavior = ShutdownBehavior.Terminate,
+				MetadataOptions = new InstanceMetadataOptionsRequest
+				{
+					HttpEndpoint = InstanceMetadataEndpointState.Enabled,
+					InstanceMetadataTags = InstanceMetadataTagsState.Enabled,
+				}
 			};
 
 			if (_launchSettings.IamInstanceProfile != null)
 			{
-				request.IamInstanceProfile = new IamInstanceProfileSpecification { Name = _launchSettings.IamInstanceProfile };
+				request.IamInstanceProfile =
+					_launchSettings.IamInstanceProfile.StartsWith("arn:aws", StringComparison.InvariantCulture)
+						? new IamInstanceProfileSpecification { Arn = _launchSettings.IamInstanceProfile }
+						: new IamInstanceProfileSpecification { Name = _launchSettings.IamInstanceProfile };
 			}
 			
 			if (tags.Count > 0)
diff --git a/Engine/Source/Programs/Horde/Horde.Build/Horde.Build.csproj b/Engine/Source/Programs/Horde/Horde.Build/Horde.Build.csproj
--- a/Engine/Source/Programs/Horde/Horde.Build/Horde.Build.csproj
+++ b/Engine/Source/Programs/Horde/Horde.Build/Horde.Build.csproj
@@ -19,10 +19,10 @@
   </PropertyGroup>
 
   <ItemGroup>
-    <PackageReference Include="AWSSDK.EC2" Version="3.7.0" />
-    <PackageReference Include="AWSSDK.Extensions.NETCore.Setup" Version="3.7.0.1" />
-    <PackageReference Include="AWSSDK.S3" Version="3.7.0.17" />
-    <PackageReference Include="AWSSDK.SecurityToken" Version="3.7.1.69" />
+    <PackageReference Include="AWSSDK.EC2" Version="3.7.75.1" />
+    <PackageReference Include="AWSSDK.Extensions.NETCore.Setup" Version="3.7.2" />
+    <PackageReference Include="AWSSDK.S3" Version="3.7.9.19" />
+    <PackageReference Include="AWSSDK.SecurityToken" Version="3.7.1.168" />
     <PackageReference Include="Datadog.Trace.OpenTracing" Version="2.11.0" />
     <PackageReference Include="DogStatsD-CSharp-Client" Version="6.0.0" />
     <PackageReference Include="Grpc.AspNetCore" Version="2.45.0" />
