Horde: New individual device and pool telemetry with tests
#rnx

diff --git a/Engine/Source/Programs/Horde/Horde.Build.Tests/DeviceServiceTests.cs b/Engine/Source/Programs/Horde/Horde.Build.Tests/DeviceServiceTests.cs
new file mode 100644
--- /dev/null
+++ b/Engine/Source/Programs/Horde/Horde.Build.Tests/DeviceServiceTests.cs
@@ -0,0 +1,248 @@
+// Copyright Epic Games, Inc. All Rights Reserved.
+
+using System.Collections.Generic;
+using System.Threading.Tasks;
+using Horde.Build.Devices;
+using Horde.Build.Jobs;
+using Horde.Build.Server;
+using Microsoft.AspNetCore.Mvc;
+using Microsoft.VisualStudio.TestTools.UnitTesting;
+using Horde.Build.Projects;
+using Horde.Build.Utilities;
+using Horde.Build.Users;
+using System.Security.Claims;
+using Microsoft.AspNetCore.Http;
+
+namespace Horde.Build.Tests
+{
+
+	using ProjectId = StringId<IProject>;
+	using DeviceId = StringId<IDevice>;
+
+	/// <summary>
+	/// Tests for the device service
+	/// </summary>
+	[TestClass]
+	public class DeviceServiceTest : TestSetup
+	{
+		private DevicesController? _deviceController;
+
+		// override DeviceController with valid user
+		private DevicesController DeviceController
+		{
+			get
+			{
+				if (_deviceController == null)
+				{
+					IUser user = UserCollection.FindOrAddUserByLoginAsync("TestUser").Result;
+					_deviceController = base.DevicesController;
+					ControllerContext controllerContext = new ControllerContext();
+					controllerContext.HttpContext = new DefaultHttpContext();
+					controllerContext.HttpContext.User = new ClaimsPrincipal(new ClaimsIdentity(
+						new List<Claim> {new Claim(ServerSettings.AdminClaimType, ServerSettings.AdminClaimValue),
+					new Claim(ClaimTypes.Name, "TestUser"),
+					new Claim(HordeClaimTypes.UserId, user.Id.ToString()) }
+						, "TestAuthType"));
+					_deviceController.ControllerContext = controllerContext;
+
+				}
+				return _deviceController;
+			}
+		}
+
+
+		private async Task PopulateDevices()
+		{
+
+			await ProjectCollection.AddOrUpdateAsync(new ProjectId("ue5"), "", "", 0, new ProjectConfig { Name = "UE5" });
+
+			// create 2 pools
+			for (int i = 1; i < 3; i++)
+			{
+				await DeviceController.CreatePoolAsync(new CreateDevicePoolRequest() { Name = "TestDevicePool" + i, PoolType = DevicePoolType.Automation, ProjectIds = new List<string>() { "ue5" } });
+			}
+
+			Dictionary<string, string> platformMap = new Dictionary<string, string>();
+			Dictionary<string, string> platformReverseMap = new Dictionary<string, string>();
+
+			// create 3 platforms
+			for (int i = 1; i < 4; i++)
+			{
+				string platformName = "TestDevicePlatform" + i;
+				await DevicesController.CreatePlatformAsync(new CreateDevicePlatformRequest() { Name = platformName });
+				platformMap[platformName] = platformName;
+				platformReverseMap[platformName] = platformName;
+			}
+
+			// add 4 devices to each platform, split between 2 pools
+			for (int i = 1; i < 4; i++)
+			{
+				for (int j = 1; j < 5; j++)
+				{
+					// @todo: add model tests
+					// one base model, and 3 other models 
+					/*
+					string? ModelId = null;
+					if (j > 1)
+					{
+						ModelId = "TestDevicePlatform" + i + "_Model" + j;
+					}
+					*/
+
+					string poolId = (j & 1) != 0 ? "testdevicepool1" : "testdevicepool2";
+
+					await DeviceController.CreateDeviceAsync(new CreateDeviceRequest() { Name = "TestDevice" + j + "_Platform" + i + "_" + poolId, Address = "10.0.0.1", Enabled = true, PlatformId = "testdeviceplatform" + i/*, ModelId = ModelId*/, PoolId = poolId });
+				}
+			}
+
+			// populate the legacy platform map, this exists to avoid having platform names in server code, though needs to be cleaned up.
+			// Gauntlet in various streams can't/hasn't been updated and asks for non-platform-id device types
+			await DeviceController.UpdatePlatformMapV1(new UpdatePlatformMapRequest() { PlatformMap = platformMap, PlatformReverseMap = platformReverseMap, PerfSpecHighMap = new Dictionary<string, string>() });
+
+		}
+
+		static T ResultToValue<T>(ActionResult<T> result) where T: class
+		{
+			return ((result.Result! as JsonResult)!.Value! as T)!;
+		}
+	
+		async Task<LegacyCreateReservationRequest> SetupReservationTest()
+		{
+			await CreateFixtureAsync();
+			await PopulateDevices();
+
+			ActionResult<List<object>> res = await JobsController.FindJobsAsync();
+			Assert.AreEqual(2, res.Value!.Count);
+			GetJobResponse job = (res.Value[0] as GetJobResponse)!;
+
+			// Gauntlet uses the legacy v1 API
+			LegacyCreateReservationRequest request = new LegacyCreateReservationRequest();
+			request.PoolId = "TestDevicePool1";
+			request.DeviceTypes = new string[] { "TestDevicePlatform1" };
+			request.Hostname = "localhost";
+			request.Duration = "00:10:00";
+			request.JobId = job.Id;
+			request.StepId = "abcd";
+
+			return request;
+		}
+
+
+		[TestMethod]
+		public async Task TestReservation()
+		{
+			LegacyCreateReservationRequest request = await SetupReservationTest();
+			
+			// create a reservation
+			GetLegacyReservationResponse reservation = ResultToValue(await DeviceController!.CreateDeviceReservationV1Async(request));			
+			Assert.AreEqual(1, reservation.DeviceNames.Length);
+
+			// get the device in the reservation, and make sure it is the right platform
+			GetLegacyDeviceResponse device = ResultToValue(await DeviceController!.GetDeviceV1Async(reservation.DeviceNames[0]));
+			Assert.AreEqual("TestDevicePlatform1", device.Type);
+
+			// update the reservation
+			GetLegacyReservationResponse renewed = ResultToValue(await DeviceController!.UpdateReservationV1Async(reservation.Guid));
+			Assert.AreEqual(renewed.Guid, reservation.Guid);
+
+			// delete the reservation, would be nice to also test reservation expiration
+			OkResult? deleted = (await DeviceController!.DeleteReservationV1Async(reservation.Guid)) as OkResult;
+			Assert.IsNotNull(deleted);
+
+			// check that telemetry was created
+			List<GetDeviceTelemetryResponse> telemetry = (await DeviceController!.GetDeviceTelemetry()).Value!;
+			Assert.AreEqual(telemetry.Count, 1);
+			Assert.AreEqual(telemetry[0].Telemetry.Count, 1);
+			Assert.AreEqual(telemetry[0].Telemetry[0].StepId, "abcd");
+		}
+
+		[TestMethod]
+		public async Task TestProblemDevice()
+		{
+			LegacyCreateReservationRequest request = await SetupReservationTest();
+
+			// create a reservation
+			GetLegacyReservationResponse reservation = ResultToValue(await DeviceController!.CreateDeviceReservationV1Async(request));
+			Assert.AreEqual(1, reservation.DeviceNames.Length);
+
+			// report a problem with the device 
+			OkResult? problem = (await DeviceController!.PutDeviceErrorAsync(reservation.DeviceNames[0])) as OkResult;
+			Assert.IsNotNull(problem);
+
+			OkResult? deleted = (await DeviceController!.DeleteReservationV1Async(reservation.Guid)) as OkResult;
+			Assert.IsNotNull(deleted);
+
+			// check that telemetry was created
+			List<GetDeviceTelemetryResponse> telemetry = (await DeviceController!.GetDeviceTelemetry()).Value!;
+			Assert.AreEqual(telemetry.Count, 1);
+			// check that problem was recorded
+			Assert.IsNotNull(telemetry[0].Telemetry[0].ProblemTimeUtc);
+			// make sure the reswrvation time was finished for the problem device
+			Assert.IsNotNull(telemetry[0].Telemetry[0].ReservationFinishUtc);
+		}
+
+		[TestMethod]
+		public async Task TestDevicePoolTelemetryCapture()
+		{
+			LegacyCreateReservationRequest reservationRequest = await SetupReservationTest();
+
+			// set some device status
+			UpdateDeviceRequest request = new UpdateDeviceRequest();
+
+			request.Maintenance = true;
+			await DeviceController!.UpdateDeviceAsync("testdevice1_platform1_testdevicepool1", request);
+
+			request.Maintenance = null;
+			request.Enabled = false;
+			await DeviceController!.UpdateDeviceAsync("testdevice3_platform3_testdevicepool1", request);
+
+			request.Enabled = null;
+			request.Problem = true;
+			await DeviceController!.UpdateDeviceAsync("testdevice4_platform2_testdevicepool2", request);
+
+
+			IDevice? maintenanceDevice = await DeviceService.GetDeviceAsync(new DeviceId("testdevice1_platform1_testdevicepool1"));
+			Assert.IsNotNull(maintenanceDevice);
+			
+			GetLegacyReservationResponse reservation = ResultToValue(await DeviceController!.CreateDeviceReservationV1Async(reservationRequest));
+			Assert.AreEqual(1, reservation.DeviceNames.Length);
+			Assert.AreNotSame(maintenanceDevice!.Name, reservation.DeviceNames[0]);
+
+			await DeviceService.TickForTestingAsync();
+
+			List<GetDevicePoolTelemetryResponse> telemetry = (await DeviceController!.GetDevicePoolTelemetry()).Value!;
+			Assert.AreEqual(1, telemetry.Count);
+			foreach (GetDevicePlatformTelemetryResponse t in telemetry[0].Telemetry["testdevicepool1"])
+			{
+				if (t.PlatformId == "testdeviceplatform1")
+				{
+					Assert.AreEqual(0, t.Available);
+					Assert.AreEqual(1, t.Reserved);
+					Assert.AreEqual(1, t.Maintenance);
+				}
+				else if (t.PlatformId == "testdeviceplatform3")
+				{
+					Assert.AreEqual(1, t.Available);
+					Assert.AreEqual(1, t.Disabled);
+				}
+				else
+				{
+					Assert.AreEqual(2, t.Available);
+				}
+			}
+
+			foreach (GetDevicePlatformTelemetryResponse t in telemetry[0].Telemetry["testdevicepool2"])
+			{
+				if (t.PlatformId == "testdeviceplatform2")
+				{
+					Assert.AreEqual(1, t.Available);
+					Assert.AreEqual(1, t.Problem);
+				}
+				else
+				{
+					Assert.AreEqual(2, t.Available);
+				}
+			}
+		}
+	}
+}
diff --git a/Engine/Source/Programs/Horde/Horde.Build.Tests/TestSetup.cs b/Engine/Source/Programs/Horde/Horde.Build.Tests/TestSetup.cs
--- a/Engine/Source/Programs/Horde/Horde.Build.Tests/TestSetup.cs
+++ b/Engine/Source/Programs/Horde/Horde.Build.Tests/TestSetup.cs
@@ -58,6 +58,7 @@ using Horde.Build.Tasks;
 using Horde.Build.Auditing;
 using Horde.Build.Storage.Backends;
 using EpicGames.Horde.Storage.Bundles;
+using Horde.Build.Devices;
 
 namespace Horde.Build.Tests
 {
@@ -90,6 +91,7 @@ namespace Horde.Build.Tests
 		public IAgentSoftwareCollection AgentSoftwareCollection => ServiceProvider.GetRequiredService <IAgentSoftwareCollection>();
 		public ITestDataCollection TestDataCollection => ServiceProvider.GetRequiredService<ITestDataCollection>();
 		public IUserCollection UserCollection => ServiceProvider.GetRequiredService<IUserCollection>();
+		public IDeviceCollection DeviceCollection => ServiceProvider.GetRequiredService<IDeviceCollection>();
 
 		public AclService AclService => ServiceProvider.GetRequiredService<AclService>();
 		public AgentSoftwareService AgentSoftwareService => ServiceProvider.GetRequiredService<AgentSoftwareService>();
@@ -112,6 +114,7 @@ namespace Horde.Build.Tests
 		public PoolService PoolService => ServiceProvider.GetRequiredService<PoolService>();
 		public LifetimeService LifetimeService => ServiceProvider.GetRequiredService<LifetimeService>();
 		public ScheduleService ScheduleService => ServiceProvider.GetRequiredService<ScheduleService>();
+		public DeviceService DeviceService => ServiceProvider.GetRequiredService<DeviceService>();
 
 		public ServerSettings ServerSettings => ServiceProvider.GetRequiredService<IOptions<ServerSettings>>().Value;
 		public IOptionsMonitor<ServerSettings> ServerSettingsMon => ServiceProvider.GetRequiredService<IOptionsMonitor<ServerSettings>>();
@@ -120,6 +123,7 @@ namespace Horde.Build.Tests
 		public AgentsController AgentsController => GetAgentsController();
 		public PoolsController PoolsController => GetPoolsController();
 		public LeasesController LeasesController => GetLeasesController();
+		public DevicesController DevicesController => GetDevicesController();
 
 		private static bool s_datadogWriterPatched;
 
@@ -181,6 +185,7 @@ namespace Horde.Build.Tests
 			services.AddSingleton<ITemplateCollection, TemplateCollection>();
 			services.AddSingleton<IUgsMetadataCollection, UgsMetadataCollection>();
 			services.AddSingleton<IUserCollection, UserCollectionV1>();
+			services.AddSingleton<IDeviceCollection, DeviceCollection>();
 
 			services.AddSingleton<ToolCollection>();
 
@@ -216,6 +221,7 @@ namespace Horde.Build.Tests
 			services.AddSingleton<ScheduleService>();
 			services.AddSingleton<StreamService>();
 			services.AddSingleton<UpgradeService>();
+			services.AddSingleton<DeviceService>();
 
 			services.AddSingleton<ConformTaskSource>();
 			services.AddSingleton<ICommitService, CommitService>();
@@ -233,6 +239,7 @@ namespace Horde.Build.Tests
 			services.AddSingleton<IStorageClient, BasicStorageClient>();
 
 			services.AddSingleton<ISingletonDocument<AgentSoftwareChannels>>(new SingletonDocumentStub<AgentSoftwareChannels>());
+			services.AddSingleton<ISingletonDocument<DevicePlatformMapV1>>(new SingletonDocumentStub<DevicePlatformMapV1>());
 		}
 
 		public async Task<IStream?> CreateOrReplaceStreamAsync(StreamId streamId, IStream? stream, ProjectId projectId, StreamConfig config)
@@ -262,6 +269,14 @@ namespace Horde.Build.Tests
 	        return jobsCtrl;
         }
 
+		private DevicesController GetDevicesController()
+		{
+			ILogger<DevicesController> logger = ServiceProvider.GetRequiredService<ILogger<DevicesController>>();
+			DevicesController devicesCtrl = new DevicesController(DeviceService, AclService, UserCollection, logger);
+			devicesCtrl.ControllerContext = GetControllerContext();
+			return devicesCtrl;
+		}
+
 		private AgentsController GetAgentsController()
 		{
 			AgentsController agentCtrl = new AgentsController(AclService, AgentService);
diff --git a/Engine/Source/Programs/Horde/Horde.Build/Devices/DeviceCollection.cs b/Engine/Source/Programs/Horde/Horde.Build/Devices/DeviceCollection.cs
--- a/Engine/Source/Programs/Horde/Horde.Build/Devices/DeviceCollection.cs
+++ b/Engine/Source/Programs/Horde/Horde.Build/Devices/DeviceCollection.cs
@@ -63,7 +63,7 @@ namespace Horde.Build.Devices
 			public DevicePoolId Id { get; set; }
 
 			[BsonRequired]
-            public DevicePoolType PoolType { get; set; }
+			public DevicePoolType PoolType { get; set; }
 
 			[BsonIgnoreIfNull]
 			public List<ProjectId>? ProjectIds { get; set; }
@@ -77,15 +77,15 @@ namespace Horde.Build.Devices
 			[BsonConstructor]
 			private DevicePoolDocument()
 			{
-            }
+			}
 
 			public DevicePoolDocument(DevicePoolId id, string name, DevicePoolType poolType, List<ProjectId>? projectIds)
 			{
 				Id = id;
 				Name = name;
-                PoolType = poolType;
+				PoolType = poolType;
 				ProjectIds = projectIds;
-            }
+			}
 		}
 
 		/// <summary>
@@ -108,16 +108,16 @@ namespace Horde.Build.Devices
 			[BsonIgnoreIfNull]
 			public string? StepId { get; set; }
 
-            /// <summary>
-            /// Reservations held by a user, requires a token like download code
-            /// </summary>
-            [BsonIgnoreIfNull]
+			/// <summary>
+			/// Reservations held by a user, requires a token like download code
+			/// </summary>
+			[BsonIgnoreIfNull]
 			public UserId? UserId { get; set; }
 
-            /// <summary>
-            /// The hostname of the machine which has made the reservation
-            /// </summary>
-            [BsonIgnoreIfNull]
+			/// <summary>
+			/// The hostname of the machine which has made the reservation
+			/// </summary>
+			[BsonIgnoreIfNull]
 			public string? Hostname { get; set; }
 
 			/// <summary>
@@ -156,10 +156,10 @@ namespace Horde.Build.Devices
 				UpdateTimeUtc = createTimeUtc;
 				Hostname = hostname;
 				ReservationDetails = reservationDetails;
-                JobId = jobId;
-                StepId = stepId;
+				JobId = jobId;
+				StepId = stepId;
 
-                LegacyGuid = Guid.NewGuid().ToString();
+				LegacyGuid = Guid.NewGuid().ToString();
 			}
 		}
 
@@ -212,6 +212,10 @@ namespace Horde.Build.Devices
 			[BsonIgnoreIfNull]
 			public string? Notes { get; set; }
 
+
+			/// <summary>
+			/// [DEPRECATED]
+			/// </summary>
 			[BsonIgnoreIfNull]
 			public List<DeviceUtilizationTelemetry>? Utilization { get; set; }
 
@@ -232,10 +236,198 @@ namespace Horde.Build.Devices
 				Enabled = enabled;
 				Address = address;
 				ModelId = modelId;
-                ModifiedByUser = userId?.ToString();
-            }
+				ModifiedByUser = userId?.ToString();
+			}
+		}
+
+		/// <summary>
+		/// Device telemetry information for an individual device 
+		/// </summary>
+		class DeviceTelemetryDocument : IDeviceTelemetry
+		{
+			/// <summary>
+			/// Id of telemetry document
+			/// </summary>
+			[BsonRequired, BsonId]
+			public ObjectId TelemetryId { get; set; }
+
+			/// <summary>
+			/// The device id
+			/// </summary>
+			[BsonRequired]
+			public DeviceId DeviceId { get; set; }
+
+			/// <summary>
+			/// The time this telemetry data was created
+			/// </summary>
+			[BsonRequired]
+			public DateTime CreateTimeUtc { get; set; }
+
+			/// <summary>
+			/// The job id which utilized device
+			/// </summary>
+			[BsonIgnoreIfNull]
+			public string? JobId { get; set; }
+
+			/// <summary>
+			/// The job's step id
+			/// </summary>
+			[BsonIgnoreIfNull]
+			public string? StepId { get; set; }
+
+			/// <summary>
+			/// Reservation Id (transient, reservations are deleted upon expiration)
+			/// </summary>
+			[BsonIgnoreIfNull]
+			public ObjectId? ReservationId { get; set; }
+
+			/// <summary>
+			/// The time device was reserved
+			/// </summary>
+			[BsonIgnoreIfNull]
+			public DateTime? ReservationStartUtc { get; set; }
+
+			/// <summary>
+			/// The time device was freed
+			/// </summary>
+			[BsonIgnoreIfNull]
+			public DateTime? ReservationFinishUtc { get; set; }
+
+			/// <summary>
+			/// If the device reported a problem
+			/// </summary>
+			[BsonIgnoreIfNull]
+			public DateTime? ProblemTimeUtc { get; set; }
+
+			[BsonConstructor]
+			private DeviceTelemetryDocument()
+			{
+			}
+
+			public DeviceTelemetryDocument(DeviceId deviceId, ObjectId? reservationId = null, DateTime? reservationStartTime = null, string? jobId = null, string? stepId = null)
+			{
+				TelemetryId = ObjectId.GenerateNewId();
+				CreateTimeUtc = DateTime.UtcNow;
+				DeviceId = deviceId;
+
+				ReservationId = reservationId;
+				ReservationStartUtc = reservationStartTime;
+
+				JobId = jobId;
+				StepId = stepId;
+			}
+
+		}
+
+		class DevicePlatformTelemetryDocument : IDevicePlatformTelemetry
+		{
+			/// <summary>
+			/// Id of telemetry document
+			/// </summary>
+			[BsonRequired, BsonId]
+			public ObjectId TelemetryId { get; set; }
+
+			/// <summary>
+			/// The time this telemetry data was created
+			/// </summary>
+			[BsonRequired]
+			public DateTime CreateTimeUtc { get; set; }
+
+			/// <summary>
+			/// The platform id 
+			/// </summary>
+			[BsonRequired]
+			public DevicePlatformId PlatformId { get; set; }
+
+			/// <summary>
+			/// Number of available devices of this platform 
+			/// </summary>
+			[BsonRequired]
+			public int Available { get; set; } = 0;
+
+			/// <summary>
+			/// Number of reserved devices of this platform 
+			/// </summary>
+			[BsonRequired]
+			public int Reserved { get; set; } = 0;
+
+			/// <summary>
+			/// Number of devices in maintenance state
+			/// </summary>
+			[BsonRequired]
+			public int Maintenance { get; set; } = 0;
+
+			/// <summary>
+			/// Number of devices in problem state
+			/// </summary>
+			[BsonRequired]
+			public int Problem { get; set; } = 0;
+
+			/// <summary>
+			/// Number of devices that are disabled
+			/// </summary>
+			[BsonRequired]
+			public int Disabled { get; set; } = 0;
+
+
+			[BsonConstructor]
+			private DevicePlatformTelemetryDocument()
+			{
+			}
+
+			public DevicePlatformTelemetryDocument(DevicePlatformId platformId, int available, int reserved, int maintenance, int problem, int disabled)
+			{
+				TelemetryId = ObjectId.GenerateNewId();
+				CreateTimeUtc = DateTime.UtcNow;
+				PlatformId = platformId;
+				Available = available;
+				Reserved = reserved;
+				Maintenance = maintenance;	
+				Problem = problem;
+				Disabled = disabled;
+			}
+
 		}
 
+		/// <summary>
+		/// Device telemetry information for pools
+		/// </summary>
+		class DevicePoolTelemetryDocument : IDevicePoolTelemetry
+		{
+			/// <summary>
+			/// Id of telemetry document
+			/// </summary>
+			[BsonRequired, BsonId]
+			public ObjectId TelemetryId { get; set; }
+
+			/// <summary>
+			/// The time this telemetry data was created
+			/// </summary>
+			[BsonRequired]
+			public DateTime CreateTimeUtc { get; set; }
+
+			/// <summary>
+			/// Pool platform state
+			/// </summary>
+			[BsonRequired]
+			public Dictionary<DevicePoolId, List<DevicePlatformTelemetryDocument>> Pools { get; set; } = new Dictionary<DevicePoolId, List<DevicePlatformTelemetryDocument>>();
+
+			IReadOnlyDictionary<DevicePoolId, IReadOnlyList<IDevicePlatformTelemetry>> IDevicePoolTelemetry.Pools => Pools.ToDictionary(kvp => kvp.Key, kvp => kvp.Value as IReadOnlyList<IDevicePlatformTelemetry>);
+
+			[BsonConstructor]
+			private DevicePoolTelemetryDocument()
+			{
+			}
+
+			public DevicePoolTelemetryDocument(Dictionary<DevicePoolId, List<DevicePlatformTelemetryDocument>> pools)
+			{
+				TelemetryId = ObjectId.GenerateNewId();
+				CreateTimeUtc = DateTime.UtcNow;
+				Pools = pools;
+			}
+		}
+
+
 		readonly IMongoCollection<DevicePlatformDocument> _platforms;
 
 		readonly IMongoCollection<DeviceDocument> _devices;
@@ -244,6 +436,10 @@ namespace Horde.Build.Devices
 
 		readonly IMongoCollection<DeviceReservationDocument> _reservations;
 
+		readonly IMongoCollection<DeviceTelemetryDocument> _deviceTelemetry;
+
+		readonly IMongoCollection<DevicePoolTelemetryDocument> _poolTelemetry;
+
 		readonly int _checkoutDays = 7;
 
 		/// <summary>
@@ -255,6 +451,14 @@ namespace Horde.Build.Devices
 			_platforms = mongoService.GetCollection<DevicePlatformDocument>("Devices.Platforms", keys => keys.Ascending(x => x.Name), unique: true);
 			_pools = mongoService.GetCollection<DevicePoolDocument>("Devices.Pools", keys => keys.Ascending(x => x.Name), unique: true);
 			_reservations = mongoService.GetCollection<DeviceReservationDocument>("Devices.Reservations");
+
+			List<MongoIndex<DeviceTelemetryDocument>> deviceTelemetryIndexes = new List<MongoIndex<DeviceTelemetryDocument>>();
+			deviceTelemetryIndexes.Add((keys => keys.Descending(x => x.CreateTimeUtc)));
+			_deviceTelemetry = mongoService.GetCollection<DeviceTelemetryDocument>("Devices.DeviceTelemetry", deviceTelemetryIndexes);
+
+			List<MongoIndex<DevicePoolTelemetryDocument>> poolTelemetryIndexes = new List<MongoIndex<DevicePoolTelemetryDocument>>();
+			poolTelemetryIndexes.Add((keys => keys.Descending(x => x.CreateTimeUtc)));
+			_poolTelemetry = mongoService.GetCollection<DevicePoolTelemetryDocument>("Devices.PoolTelemetry", poolTelemetryIndexes);
 		}
 
 		/// <inheritdoc/>
@@ -358,13 +562,22 @@ namespace Horde.Build.Devices
 		}
 
 		/// <inheritdoc/>
-		public async Task<List<IDevice>> FindAllDevicesAsync(List<DeviceId>? deviceIds = null)
+		public async Task<List<IDevice>> FindAllDevicesAsync(List<DeviceId>? deviceIds = null, DevicePoolId? poolId = null, DevicePlatformId? platformId = null)
 		{
 			FilterDefinition<DeviceDocument> filter = Builders<DeviceDocument>.Filter.Empty;
 			if (deviceIds != null)
 			{
 				filter &= Builders<DeviceDocument>.Filter.In(x => x.Id, deviceIds);
 			}
+			if (poolId != null)
+			{
+				filter &= Builders<DeviceDocument>.Filter.Eq(x => x.PoolId, poolId);
+			}
+			if (platformId != null)
+			{
+				filter &= Builders<DeviceDocument>.Filter.Eq(x => x.PlatformId, platformId);
+			}
+
 			List<DeviceDocument> results = await _devices.Find(filter).ToListAsync();
 			return results.OrderBy(x => x.Name).Select<DeviceDocument, IDevice>(x => x).ToList();
 		}
@@ -407,9 +620,9 @@ namespace Horde.Build.Devices
 
 			List<UpdateDefinition<DeviceDocument>> updates = new List<UpdateDefinition<DeviceDocument>>();
 
-            string? userId = checkedOutByUserId?.ToString();
+			string? userId = checkedOutByUserId?.ToString();
 
-            updates.Add(updateBuilder.Set(x => x.CheckedOutByUser, String.IsNullOrEmpty(userId) ? null : userId));			
+			updates.Add(updateBuilder.Set(x => x.CheckedOutByUser, String.IsNullOrEmpty(userId) ? null : userId));
 
 			if (checkedOutByUserId != null)
 			{
@@ -467,7 +680,7 @@ namespace Horde.Build.Devices
 				else
 				{
 					updates.Add(updateBuilder.Set(x => x.ModelId, newModelId));
-				}				
+				}
 			}
 
 			if (newNotes != null)
@@ -501,6 +714,16 @@ namespace Horde.Build.Devices
 			{
 				await _devices.FindOneAndUpdateAsync<DeviceDocument>(x => x.Id == deviceId, updateBuilder.Combine(updates));
 			}
+
+			if (newProblem.HasValue && newProblem.Value)
+			{
+				IDeviceReservation? reservation = await TryGetDeviceReservationAsync(deviceId);
+				if (reservation != null)
+				{
+					UpdateDefinition<DeviceTelemetryDocument> update = Builders<DeviceTelemetryDocument>.Update.Set(x => x.ProblemTimeUtc, utcNow);
+					await _deviceTelemetry.FindOneAndUpdateAsync(t => t.ReservationId == reservation.Id && t.DeviceId == deviceId, update);
+				}
+			}
 		}
 
 		/// <inheritdoc/>
@@ -531,9 +754,9 @@ namespace Horde.Build.Devices
 
 			if (pool == null || pool.PoolType != DevicePoolType.Automation)
 			{
-                return null;
-            }
-			
+				return null;
+			}
+
 			HashSet<DeviceId> allocated = new HashSet<DeviceId>();
 			Dictionary<DeviceId, string> platformRequestMap = new Dictionary<DeviceId, string>();
 
@@ -553,6 +776,8 @@ namespace Horde.Build.Devices
 			// filter out currently reserved devices
 			poolDevices = poolDevices.FindAll(x => poolReservations.FirstOrDefault(p => p.Devices.Contains(x.Id)) == null);
 
+			int availablePoolDevices = poolDevices.Count;			
+
 			// sort to use last reserved first to cycle devices
 			poolDevices.Sort((a, b) =>
 			{
@@ -646,6 +871,18 @@ namespace Horde.Build.Devices
 			DeviceReservationDocument newReservation = new DeviceReservationDocument(ObjectId.GenerateNewId(), poolId, deviceIds, requestedPlatforms, reservationTimeUtc, hostname, reservationDetails, jobId, stepId);
 			await _reservations.InsertOneAsync(newReservation);
 
+			// Create device telemetry data for reservation
+			List<DeviceTelemetryDocument> telemetry = new List<DeviceTelemetryDocument>();
+			foreach (DeviceId deviceId in deviceIds)
+			{
+				telemetry.Add(new DeviceTelemetryDocument(deviceId, newReservation.Id, newReservation.CreateTimeUtc, jobId, stepId));
+			}
+
+			if (telemetry.Count > 0)
+			{
+				await _deviceTelemetry.InsertManyAsync(telemetry);
+			}
+
 			return newReservation;
 
 		}
@@ -660,6 +897,11 @@ namespace Horde.Build.Devices
 		/// <inheritdoc/>
 		public async Task<bool> DeleteReservationAsync(ObjectId id)
 		{
+			FilterDefinition<DeviceTelemetryDocument> telemetryFilter = Builders<DeviceTelemetryDocument>.Filter.Eq(x => x.ReservationId, id);
+
+			// update telemetry
+			await _deviceTelemetry.UpdateManyAsync(x => x.ReservationId == id, Builders<DeviceTelemetryDocument>.Update.Set(x => x.ReservationFinishUtc, DateTime.UtcNow));
+
 			FilterDefinition<DeviceReservationDocument> filter = Builders<DeviceReservationDocument>.Filter.Eq(x => x.Id, id);
 			DeleteResult result = await _reservations.DeleteOneAsync(filter);
 			return result.DeletedCount > 0;
@@ -676,14 +918,16 @@ namespace Horde.Build.Devices
 
 			reserves = reserves.FindAll(r => (utcNow - r.UpdateTimeUtc).TotalMinutes > 10).ToList();
 
-			if (reserves.Count > 0)
+			bool result = true;
+			foreach (IDeviceReservation reservation in reserves)
 			{
-				DeleteResult result = await _reservations.DeleteManyAsync(Builders<DeviceReservationDocument>.Filter.In(x => x.Id, reserves.Select(y => y.Id)));
-				return (result.DeletedCount > 0);
+				if (!await DeleteReservationAsync(reservation.Id))
+				{
+					result = false;
+				}
 			}
 
-			return false;
-
+			return result;
 		}
 
 		/// <summary>
@@ -715,7 +959,7 @@ namespace Horde.Build.Devices
 
 				if (result.ModifiedCount > 0)
 				{
-					return expiredDevices.Select(x => (UserId.Parse(x.CheckedOutByUser!), (IDevice) x)).ToList();
+					return expiredDevices.Select(x => (UserId.Parse(x.CheckedOutByUser!), (IDevice)x)).ToList();
 				}
 			}
 
@@ -774,6 +1018,148 @@ namespace Horde.Build.Devices
 		{
 			List<DeviceReservationDocument> results = await _reservations.Find(x => true).ToListAsync();
 			return results.FirstOrDefault(r => r.Devices.Contains(id));
-        }
+		}
+
+		/// <inheritdoc/>
+		public async Task<List<IDeviceTelemetry>> FindDeviceTelemetryAsync(DeviceId[]? deviceIds = null, DateTimeOffset? minCreateTime = null, DateTimeOffset? maxCreateTime = null)
+		{
+			FilterDefinitionBuilder<DeviceTelemetryDocument> filterBuilder = Builders<DeviceTelemetryDocument>.Filter;
+			FilterDefinition<DeviceTelemetryDocument> filter = filterBuilder.Empty;
+
+			if (deviceIds != null && deviceIds.Length > 0)
+			{
+				filter &= filterBuilder.In(x => x.DeviceId, deviceIds);
+			}
+
+			if (minCreateTime != null)
+			{
+				filter &= filterBuilder.Gte(x => x.CreateTimeUtc!, minCreateTime.Value.UtcDateTime);
+			}
+
+			if (maxCreateTime != null)
+			{
+				filter &= filterBuilder.Lte(x => x.CreateTimeUtc!, maxCreateTime.Value.UtcDateTime);
+			}
+
+			List<DeviceTelemetryDocument> results = await _deviceTelemetry.Find(filter).ToListAsync();	
+			return results.ConvertAll<IDeviceTelemetry>(x => x);
+		}
+
+		internal class DevicePoolTelemetryHelper
+		{
+			public int available = 0;
+			public int reserved = 0;
+			public int maintenance = 0;
+			public int problem = 0;
+			public int disabled = 0;
+
+		}
+
+		/// <summary>
+		/// Create a device pool telemetry snapshot
+		/// </summary>
+		/// <returns></returns>
+		public async Task CreatePoolTelemetrySnapshot()
+		{
+			List<IDevice> devices = await FindAllDevicesAsync();
+			List<IDevicePool> pools = await FindAllPoolsAsync();			
+			List<IDeviceReservation> reservations = await FindAllDeviceReservationsAsync();
+
+			// narrow to automation pools, may want to collect telemetry on other pools in the future
+			pools = pools.Where(x => x.PoolType == DevicePoolType.Automation).ToList();
+			devices = devices.Where(x => pools.FirstOrDefault(p => p.Id == x.PoolId) != null).ToList();
+
+			if (devices.Count == 0 || pools.Count == 0)
+			{
+				return;
+			}
+
+			DateTime now = DateTime.UtcNow;
+
+			List<IDevice> reservedDevices = devices.Where(x => reservations.FirstOrDefault(r => r.Devices.Contains(x.Id)) != null).ToList();
+			List<IDevice> maintenanceDevices = devices.Where(x => x.MaintenanceTimeUtc != null).ToList();
+			List<IDevice> disabledDevices  = devices.Where(x => !x.Enabled).ToList();
+			List<IDevice> problemDevices = devices.Where(x => (x.ProblemTimeUtc != null && ((now - x.ProblemTimeUtc).Value.TotalMinutes < 30))).ToList();
+
+			Dictionary<DevicePoolId, List<DevicePlatformTelemetryDocument>> poolTelemetry = new Dictionary<DevicePoolId, List<DevicePlatformTelemetryDocument>>();
+
+			foreach (IDevicePool pool in pools)
+			{								
+				List<IDevice> poolDevices = devices.Where(x => x.PoolId == pool.Id).ToList();
+
+				HashSet<DevicePlatformId> platforms = new HashSet<DevicePlatformId>();
+				poolDevices.ForEach(d => platforms.Add(d.PlatformId));
+
+				Dictionary<DevicePlatformId, DevicePoolTelemetryHelper> helpers = new Dictionary<DevicePlatformId, DevicePoolTelemetryHelper>();
+				foreach (DevicePlatformId platform in platforms)
+				{
+					helpers[platform] = new DevicePoolTelemetryHelper();
+				}
+				
+				foreach (IDevice device in poolDevices)
+				{
+					DevicePoolTelemetryHelper helper = helpers[device.PlatformId];
+
+					if (reservedDevices.Contains(device))
+					{
+						helper.reserved++;
+						continue;
+					}
+
+					if (problemDevices.Contains(device))
+					{
+						helper.problem++;
+						continue;
+					}
+
+					if (maintenanceDevices.Contains(device))
+					{
+						helper.maintenance++;
+						continue;
+					}
+
+					if (disabledDevices.Contains(device))
+					{
+						helper.disabled++;
+						continue;
+					}
+
+					helper.available++;
+				}
+
+				List<DevicePlatformTelemetryDocument> platformTelemtry = new List<DevicePlatformTelemetryDocument>();
+				foreach (KeyValuePair<DevicePlatformId, DevicePoolTelemetryHelper> platform in helpers)
+				{
+					DevicePoolTelemetryHelper helper = platform.Value;
+					platformTelemtry.Add(new DevicePlatformTelemetryDocument(platform.Key, helper.available, helper.reserved, helper.maintenance, helper.problem, helper.disabled));
+				}
+
+				poolTelemetry[pool.Id] = platformTelemtry;
+			}
+
+			await _poolTelemetry.InsertOneAsync(new DevicePoolTelemetryDocument(poolTelemetry));
+		}
+
+		/// <inheritdoc/>
+		public async Task<List<IDevicePoolTelemetry>> FindPoolTelemetryAsync(DateTimeOffset? minCreateTime = null, DateTimeOffset? maxCreateTime = null)
+		{
+			FilterDefinitionBuilder<DevicePoolTelemetryDocument> filterBuilder = Builders<DevicePoolTelemetryDocument>.Filter;
+			FilterDefinition<DevicePoolTelemetryDocument> filter = filterBuilder.Empty;
+
+			if (minCreateTime != null)
+			{
+				filter &= filterBuilder.Gte(x => x.CreateTimeUtc!, minCreateTime.Value.UtcDateTime);
+			}
+
+			if (maxCreateTime != null)
+			{
+				filter &= filterBuilder.Lte(x => x.CreateTimeUtc!, maxCreateTime.Value.UtcDateTime);
+			}
+
+			List<DevicePoolTelemetryDocument> results = await _poolTelemetry.Find(filter).ToListAsync();
+			return results.ConvertAll<IDevicePoolTelemetry>(x => x);
+		}
+
+
 	}
 }
diff --git a/Engine/Source/Programs/Horde/Horde.Build/Devices/DeviceMessages.cs b/Engine/Source/Programs/Horde/Horde.Build/Devices/DeviceMessages.cs
--- a/Engine/Source/Programs/Horde/Horde.Build/Devices/DeviceMessages.cs
+++ b/Engine/Source/Programs/Horde/Horde.Build/Devices/DeviceMessages.cs
@@ -240,7 +240,7 @@ namespace Horde.Build.Devices
 	}
 
 	/// <summary>
-	/// Get response object which describes a device
+	/// Get response object which describes a device (DEPRECATED)
 	/// </summary>
 	public class GetDeviceUtilizationResponse
 	{
@@ -546,6 +546,164 @@ namespace Horde.Build.Devices
 		public string LegacyGuid { get; set; } = null!;
 	}
 
+	/// <summary>
+	/// Device telemetry respponse
+	/// </summary>
+	public class GetTelemetryInfoResponse
+	{
+		/// <summary>
+		/// The UTC time the telemetry data was created
+		/// </summary>
+		[Required]
+		public DateTime CreateTimeUtc { get; set; }
+
+		/// <summary>
+		/// The job id which utilized device
+		/// </summary>
+		public string? JobId { get; set; }
+
+		/// <summary>
+		/// The job's step id
+		/// </summary>
+		public string? StepId { get; set; }
+
+		/// <summary>
+		/// If this telemetry has a reservation, the start time of the reservation
+		/// </summary>
+		public DateTime? ReservationStartUtc { get; set; }
+
+		/// <summary>
+		/// If this telemetry has a reservation, the finish time of the reservation
+		/// </summary>
+		public DateTime? ReservationFinishUtc { get; set; }
+
+		/// <summary>
+		/// If this telemetry marks a detected device issue, the time of the issue
+		/// </summary>
+		public DateTime? ProblemTimeUtc { get; set; }
+
+		/// <summary>
+		/// Constructor
+		/// </summary>
+		/// <param name="data"></param>
+		public GetTelemetryInfoResponse(IDeviceTelemetry data)
+		{
+			CreateTimeUtc = data.CreateTimeUtc;
+			JobId = data.JobId;
+			StepId = data.StepId;
+			ReservationStartUtc = data.ReservationStartUtc;
+			ReservationFinishUtc = data.ReservationFinishUtc;
+			ProblemTimeUtc = data.ProblemTimeUtc;
+		}
+
+	}
+
+	/// <summary>
+	/// Device telemetry respponse
+	/// </summary>
+	public class GetDeviceTelemetryResponse
+	{
+		/// <summary>
+		/// The device id for the telemetry data
+		/// </summary>
+		[Required]
+		public string DeviceId { get; set; } = null!;
+
+		/// <summary>
+		/// Individual telemetry data points
+		/// </summary>
+		[Required]
+		public List<GetTelemetryInfoResponse> Telemetry { get; set; } = null!;
+
+		/// <summary>
+		/// Constructor
+		/// </summary>
+		/// <param name="deviceId"></param>
+		/// <param name="telemetry"></param>
+		public GetDeviceTelemetryResponse(string deviceId, List<GetTelemetryInfoResponse> telemetry)
+		{
+			DeviceId = deviceId;
+			Telemetry = telemetry;
+		}
+	}
+
+
+	/// <summary>
+	/// Device telemetry respponse
+	/// </summary>
+	public class GetDevicePlatformTelemetryResponse
+	{
+		/// <summary>
+		/// The corresponding platform id
+		/// </summary>
+		public string PlatformId { get; set; }
+
+		/// <summary>
+		/// Number of available devices of this platform 
+		/// </summary>
+		public int Available { get; set; }
+
+		/// <summary>
+		/// Number of reserved devices of this platform 
+		/// </summary>
+		public int Reserved { get; set; }
+
+		/// <summary>
+		/// Number of devices in maintenance state
+		/// </summary>
+		public int Maintenance { get; set; }
+
+		/// <summary>
+		/// Number of devices in problem state
+		/// </summary>
+		public int Problem { get; set; }
+
+		/// <summary>
+		/// Number of devices in disabled state
+		/// </summary>
+		public int Disabled { get; set; }
+
+		/// <summary>
+		/// Constructor
+		/// </summary>
+		public GetDevicePlatformTelemetryResponse(string platformId, int available, int reserved, int maintenance, int problem, int disabled)
+		{
+			PlatformId = platformId;
+			Available = available;
+			Reserved = reserved;
+			Maintenance = maintenance;
+			Problem = problem;
+			Disabled = disabled;
+		}
+	}
+
+	/// <summary>
+	/// Device telemetry respponse
+	/// </summary>
+	public class GetDevicePoolTelemetryResponse
+	{
+		/// <summary>
+		/// The UTC time the telemetry data was created
+		/// </summary>
+		public DateTime CreateTimeUtc { get; set; }
+
+		/// <summary>
+		/// Individual pool telemetry data points
+		/// </summary>
+		public Dictionary<string, List<GetDevicePlatformTelemetryResponse>> Telemetry { get; set; } = null!;
+
+		/// <summary>
+		/// Constructor
+		/// </summary>
+		public GetDevicePoolTelemetryResponse(DateTime createTimeUtc, Dictionary<string, List<GetDevicePlatformTelemetryResponse>> telemetry)
+		{
+			CreateTimeUtc = createTimeUtc;
+			Telemetry = telemetry;
+		}
+
+	}
+
+
 	// Legacy clients
 
 	/// <summary>
diff --git a/Engine/Source/Programs/Horde/Horde.Build/Devices/DeviceService.cs b/Engine/Source/Programs/Horde/Horde.Build/Devices/DeviceService.cs
--- a/Engine/Source/Programs/Horde/Horde.Build/Devices/DeviceService.cs
+++ b/Engine/Source/Programs/Horde/Horde.Build/Devices/DeviceService.cs
@@ -130,7 +130,9 @@ namespace Horde.Build.Devices
 		/// Device collection
 		/// </summary>
 		readonly IDeviceCollection _devices;
-		readonly ITicker _ticker;
+		readonly ITicker _ticker;		
+		static readonly int _poolTelemetryMinutes = 15;
+		int _poolTelemetryTick = _poolTelemetryMinutes;
 
 		/// <summary>
 		/// Platform map V1 singleton
@@ -173,28 +175,70 @@ namespace Horde.Build.Devices
 
 			if (!stoppingToken.IsCancellationRequested)
 			{
-				await _devices.ExpireReservationsAsync();
+				try
+				{
+					await _devices.ExpireReservationsAsync();
+				}
+				catch (Exception ex)
+				{
+					_logger.LogError(ex, "Exception while expiring reservations: {Message}", ex.Message);
+				}
+
 
-				List<(UserId, IDevice)>? expireNotifications = await _devices.ExpireNotificatonsAsync();
-				if (expireNotifications != null && expireNotifications.Count > 0)
+				// expire shared devices and notifications
+				try
 				{
-					foreach ((UserId, IDevice) expiredDevice in expireNotifications)
+					List<(UserId, IDevice)>? expireNotifications = await _devices.ExpireNotificatonsAsync();
+					if (expireNotifications != null && expireNotifications.Count > 0)
+					{
+						foreach ((UserId, IDevice) expiredDevice in expireNotifications)
+						{
+							await NotifyDeviceServiceAsync($"Device {expiredDevice.Item2.PlatformId.ToString().ToUpperInvariant()} / {expiredDevice.Item2.Name} checkout will expire in 24 hours.  Please visit https://horde.devtools.epicgames.com/devices to renew the checkout if needed.", null, null, null, expiredDevice.Item1);
+						}
+					}
+
+					List<(UserId, IDevice)>? expireCheckouts = await _devices.ExpireCheckedOutAsync();
+					if (expireCheckouts != null && expireCheckouts.Count > 0)
 					{
-						await NotifyDeviceServiceAsync($"Device {expiredDevice.Item2.PlatformId.ToString().ToUpperInvariant()} / {expiredDevice.Item2.Name} checkout will expire in 24 hours.  Please visit https://horde.devtools.epicgames.com/devices to renew the checkout if needed.", null, null, null, expiredDevice.Item1);
+						foreach ((UserId, IDevice) expiredDevice in expireCheckouts)
+						{
+							await NotifyDeviceServiceAsync($"Device {expiredDevice.Item2.PlatformId.ToString().ToUpperInvariant()} / {expiredDevice.Item2.Name} checkout has expired.  The device has been returned to the shared pool and should no longer be accessed.  Please visit https://horde.devtools.epicgames.com/devices to checkout devices as needed.", null, null, null, expiredDevice.Item1);
+						}
 					}
 				}
+				catch (Exception ex)
+				{
+					_logger.LogError(ex, "Exception while expiring device checkouts: {Message}", ex.Message);
+				}
 
-				List<(UserId, IDevice)>? expireCheckouts = await _devices.ExpireCheckedOutAsync();
-				if (expireCheckouts != null && expireCheckouts.Count > 0)
+				// update pool telemetry
+				try
 				{
-					foreach ((UserId, IDevice) expiredDevice in expireCheckouts)
+
+					if (_poolTelemetryTick >= _poolTelemetryMinutes)
+					{
+						_poolTelemetryTick = 0;
+						await _devices.CreatePoolTelemetrySnapshot();						
+					}
+					else
 					{
-						await NotifyDeviceServiceAsync($"Device {expiredDevice.Item2.PlatformId.ToString().ToUpperInvariant()} / {expiredDevice.Item2.Name} checkout has expired.  The device has been returned to the shared pool and should no longer be accessed.  Please visit https://horde.devtools.epicgames.com/devices to checkout devices as needed.", null, null, null, expiredDevice.Item1);
+						_poolTelemetryTick++;
 					}
+
+				}
+				catch (Exception ex)
+				{
+					_logger.LogError(ex, "Exception while generating pool telemetry: {Message}", ex.Message);
 				}
 			}
 		}
 
+		internal async Task TickForTestingAsync()
+		{
+			await TickAsync(CancellationToken.None);
+		}
+
+
 		/// <summary>
 		/// Create a new device platform
 		/// </summary>
@@ -262,11 +306,28 @@ namespace Horde.Build.Devices
 		/// <summary>
 		/// Get a list of devices, optionally filtered to provided ids
 		/// </summary>
-		public Task<List<IDevice>> GetDevicesAsync(List<DeviceId>? deviceIds = null)
+		public Task<List<IDevice>> GetDevicesAsync(List<DeviceId>? deviceIds = null, DevicePoolId? poolId = null, DevicePlatformId? platformId = null)
 		{
-			return _devices.FindAllDevicesAsync(deviceIds);
+			return _devices.FindAllDevicesAsync(deviceIds, poolId, platformId);
 		}
 
+		/// <summary>
+		/// Get device telemetry
+		/// </summary>
+		public Task<List<IDeviceTelemetry>> GetDeviceTelemetryAsync(DeviceId[]? deviceIds = null, DateTimeOffset? minCreateTime=null, DateTimeOffset? maxCreateTime=null)
+		{
+			return _devices.FindDeviceTelemetryAsync(deviceIds, minCreateTime, maxCreateTime);
+		}
+
+		/// <summary>
+		/// Get device pool telemetry
+		/// </summary>
+		public Task<List<IDevicePoolTelemetry>> GetDevicePoolTelemetryAsync(DateTimeOffset? minCreateTime = null, DateTimeOffset? maxCreateTime = null)
+		{
+			return _devices.FindPoolTelemetryAsync(minCreateTime, maxCreateTime);
+		}
+
+
 		/// <summary>
 		/// Get a specific device
 		/// </summary>
diff --git a/Engine/Source/Programs/Horde/Horde.Build/Devices/DevicesController.cs b/Engine/Source/Programs/Horde/Horde.Build/Devices/DevicesController.cs
--- a/Engine/Source/Programs/Horde/Horde.Build/Devices/DevicesController.cs
+++ b/Engine/Source/Programs/Horde/Horde.Build/Devices/DevicesController.cs
@@ -489,6 +489,49 @@ namespace Horde.Build.Devices
 			return responses;
 		}
 
+		/// <summary>
+		/// Get device telemetry
+		/// </summary>		
+		[HttpGet]
+		[Authorize]
+		[Route("/api/v2/devices/pools/telemetry")]
+		[ProducesResponseType(typeof(List<GetDevicePoolTelemetryResponse>), 200)]
+		public async Task<ActionResult<List<GetDevicePoolTelemetryResponse>>> GetDevicePoolTelemetry(
+			[FromQuery] DateTimeOffset? minCreateTime = null,
+			[FromQuery] DateTimeOffset? maxCreateTime = null)
+		{
+			if (!await _deviceService.AuthorizeAsync(AclAction.DeviceRead, User))
+			{
+				return Forbid();
+			}
+
+			List<IDevicePoolTelemetry> telemetryData = await _deviceService.GetDevicePoolTelemetryAsync(minCreateTime, maxCreateTime);
+
+			List<GetDevicePoolTelemetryResponse> response = new List<GetDevicePoolTelemetryResponse>();
+						
+			foreach (IDevicePoolTelemetry t in telemetryData)
+			{
+				Dictionary<string, List<GetDevicePlatformTelemetryResponse>> poolData = new Dictionary<string, List<GetDevicePlatformTelemetryResponse>>();
+
+				foreach ( KeyValuePair<DevicePoolId, IReadOnlyList<IDevicePlatformTelemetry>> pool in t.Pools)
+				{
+					List<GetDevicePlatformTelemetryResponse> platformTelemetry = new List<GetDevicePlatformTelemetryResponse>();
+
+					foreach (IDevicePlatformTelemetry telemetry in pool.Value)
+					{
+						platformTelemetry.Add(new GetDevicePlatformTelemetryResponse(telemetry.PlatformId.ToString(), telemetry.Available, telemetry.Reserved, telemetry.Maintenance, telemetry.Problem, telemetry.Disabled));
+					}
+
+					poolData[pool.Key.ToString()] = platformTelemetry;
+				}
+
+				response.Add(new GetDevicePoolTelemetryResponse(t.CreateTimeUtc, poolData));
+			}		
+
+			return response;
+		}
+
+
 		// RESERVATIONS
 
 		/// <summary>
@@ -665,6 +708,73 @@ namespace Horde.Build.Devices
 			return Ok();
 		}
 
+		/// <summary>
+		/// Get device telemetry
+		/// </summary>		
+		[HttpGet]
+		[Authorize]
+		[Route("/api/v2/devices/telemetry")]
+		[ProducesResponseType(typeof(List<GetDeviceTelemetryResponse>), 200)]
+		public async Task<ActionResult<List<GetDeviceTelemetryResponse>>> GetDeviceTelemetry(
+			[FromQuery(Name = "Id")] string[]? deviceIds = null,
+			[FromQuery] string? poolId = null,
+			[FromQuery] string? platformId = null,
+			[FromQuery] DateTimeOffset? minCreateTime = null,
+			[FromQuery] DateTimeOffset? maxCreateTime = null)
+		{
+			if (!await _deviceService.AuthorizeAsync(AclAction.DeviceRead, User))
+			{
+				return Forbid();
+			}
+
+			HashSet<DeviceId> deviceIdValues = new HashSet<DeviceId>();
+
+			if (deviceIds != null && deviceIds.Length > 0)
+			{
+				foreach(string deviceId in deviceIds)
+				{
+					deviceIdValues.Add(new DeviceId(deviceId));
+				}
+			}
+
+			DevicePoolId? poolIdValue = poolId == null ? null : new DevicePoolId(poolId);
+			DevicePlatformId? platformIdValue = platformId == null ? null : new DevicePlatformId(platformId);
+
+			List<IDevice> devices = await _deviceService.GetDevicesAsync(deviceIdValues.Count > 0 ? deviceIdValues.ToList() : null, poolIdValue, platformIdValue);
+
+			List<GetDeviceTelemetryResponse> response = new List<GetDeviceTelemetryResponse>();
+
+			if (devices.Count == 0)
+			{
+				return response;
+			}
+			
+			List<IDeviceTelemetry> telemetryData = await _deviceService.GetDeviceTelemetryAsync(devices.Select(x => x.Id).ToArray(), minCreateTime, maxCreateTime);
+
+			Dictionary<string, List<GetTelemetryInfoResponse>> results = new Dictionary<string, List<GetTelemetryInfoResponse>>();
+
+			foreach (IDeviceTelemetry t in telemetryData)
+			{
+				List<GetTelemetryInfoResponse>? info = null;
+				string deviceId = t.DeviceId.ToString();
+				if (!results.TryGetValue(deviceId, out info))
+				{
+					info = new List<GetTelemetryInfoResponse>();
+					results[deviceId] = info;
+				}
+
+				info.Add(new GetTelemetryInfoResponse(t));
+			}
+
+			foreach(KeyValuePair<string, List<GetTelemetryInfoResponse >> result in results)
+			{
+				GetDeviceTelemetryResponse deviceTelemetry = new GetDeviceTelemetryResponse(result.Key, result.Value);
+				response.Add(deviceTelemetry);
+			}
+
+			return response;
+		}
+
 		// LEGACY V1 API
 
 		enum LegacyPerfSpec
diff --git a/Engine/Source/Programs/Horde/Horde.Build/Devices/IDevice.cs b/Engine/Source/Programs/Horde/Horde.Build/Devices/IDevice.cs
--- a/Engine/Source/Programs/Horde/Horde.Build/Devices/IDevice.cs
+++ b/Engine/Source/Programs/Horde/Horde.Build/Devices/IDevice.cs
@@ -152,7 +152,7 @@ namespace Horde.Build.Devices
 	}
 
 	/// <summary>
-	/// A device utilization snapshot
+	/// A device utilization snapshot [DEPRECATED]
 	/// </summary>
 	public class DeviceUtilizationTelemetry
 	{
@@ -194,6 +194,54 @@ namespace Horde.Build.Devices
 		}
 	}
 
+	/// <summary>
+	/// Device telemetry information
+	/// </summary>
+	public interface IDeviceTelemetry
+	{
+		/// <summary>
+		///  The creation time of this telemetry data
+		/// </summary>
+		public DateTime CreateTimeUtc { get; }
+
+		/// <summary>
+		/// The device id for telemetry
+		/// </summary>
+		public DeviceId DeviceId { get; }
+
+		/// <summary>
+		/// The job id which utilized device
+		/// </summary>
+		public string? JobId { get; }
+
+		/// <summary>
+		/// The job's step id
+		/// </summary>
+		public string? StepId { get; }
+
+		/// <summary>
+		/// Reservation Id (transient, reservations are deleted upon expiration)
+		/// </summary>
+		public ObjectId? ReservationId { get; }
+
+		/// <summary>
+		/// The time device was reserved
+		/// </summary>
+		public DateTime? ReservationStartUtc { get; }
+
+		/// <summary>
+		/// The time device was freed
+		/// </summary>
+		public DateTime? ReservationFinishUtc { get; }
+
+		/// <summary>
+		/// If the device reported a problem
+		/// </summary>
+		public DateTime? ProblemTimeUtc { get; }
+
+	}
+
+
 	/// <summary>
 	/// A physical device
 	/// </summary>
@@ -274,4 +322,59 @@ namespace Horde.Build.Devices
 		/// </summary>
 		public Acl? Acl { get; }
 	}
+
+	/// <summary>
+	/// Platform telemetry for a device pool
+	/// </summary>
+	public interface IDevicePlatformTelemetry
+	{
+		/// <summary>
+		/// The corresponding platform id
+		/// </summary>
+		public DevicePlatformId PlatformId { get; }
+
+		/// <summary>
+		/// Number of available devices of this platform 
+		/// </summary>
+		public int Available { get; }
+
+		/// <summary>
+		/// Number of reserved devices of this platform 
+		/// </summary>
+		public int Reserved { get; }
+
+		/// <summary>
+		/// Number of devices in maintenance state
+		/// </summary>
+		public int Maintenance { get; }
+
+		/// <summary>
+		/// Number of devices in problem state
+		/// </summary>
+		public int Problem { get; }
+
+		/// <summary>
+		/// Number of devices in disabled state
+		/// </summary>
+		public int Disabled { get; }
+
+	}
+
+	/// <summary>
+	/// A snapshot of device pool telemetry
+	/// </summary>
+	public interface IDevicePoolTelemetry
+	{
+		/// <summary>
+		/// The creation time of this telemetry data
+		/// </summary>
+		public DateTime CreateTimeUtc { get; }
+
+		/// <summary>
+		/// Pool platform telemetry
+		/// </summary>
+		public IReadOnlyDictionary<DevicePoolId, IReadOnlyList<IDevicePlatformTelemetry>> Pools { get; }
+
+	}
+
 }
diff --git a/Engine/Source/Programs/Horde/Horde.Build/Devices/IDeviceCollection.cs b/Engine/Source/Programs/Horde/Horde.Build/Devices/IDeviceCollection.cs
--- a/Engine/Source/Programs/Horde/Horde.Build/Devices/IDeviceCollection.cs
+++ b/Engine/Source/Programs/Horde/Horde.Build/Devices/IDeviceCollection.cs
@@ -1,5 +1,6 @@
 // Copyright Epic Games, Inc. All Rights Reserved.
 
+using System;
 using System.Collections.Generic;
 using System.Threading.Tasks;
 using Horde.Build.Projects;
@@ -127,7 +128,9 @@ namespace Horde.Build.Devices
 		/// Get a list of all devices
 		/// </summary>
 		/// <param name="deviceIds">Optional list of device ids to get</param>
-		Task<List<IDevice>> FindAllDevicesAsync(List<DeviceId>? deviceIds);
+		/// <param name="poolId">Optional id of pool for devices</param>
+		/// <param name="platformId">Optional id of platform for devices</param>
+		Task<List<IDevice>> FindAllDevicesAsync(List<DeviceId>? deviceIds, DevicePoolId? poolId = null, DevicePlatformId? platformId = null);
 
 		/// <summary>
 		/// Adds a new device
@@ -227,5 +230,30 @@ namespace Horde.Build.Devices
 		/// Gets a list of users to notify whose device checkout is about to expire
 		/// </summary>
 		public Task<List<(UserId, IDevice)>?> ExpireNotificatonsAsync();
+
+		/// <summary>
+		/// Find device telemetry
+		/// </summary>
+		/// <param name="deviceIds"></param>
+		/// <param name="minCreateTime"></param>
+		/// <param name="maxCreateTime"></param>
+		/// <returns></returns>
+		public Task<List<IDeviceTelemetry>> FindDeviceTelemetryAsync(DeviceId[]? deviceIds = null, DateTimeOffset? minCreateTime = null, DateTimeOffset? maxCreateTime = null);
+
+		/// <summary>
+		/// Creates a device pool telemetry snapshot
+		/// </summary>
+		/// <returns></returns>
+		public Task CreatePoolTelemetrySnapshot();
+
+		/// <summary>
+		/// Gets pool telemetry for an optional date range
+		/// </summary>
+		/// <param name="minCreateTime"></param>
+		/// <param name="maxCreateTime"></param>
+		/// <returns></returns>
+		public Task<List<IDevicePoolTelemetry>> FindPoolTelemetryAsync(DateTimeOffset? minCreateTime = null, DateTimeOffset? maxCreateTime = null);
+
+
 	}
 }
\ No newline at end of file
