Horde: Set min create time for telemetry if not specified in query
#rnx

diff --git a/Engine/Source/Programs/Horde/Horde.Build/Devices/DevicesController.cs b/Engine/Source/Programs/Horde/Horde.Build/Devices/DevicesController.cs
--- a/Engine/Source/Programs/Horde/Horde.Build/Devices/DevicesController.cs
+++ b/Engine/Source/Programs/Horde/Horde.Build/Devices/DevicesController.cs
@@ -505,6 +505,11 @@ namespace Horde.Build.Devices
 				return Forbid();
 			}
 
+			if (minCreateTime == null)
+			{
+				minCreateTime = new DateTimeOffset(DateTime.UtcNow - TimeSpan.FromHours(24));
+			}
+
 			List<IDevicePoolTelemetry> telemetryData = await _deviceService.GetDevicePoolTelemetryAsync(minCreateTime, maxCreateTime);
 
 			List<GetDevicePoolTelemetryResponse> response = new List<GetDevicePoolTelemetryResponse>();
@@ -727,6 +732,11 @@ namespace Horde.Build.Devices
 				return Forbid();
 			}
 
+			if (minCreateTime == null)
+			{
+				minCreateTime = new DateTimeOffset(DateTime.UtcNow - TimeSpan.FromHours(24));
+			}
+
 			HashSet<DeviceId> deviceIdValues = new HashSet<DeviceId>();
 
 			if (deviceIds != null && deviceIds.Length > 0)
