Fix NuGet package version mismatches.

diff --git a/Engine/Source/Programs/Horde/Horde.Agent.Tests/Horde.Agent.Tests.csproj b/Engine/Source/Programs/Horde/Horde.Agent.Tests/Horde.Agent.Tests.csproj
--- a/Engine/Source/Programs/Horde/Horde.Agent.Tests/Horde.Agent.Tests.csproj
+++ b/Engine/Source/Programs/Horde/Horde.Agent.Tests/Horde.Agent.Tests.csproj
@@ -10,7 +10,7 @@
 
   <ItemGroup>
 
-    <PackageReference Include="Microsoft.Extensions.Logging.Abstractions" Version="5.0.0" />
+    <PackageReference Include="Microsoft.Extensions.Logging.Abstractions" Version="6.0.0" />
     <PackageReference Include="Microsoft.NET.Test.Sdk" Version="16.5.0" />
     <PackageReference Include="Moq" Version="4.14.5" />
     <PackageReference Include="MSTest.TestAdapter" Version="2.1.0" />
diff --git a/Engine/Source/Programs/Horde/Horde.Agent.Tests/JsonRpcLoggerTests.cs b/Engine/Source/Programs/Horde/Horde.Agent.Tests/JsonRpcLoggerTests.cs
--- a/Engine/Source/Programs/Horde/Horde.Agent.Tests/JsonRpcLoggerTests.cs
+++ b/Engine/Source/Programs/Horde/Horde.Agent.Tests/JsonRpcLoggerTests.cs
@@ -23,7 +23,7 @@ namespace Horde.Agent.Tests
 
 			public bool IsEnabled(LogLevel logLevel) => true;
 
-			public void Log<TState>(LogLevel logLevel, EventId eventId, TState state, Exception exception, Func<TState, Exception?, string> formatter)
+			public void Log<TState>(LogLevel logLevel, EventId eventId, TState state, Exception? exception, Func<TState, Exception?, string> formatter)
 			{
 				JsonLogEvent logEvent = JsonLogEvent.FromLoggerState(logLevel, eventId, state, exception, formatter);
 				_lines.Add(LogEvent.Read(logEvent.Data.Span));
diff --git a/Engine/Source/Programs/Horde/Horde.Agent.Tests/LogParserTests.cs b/Engine/Source/Programs/Horde/Horde.Agent.Tests/LogParserTests.cs
--- a/Engine/Source/Programs/Horde/Horde.Agent.Tests/LogParserTests.cs
+++ b/Engine/Source/Programs/Horde/Horde.Agent.Tests/LogParserTests.cs
@@ -27,11 +27,11 @@ namespace Horde.Agent.Tests
 
 			public List<LogEvent> _events = new List<LogEvent>();
 
-			public IDisposable? BeginScope<TState>(TState state) => null;
+			public IDisposable BeginScope<TState>(TState state) => null!;
 
 			public bool IsEnabled(LogLevel logLevel) => true;
 
-			public void Log<TState>(LogLevel logLevel, EventId eventId, TState state, Exception exception, Func<TState, Exception?, string> formatter)
+			public void Log<TState>(LogLevel logLevel, EventId eventId, TState state, Exception? exception, Func<TState, Exception?, string> formatter)
 			{
 				LogEvent logEvent = LogEvent.Read(JsonLogEvent.FromLoggerState(logLevel, eventId, state, exception, formatter).Data.Span);
 				if (logEvent.Level != LogLevel.Information || logEvent.Id != default || logEvent.Properties != null)
diff --git a/Engine/Source/Programs/Horde/Horde.Agent/Horde.Agent.csproj b/Engine/Source/Programs/Horde/Horde.Agent/Horde.Agent.csproj
--- a/Engine/Source/Programs/Horde/Horde.Agent/Horde.Agent.csproj
+++ b/Engine/Source/Programs/Horde/Horde.Agent/Horde.Agent.csproj
@@ -32,13 +32,13 @@
       <PrivateAssets>all</PrivateAssets>
       <IncludeAssets>runtime; build; native; contentfiles; analyzers; buildtransitive</IncludeAssets>
     </PackageReference>
-    <PackageReference Include="Microsoft.Extensions.Caching.Memory" Version="5.0.0" />
+    <PackageReference Include="Microsoft.Extensions.Caching.Memory" Version="6.0.0" />
     <PackageReference Include="Microsoft.Extensions.Hosting" Version="3.1.5" />
     <PackageReference Include="Microsoft.Extensions.Hosting.Abstractions" Version="3.1.5" />
     <PackageReference Include="Microsoft.Extensions.Hosting.WindowsServices" Version="3.1.5" />
     <PackageReference Include="Microsoft.Extensions.Http" Version="3.1.22" />
     <PackageReference Include="Microsoft.Extensions.Http.Polly" Version="3.1.22" />
-    <PackageReference Include="Microsoft.Extensions.Logging.Abstractions" Version="5.0.0" />
+    <PackageReference Include="Microsoft.Extensions.Logging.Abstractions" Version="6.0.0" />
     <PackageReference Include="Microsoft.Extensions.Options.DataAnnotations" Version="3.1.5" />
     <PackageReference Include="OpenTracing" Version="0.12.1" />
     <PackageReference Include="Serilog" Version="2.9.0" />
diff --git a/Engine/Source/Programs/Horde/Horde.Agent/Utility/JsonRpcLogger.cs b/Engine/Source/Programs/Horde/Horde.Agent/Utility/JsonRpcLogger.cs
--- a/Engine/Source/Programs/Horde/Horde.Agent/Utility/JsonRpcLogger.cs
+++ b/Engine/Source/Programs/Horde/Horde.Agent/Utility/JsonRpcLogger.cs
@@ -146,7 +146,7 @@ namespace Horde.Agent.Parser
 		}
 
 		/// <inheritdoc/>
-		public void Log<TState>(LogLevel logLevel, EventId eventId, TState state, Exception exception, Func<TState, Exception?, string> formatter)
+		public void Log<TState>(LogLevel logLevel, EventId eventId, TState state, Exception? exception, Func<TState, Exception?, string> formatter)
 		{
 			// Downgrade warnings to information if not required
 			if (logLevel == LogLevel.Warning && !_warnings)
diff --git a/Engine/Source/Programs/Horde/Horde.Agent/Utility/PerforceLogger.cs b/Engine/Source/Programs/Horde/Horde.Agent/Utility/PerforceLogger.cs
--- a/Engine/Source/Programs/Horde/Horde.Agent/Utility/PerforceLogger.cs
+++ b/Engine/Source/Programs/Horde/Horde.Agent/Utility/PerforceLogger.cs
@@ -64,7 +64,7 @@ namespace Horde.Agent.Utility
 		}
 
 		/// <inheritdoc/>
-		public void Log<TState>(LogLevel logLevel, EventId eventId, TState state, Exception exception, Func<TState, Exception?, string> formatter)
+		public void Log<TState>(LogLevel logLevel, EventId eventId, TState state, Exception? exception, Func<TState, Exception?, string> formatter)
 		{
 			JsonLogEvent logEvent = JsonLogEvent.FromLoggerState(logLevel, eventId, state, exception, formatter);
 			logEvent = new JsonLogEvent(logLevel, eventId, logEvent.LineIndex, logEvent.LineCount, Annotate(logEvent.Data));
