Horde: Rename fleet manager source files

diff --git a/Engine/Source/Programs/Horde/Horde.Build/Agents/Fleet/Providers/AwsFleetManager.cs b/Engine/Source/Programs/Horde/Horde.Build/Agents/Fleet/Providers/AwsFleetManager.cs
new file mode 100644
--- /dev/null
+++ b/Engine/Source/Programs/Horde/Horde.Build/Agents/Fleet/Providers/AwsFleetManager.cs
@@ -0,0 +1,280 @@
+// Copyright Epic Games, Inc. All Rights Reserved.
+
+using System;
+using System.Collections.Generic;
+using System.Linq;
+using System.Threading;
+using System.Threading.Tasks;
+using Amazon;
+using Amazon.EC2;
+using Amazon.EC2.Model;
+using Amazon.Extensions.NETCore.Setup;
+using Horde.Build.Agents.Pools;
+using Horde.Build.Auditing;
+using Microsoft.Extensions.Logging;
+using OpenTracing;
+using OpenTracing.Util;
+
+namespace Horde.Build.Agents.Fleet.Providers
+{
+	/// <summary>
+	/// Settings for launching an EC2 instance
+	/// </summary>
+	public class AwsInstanceLaunchSettings
+	{
+		/// <summary>
+		/// AWS region (e.g us-east-1)
+		/// </summary>
+		public string Region { get; }
+		
+		/// <summary>
+		/// Amazon Machine Image (AMI) to use (e.g ami-0c2e245e4936a2ab2)
+		/// </summary>
+		public string ImageId { get; }
+		
+		/// <summary>
+		/// An EC2 instance type (see https://docs.aws.amazon.com/AWSEC2/latest/UserGuide/instance-types.html for full list)
+		/// </summary>
+		public string InstanceType { get; }
+		
+		/// <summary>
+		/// IAM	instance profile 
+		/// </summary>
+		public string? IamInstanceProfile { get; }
+		
+		/// <summary>
+		/// Key/value tags for the instance
+		/// </summary>
+		public IReadOnlyDictionary<string, string> Tags { get; }
+		
+		/// <summary>
+		/// List of subnet IDs this instance can run in. Will implicitly pick the VPC.
+		/// A random subnet ID will be picked during instance creation.
+		/// </summary>
+		public List<string> SubnetIds { get; }
+		
+		/// <summary>
+		/// Security group IDs to apply. Must exist in the same VPC as the subnets.
+		/// </summary>
+		public List<string> SecurityGroupIds { get; }
+		
+		/// <summary>
+		/// User data to pass to the instance. Commonly used for cloud-init scripts that configures the instance during boot up. 
+		/// </summary>
+		public string? UserData { get; }
+		
+		/// <summary>
+		/// Key name for the instance (optional).
+		/// If skipped, alternative ways to administer the instance must be enabled, such as the AWS SSM Agent. 
+		/// </summary>
+		public string? KeyName { get; }
+
+		/// <summary>
+		/// Constructor
+		/// </summary>
+		/// <param name="region"></param>
+		/// <param name="imageId"></param>
+		/// <param name="instanceType"></param>
+		/// <param name="iamInstanceProfile"></param>
+		/// <param name="tags"></param>
+		/// <param name="subnetIds"></param>
+		/// <param name="securityGroupIds"></param>
+		/// <param name="userData"></param>
+		/// <param name="keyName"></param>
+		public AwsInstanceLaunchSettings(string region, string imageId, string instanceType, string? iamInstanceProfile, IReadOnlyDictionary<string, string> tags,
+			List<string> subnetIds, List<string> securityGroupIds, string? userData, string? keyName)
+		{
+			Region = region;
+			ImageId = imageId;
+			InstanceType = instanceType;
+			IamInstanceProfile = iamInstanceProfile;
+			Tags = tags;
+			SubnetIds = subnetIds;
+			SecurityGroupIds = securityGroupIds;
+			UserData = userData;
+			KeyName = keyName;
+		}
+
+		/// <summary>
+		/// Get a random subnet ID. Used for 
+		/// </summary>
+		/// <returns>A random subnet ID from the list of subnets</returns>
+		public string GetRandomSubnetId()
+		{
+			int index = new Random().Next(SubnetIds.Count);
+			return SubnetIds[index];
+		}
+
+		/// <summary>
+		/// Get user-data as base64
+		/// Required when interacting with the EC2 API.
+		/// </summary>
+		/// <returns>User-data encoded as Base64</returns>
+		public string? GetUserDataAsBase64()
+		{
+			if (UserData == null)
+			{
+				return UserData;
+			}
+			
+			byte[] plainTextBytes = System.Text.Encoding.UTF8.GetBytes(UserData);
+			return Convert.ToBase64String(plainTextBytes);
+		}
+	}
+	
+	/// <summary>
+	/// Fleet manager for handling AWS EC2 instances
+	/// Will create and/or terminate instances from scratch.
+	/// </summary>
+	public sealed class AwsFleetManager : IFleetManager, IDisposable
+	{
+		/// <summary>
+		/// Prefix for signaling that an EC2 instance belong to an agent pool
+		/// </summary>
+		public const string PoolTagName = "Horde_Autoscale_Pool";
+		private const string AwsTagPropertyName = "aws-tag";
+
+		private readonly AmazonEC2Client _client;
+		private readonly IAgentCollection _agentCollection;
+		private readonly AwsInstanceLaunchSettings _launchSettings;
+		private readonly ILogger _logger;
+
+		/// <summary>
+		/// Constructor
+		/// </summary>
+		public AwsFleetManager(AWSOptions awsOptions, IAgentCollection agentCollection, AwsInstanceLaunchSettings launchSettings, ILogger<AwsFleetManager> logger)
+		{
+			_agentCollection = agentCollection;
+			_launchSettings = launchSettings;
+			_logger = logger;
+
+			AmazonEC2Config config = new () { RegionEndpoint = RegionEndpoint.GetBySystemName(_launchSettings.Region) };
+			logger.LogInformation("Initializing AWS fleet manager for region {Region}", config.RegionEndpoint);
+			_client = new AmazonEC2Client(awsOptions.Credentials, config);
+		}
+
+		/// <inheritdoc/>
+		public void Dispose()
+		{
+			_client.Dispose();
+		}
+
+		/// <inheritdoc/>
+		public async Task ExpandPoolAsync(IPool pool, IReadOnlyList<IAgent> agents, int count, CancellationToken cancellationToken)
+		{
+			using IScope scope = GlobalTracer.Instance.BuildSpan("ExpandPool").StartActive();
+			scope.Span.SetTag("poolName", pool.Name);
+			scope.Span.SetTag("numAgents", agents.Count);
+			scope.Span.SetTag("count", count);
+
+			List<Tag> tags = _launchSettings.Tags.Select(x => new Tag(x.Key, x.Value)).ToList();
+			tags.Add(new Tag("Horde:RequestedPools", pool.Name));
+			tags.Add(new Tag("Horde:Timestamp:Ec2InstanceLaunchRequested", Convert.ToString(DateTimeOffset.UtcNow.ToUnixTimeMilliseconds())));
+			
+			RunInstancesRequest request = new ()
+			{
+				ImageId = _launchSettings.ImageId,
+				InstanceType = InstanceType.FindValue(_launchSettings.InstanceType),
+				KeyName = _launchSettings.KeyName,
+				SecurityGroupIds = _launchSettings.SecurityGroupIds,
+				SubnetId = _launchSettings.GetRandomSubnetId(), // Pick randomly to spread evenly among available subnets 
+				MinCount = count,
+				MaxCount = count,
+				UserData = _launchSettings.GetUserDataAsBase64(),
+				InstanceInitiatedShutdownBehavior = ShutdownBehavior.Terminate,
+				MetadataOptions = new InstanceMetadataOptionsRequest
+				{
+					HttpEndpoint = InstanceMetadataEndpointState.Enabled,
+					InstanceMetadataTags = InstanceMetadataTagsState.Enabled,
+				}
+			};
+
+			if (_launchSettings.IamInstanceProfile != null)
+			{
+				request.IamInstanceProfile =
+					_launchSettings.IamInstanceProfile.StartsWith("arn:aws", StringComparison.InvariantCulture)
+						? new IamInstanceProfileSpecification { Arn = _launchSettings.IamInstanceProfile }
+						: new IamInstanceProfileSpecification { Name = _launchSettings.IamInstanceProfile };
+			}
+			
+			if (tags.Count > 0)
+			{
+				request.TagSpecifications = new() { new() { ResourceType = ResourceType.Instance, Tags = tags } };
+			}
+
+			RunInstancesResponse response = await _client.RunInstancesAsync(request, cancellationToken);
+			int numStartedInstances = response.Reservation.Instances.Count;
+			scope.Span.SetTag("res.statusCode", (int)response.HttpStatusCode);
+			scope.Span.SetTag("res.numInstances", numStartedInstances);
+
+			foreach (Instance i in response.Reservation.Instances)
+			{
+				_logger.LogInformation("Created instance {InstanceId} for pool {PoolId}", i.InstanceId, pool.Id);
+			}
+
+			if (numStartedInstances != count)
+			{
+				_logger.LogWarning("Unable to create all the requested instances for pool {PoolId}. RequestedCount={RequestedCount} ActualCount={ActualCount}", pool.Id, count, numStartedInstances);
+			}
+		}
+
+		/// <inheritdoc/>
+		public Task ShrinkPoolAsync(IPool pool, IReadOnlyList<IAgent> agents, int count, CancellationToken cancellationToken)
+		{
+			return ShrinkPoolViaAgentShutdownRequestAsync(_agentCollection, pool, agents, count, cancellationToken);
+		}
+		
+		/// <summary>
+		/// Shrink pool by gracefully requesting a shutdown for the agent
+		/// Allowing the agent to terminate at next best possible moment (usually when all the leases have finished).
+		/// </summary>
+		/// <param name="agentCollection">Instance of the AgentCollection</param>
+		/// <param name="pool">Pool to resize</param>
+		/// <param name="agents">Current list of agents in the pool</param>
+		/// <param name="count">Number of agents to remove</param>		/// <param name="cancellationToken">Cancellation token for the call</param>
+		public static async Task ShrinkPoolViaAgentShutdownRequestAsync(IAgentCollection agentCollection, IPool pool, IReadOnlyList<IAgent> agents, int count, CancellationToken cancellationToken)
+		{
+			using IScope scope = GlobalTracer.Instance.BuildSpan("ShrinkPool").StartActive();
+			scope.Span.SetTag("poolName", pool.Name);
+			scope.Span.SetTag("count", count);
+			
+			string awsTagProperty = $"{AwsTagPropertyName}={PoolTagName}:{pool.Name}";
+			
+			// Sort the agents by number of active leases. It's better to shutdown agents currently doing nothing.
+			List<IAgent> filteredAgents = agents.OrderBy(x => x.Leases.Count).ToList();
+			List<IAgent> agentsWithAwsTags = filteredAgents.Where(x => x.HasProperty(awsTagProperty)).ToList(); 
+			List<IAgent> agentsLimitedByCount = agentsWithAwsTags.Take(count).ToList();
+			
+			scope.Span.SetTag("agents.num", agents.Count);
+			scope.Span.SetTag("agents.filtered.num", filteredAgents.Count);
+			scope.Span.SetTag("agents.withAwsTags.num", agentsWithAwsTags.Count);
+			scope.Span.SetTag("agents.limitedByCount.num", agentsLimitedByCount.Count);
+
+			foreach (IAgent agent in agentsLimitedByCount)
+			{
+				IAuditLogChannel<AgentId> agentLogger = agentCollection.GetLogger(agent.Id);
+				if (await agentCollection.TryUpdateSettingsAsync(agent, bRequestShutdown: true, shutdownReason: "Autoscaler") != null)
+				{
+					agentLogger.LogInformation("Marked {AgentId} in pool {PoolName} for shutdown due to autoscaling (currently {NumLeases} leases outstanding)", agent.Id, pool.Name, agent.Leases.Count);
+				}
+				else
+				{
+					agentLogger.LogError("Unable to mark agent {AgentId} in pool {PoolName} for shutdown due to autoscaling", agent.Id, pool.Name);
+				}
+			}
+		}
+
+		/// <inheritdoc/>
+		public async Task<int> GetNumStoppedInstancesAsync(IPool pool, CancellationToken cancellationToken)
+		{
+			// Find all instances in the pool
+			DescribeInstancesRequest describeRequest = new ();
+			describeRequest.Filters = new List<Filter>();
+			describeRequest.Filters.Add(new Filter("instance-state-name", new List<string> { InstanceStateName.Stopped.Value }));
+			describeRequest.Filters.Add(new Filter("tag:" + PoolTagName, new List<string> { pool.Name }));
+
+			DescribeInstancesResponse describeResponse = await _client.DescribeInstancesAsync(describeRequest);
+			return describeResponse.Reservations.SelectMany(x => x.Instances).Select(x => x.InstanceId).Distinct().Count();
+		}
+	}
+}
diff --git a/Engine/Source/Programs/Horde/Horde.Build/Agents/Fleet/Providers/AwsTerminatingFleetManager.cs b/Engine/Source/Programs/Horde/Horde.Build/Agents/Fleet/Providers/AwsTerminatingFleetManager.cs
deleted file mode 100644
