Horde: Add fleet manager for AWS EC2 auto-scaling groups

diff --git a/Engine/Source/Programs/Horde/Horde.Build/Agents/Fleet/Providers/AwsAsgFleetManager.cs b/Engine/Source/Programs/Horde/Horde.Build/Agents/Fleet/Providers/AwsAsgFleetManager.cs
new file mode 100644
--- /dev/null
+++ b/Engine/Source/Programs/Horde/Horde.Build/Agents/Fleet/Providers/AwsAsgFleetManager.cs
@@ -0,0 +1,122 @@
+// Copyright Epic Games, Inc. All Rights Reserved.
+
+using System;
+using System.Collections.Generic;
+using System.Net;
+using System.Threading;
+using System.Threading.Tasks;
+using Amazon;
+using Amazon.AutoScaling;
+using Amazon.AutoScaling.Model;
+using Amazon.Extensions.NETCore.Setup;
+using Horde.Build.Agents.Pools;
+using Horde.Build.Utilities;
+using Microsoft.Extensions.Logging;
+using OpenTracing;
+using OpenTracing.Util;
+
+namespace Horde.Build.Agents.Fleet.Providers
+{
+	/// <summary>
+	/// Settings for controlling an EC2 auto-scaling group (ASG)
+	/// </summary>
+	public class AwsAsgSettings
+	{
+		/// <summary>
+		/// AWS region (e.g us-east-1)
+		/// </summary>
+		public string Region { get; }
+		
+		/// <summary>
+		/// Name of the auto-scaling group
+		/// </summary>
+		public string Name { get; }
+
+		/// <summary>
+		/// Constructor
+		/// </summary>
+		/// <param name="region"></param>
+		/// <param name="name"></param>
+		public AwsAsgSettings(string region, string name)
+		{
+			Region = region;
+			Name = name;
+		}
+	}
+	
+	/// <summary>
+	/// Fleet manager for handling AWS EC2 instances
+	/// Uses an EC2 auto-scaling group for controlling the number of running instances.
+	/// </summary>
+	public sealed class AwsAsgFleetManager : IFleetManager, IDisposable
+	{
+		private readonly AmazonAutoScalingClient _client;
+		private readonly AwsAsgSettings _asgSettings;
+		private readonly ILogger _logger;
+
+		/// <summary>
+		/// Constructor
+		/// </summary>
+		public AwsAsgFleetManager(AWSOptions awsOptions, AwsAsgSettings asgSettings, ILogger<AwsAsgFleetManager> logger)
+		{
+			_asgSettings = asgSettings;
+			_logger = logger;
+
+			AmazonAutoScalingConfig config = new() { RegionEndpoint = RegionEndpoint.GetBySystemName(_asgSettings.Region) };
+			logger.LogInformation("Initializing AWS ASG fleet manager for region {Region}", config.RegionEndpoint);
+			_client = new AmazonAutoScalingClient(awsOptions.Credentials, config);
+		}
+
+		/// <inheritdoc/>
+		public void Dispose()
+		{
+			_client.Dispose();
+		}
+
+		/// <inheritdoc/>
+		public async Task ExpandPoolAsync(IPool pool, IReadOnlyList<IAgent> agents, int count, CancellationToken cancellationToken)
+		{
+			int desiredCapacity = agents.Count + count;
+			
+			using IScope scope = GlobalTracer.Instance.BuildSpan("AwsAsgFleetManager.ExpandPool").StartActive();
+			scope.Span.SetTag("poolName", pool.Name);
+			scope.Span.SetTag("numAgents", agents.Count);
+			scope.Span.SetTag("count", count);
+			scope.Span.SetTag("desiredCapacity", desiredCapacity);
+			
+			await UpdateAsg(pool.Id, desiredCapacity, scope, cancellationToken);
+		}
+
+		/// <inheritdoc/>
+		public async Task ShrinkPoolAsync(IPool pool, IReadOnlyList<IAgent> agents, int count, CancellationToken cancellationToken)
+		{
+			int desiredCapacity = Math.Max(0, agents.Count - count);
+			
+			using IScope scope = GlobalTracer.Instance.BuildSpan("AwsAsgFleetManager.ShrinkPool").StartActive();
+			scope.Span.SetTag("poolName", pool.Name);
+			scope.Span.SetTag("numAgents", agents.Count);
+			scope.Span.SetTag("count", count);
+			scope.Span.SetTag("desiredCapacity", desiredCapacity);
+			
+			await UpdateAsg(pool.Id, desiredCapacity, scope, cancellationToken);
+		}
+		
+		private async Task UpdateAsg(StringId<IPool> poolId, int desiredCapacity, IScope scope, CancellationToken cancellationToken)
+		{
+			UpdateAutoScalingGroupRequest request = new () { AutoScalingGroupName = _asgSettings.Name, DesiredCapacity = desiredCapacity };
+			UpdateAutoScalingGroupResponse response = await _client.UpdateAutoScalingGroupAsync(request, cancellationToken);
+			
+			scope.Span.SetTag("res.statusCode", (int)response.HttpStatusCode);
+			if (response.HttpStatusCode != HttpStatusCode.OK)
+			{
+				_logger.LogError("Unable to update auto-scaling group {AutoScalingGroup} for pool {PoolId}.", _asgSettings.Name, poolId);
+			}
+		}
+
+		/// <inheritdoc/>
+		public Task<int> GetNumStoppedInstancesAsync(IPool pool, CancellationToken cancellationToken)
+		{
+			return Task.FromResult(0);
+		}
+	}
+}
diff --git a/Engine/Source/Programs/Horde/Horde.Build/Horde.Build.csproj b/Engine/Source/Programs/Horde/Horde.Build/Horde.Build.csproj
--- a/Engine/Source/Programs/Horde/Horde.Build/Horde.Build.csproj
+++ b/Engine/Source/Programs/Horde/Horde.Build/Horde.Build.csproj
@@ -19,6 +19,7 @@
   </PropertyGroup>
 
   <ItemGroup>
+    <PackageReference Include="AWSSDK.AutoScaling" Version="3.7.10.43" />
     <PackageReference Include="AWSSDK.EC2" Version="3.7.75.1" />
     <PackageReference Include="AWSSDK.Extensions.NETCore.Setup" Version="3.7.2" />
     <PackageReference Include="AWSSDK.S3" Version="3.7.9.19" />
