Horde: Add some more useful timings for flushing data.

diff --git a/Engine/Source/Programs/Horde/Horde.Build/Perforce/ReplicationService.cs b/Engine/Source/Programs/Horde/Horde.Build/Perforce/ReplicationService.cs
--- a/Engine/Source/Programs/Horde/Horde.Build/Perforce/ReplicationService.cs
+++ b/Engine/Source/Programs/Horde/Horde.Build/Perforce/ReplicationService.cs
@@ -626,8 +626,8 @@ namespace Horde.Build.Perforce
 					if (!pathToDirectory.TryGetValue(directoryPath, out directory))
 					{
 						directory = new DirectoryToSync();
-						directory._path = directoryPath.Clone();
-						pathToDirectory.Add(directory._path, directory);
+						directory._path = NormalizePathSeparators(directoryPath);
+						pathToDirectory.Add(directoryPath, directory);
 					}
 
 					directory._size += response.Data.FileSize;
@@ -659,14 +659,6 @@ namespace Horde.Build.Perforce
 					(int dirIdx, Utf8String path, long size) = GetSyncBatch(directories, MaxBatchSize, clientInfo.ServerInfo.Utf8PathComparer, _logger);
 					syncedSize += size;
 
-					// Correct the path separators
-					byte[] newPath = new byte[path.Length];
-					for (int idx = 0; idx < path.Length; idx++)
-					{
-						newPath[idx] = (path[idx] == '\\') ? (byte)'/' : path[idx];
-					}
-					path = new Utf8String(newPath);
-
 					// If we're syncing the entire view, use the given filter instead. It will be tighter than the base path from all files.
 					if (filter != null && path.Equals("..."))
 					{
@@ -684,7 +676,6 @@ namespace Horde.Build.Perforce
 
 					Stopwatch syncTimer = Stopwatch.StartNew();
 					Stopwatch processTimer = new Stopwatch();
-					Stopwatch flushTimer = new Stopwatch();
 					Stopwatch gcTimer = new Stopwatch();
 
 					Dictionary<int, FileEntry> handles = new Dictionary<int, FileEntry>();
@@ -715,7 +706,11 @@ namespace Horde.Build.Perforce
 							}
 							else if (io.Command == PerforceIoCommand.Close)
 							{
-								handles.Remove(io.File);
+								FileEntry? entry;
+								if (handles.Remove(io.File, out entry))
+								{
+									await entry.CollapseAsync(writer, cancellationToken);
+								}
 							}
 							else if (io.Command == PerforceIoCommand.Unlink)
 							{
@@ -730,7 +725,7 @@ namespace Horde.Build.Perforce
 
 						processTimer.Stop();
 					}
-					_logger.LogInformation("Completed batch in {TimeSeconds:n1}s ({ProcessTimeSeconds:n1}s processing, {TrimTimeSeconds:n1}s stalled for flush)", syncTimer.Elapsed.TotalSeconds, processTimer.Elapsed.TotalSeconds, flushTimer.Elapsed.TotalSeconds);
+					_logger.LogInformation("Completed batch in {TimeSeconds:n1}s ({ProcessTimeSeconds:n1}s processing)", syncTimer.Elapsed.TotalSeconds, processTimer.Elapsed.TotalSeconds);
 
 					// Combine all the existing sync paths together with a new wildcard.
 					List<Utf8String> paths = syncNode.Paths;
@@ -764,19 +759,19 @@ namespace Horde.Build.Perforce
 					// Print the new state
 					if (dataSize > flushDataSize)
 					{
-						_logger.LogInformation("Queuing flush of working set after receiving {Size:n1}mb. Filter:", dataSize / (1024 * 1024));
+						_logger.LogInformation("Queuing flush of working set after receiving {Size:n1}mb.", dataSize / (1024 * 1024));
 						foreach (Utf8String nextPath in paths)
 						{
-							_logger.LogInformation("  {Path}", nextPath);
+							_logger.LogInformation("Filter: {Path}", nextPath);
 						}
 
-						flushTimer.Start();
+						Stopwatch flushTimer = Stopwatch.StartNew();
 						await flushTask;
 						flushTask = writer.FlushAsync(syncNode, cancellationToken);
 						flushTimer.Stop();
 
 						flushDataSize = dataSize + FlushInterval;
-						_logger.LogInformation("Next flush at {Size:n1}mb.", flushDataSize / (1024 * 1024));
+						_logger.LogInformation("Flush took {Time:n1}s. Next flush at {Size:n1}mb.", flushTimer.Elapsed.TotalSeconds, flushDataSize / (1024 * 1024));
 					}
 
 					// Remove the paths we've synced
@@ -799,6 +794,16 @@ namespace Horde.Build.Perforce
 			return syncNode;
 		}
 
+		static Utf8String NormalizePathSeparators(Utf8String path)
+		{
+			byte[] newPath = new byte[path.Length];
+			for (int idx = 0; idx < path.Length; idx++)
+			{
+				newPath[idx] = (path[idx] == '\\') ? (byte)'/' : path[idx];
+			}
+			return new Utf8String(newPath);
+		}
+
 		static (int, Utf8String, long) GetSyncBatch(List<DirectoryToSync> directories, long maxSize, Utf8StringComparer comparer, ILogger logger)
 		{
 			int idx = directories.Count - 1;
