Horde: Update ref version number for replication.

diff --git a/Engine/Source/Programs/Horde/Horde.Build/Perforce/ReplicationService.cs b/Engine/Source/Programs/Horde/Horde.Build/Perforce/ReplicationService.cs
--- a/Engine/Source/Programs/Horde/Horde.Build/Perforce/ReplicationService.cs
+++ b/Engine/Source/Programs/Horde/Horde.Build/Perforce/ReplicationService.cs
@@ -471,7 +471,7 @@ namespace Horde.Build.Perforce
 		/// <returns></returns>
 		static RefName GetRefName(StreamId streamId, int change, string? filter, bool revisionsOnly)
 		{
-			StringBuilder builder = new StringBuilder($"v4/{streamId}/{change}");
+			StringBuilder builder = new StringBuilder($"v5/{streamId}/{change}");
 			if (filter != null)
 			{
 				builder.Append($"-flt_{IoHash.Compute(Encoding.UTF8.GetBytes(filter)).ToString().Substring(0, 6)}");
