Horde: Build fixes.

diff --git a/Engine/Source/Programs/Horde/Horde.Agent/Commands/Bundles/CreateCommand.cs b/Engine/Source/Programs/Horde/Horde.Agent/Commands/Bundles/CreateCommand.cs
--- a/Engine/Source/Programs/Horde/Horde.Agent/Commands/Bundles/CreateCommand.cs
+++ b/Engine/Source/Programs/Horde/Horde.Agent/Commands/Bundles/CreateCommand.cs
@@ -120,10 +120,10 @@ namespace Horde.Agent.Commands.Bundles
 			return _blobStore;
 		}
 
-		protected ITreeStore CreateTreeStore(ILogger logger, IMemoryCache cache)
+		protected ITreeStore CreateTreeStore(ILogger logger)
 		{
 			IBlobStore blobStore = CreateBlobStore(logger);
-			return new BundleStore(blobStore, new BundleOptions(), cache);
+			return new BundleStore(blobStore, new BundleOptions());
 		}
 	}
 
@@ -138,11 +138,10 @@ namespace Horde.Agent.Commands.Bundles
 
 		public override async Task<int> ExecuteAsync(ILogger logger)
 		{
-			using IMemoryCache cache = new MemoryCache(new MemoryCacheOptions());
-			using ITreeStore store = CreateTreeStore(logger, cache);
+			using ITreeStore store = CreateTreeStore(logger);
 			ITreeWriter writer = store.CreateTreeWriter(RefName);
 
-			DirectoryNode node = new DirectoryNode();
+			DirectoryNode node = new DirectoryNode(DirectoryFlags.None);
 			await node.CopyFromDirectoryAsync(InputDir.ToDirectoryInfo(), new ChunkingOptions(), writer, logger, CancellationToken.None);
 
 			await writer.FlushAsync(node, CancellationToken.None);
diff --git a/Engine/Source/Programs/Horde/Horde.Agent/Commands/Bundles/ExtractCommand.cs b/Engine/Source/Programs/Horde/Horde.Agent/Commands/Bundles/ExtractCommand.cs
--- a/Engine/Source/Programs/Horde/Horde.Agent/Commands/Bundles/ExtractCommand.cs
+++ b/Engine/Source/Programs/Horde/Horde.Agent/Commands/Bundles/ExtractCommand.cs
@@ -5,7 +5,6 @@ using System.Threading.Tasks;
 using EpicGames.Core;
 using EpicGames.Horde.Storage;
 using EpicGames.Horde.Storage.Nodes;
-using Microsoft.Extensions.Caching.Memory;
 using Microsoft.Extensions.Logging;
 
 namespace Horde.Agent.Commands.Bundles
@@ -21,8 +20,7 @@ namespace Horde.Agent.Commands.Bundles
 
 		public override async Task<int> ExecuteAsync(ILogger logger)
 		{
-			using IMemoryCache cache = new MemoryCache(new MemoryCacheOptions());
-			using ITreeStore store = CreateTreeStore(logger, cache);
+			using ITreeStore store = CreateTreeStore(logger);
 
 			DirectoryNode node = await store.ReadTreeAsync<DirectoryNode>(RefName);
 			await node.CopyToDirectoryAsync(OutputDir.ToDirectoryInfo(), logger, CancellationToken.None);
