Build fix.

diff --git a/Engine/Source/Programs/Horde/Horde.Build.Tests/TestSetup.cs b/Engine/Source/Programs/Horde/Horde.Build.Tests/TestSetup.cs
--- a/Engine/Source/Programs/Horde/Horde.Build.Tests/TestSetup.cs
+++ b/Engine/Source/Programs/Horde/Horde.Build.Tests/TestSetup.cs
@@ -233,7 +233,7 @@ namespace Horde.Build.Tests
 			services.AddSingleton<ITreeStore<CommitService>>(sp =>
 			{
 				IBlobStore blobStore = new BlobStore(new TransientStorageBackend());
-				return new BundleStore(blobStore, new BundleOptions(), sp.GetRequiredService<IMemoryCache>()).ForType<CommitService>();
+				return new BundleStore(blobStore, new BundleOptions()).ForType<CommitService>();
 			});
 
 			services.AddSingleton<IStorageClient, BasicStorageClient>();
