Horde: Support querying number of compute tasks in queue per pool

Uses the requirements together with pool properties to see if the agents from a pool potentially can run the task.
Primarily used for distinguishing between high-level properties such as OS or processor arch.

diff --git a/Engine/Source/Programs/Horde/Horde.Build.Tests/Compute/ComputeServiceTest.cs b/Engine/Source/Programs/Horde/Horde.Build.Tests/Compute/ComputeServiceTest.cs
new file mode 100644
--- /dev/null
+++ b/Engine/Source/Programs/Horde/Horde.Build.Tests/Compute/ComputeServiceTest.cs
@@ -0,0 +1,108 @@
+// Copyright Epic Games, Inc. All Rights Reserved.
+
+using System.Collections.Generic;
+using System.Threading;
+using System.Threading.Tasks;
+using EpicGames.Core;
+using EpicGames.Horde.Compute;
+using EpicGames.Horde.Storage;
+using EpicGames.Serialization;
+using Horde.Build.Agents.Pools;
+using Horde.Build.Server;
+using Microsoft.Extensions.DependencyInjection;
+using Microsoft.VisualStudio.TestTools.UnitTesting;
+
+namespace Horde.Build.Tests.Compute
+{
+	[TestClass]
+	public class ComputeServiceTest : TestSetup
+	{
+		private readonly ComputeClusterConfig _clusterConfig;
+		private readonly NamespaceId _namespaceId;
+		private readonly ClusterId _clusterId;
+		private static int s_poolId;
+		private static int s_channelId;
+		private static int s_refId;
+
+		private IStorageClient StorageClient => ServiceProvider.GetRequiredService<IStorageClient>();
+
+		public ComputeServiceTest()
+		{
+			_clusterConfig = new();
+			_namespaceId = new (_clusterConfig.NamespaceId);
+			_clusterId = new (_clusterConfig.Id);
+		}
+
+		[TestInitialize]
+		public async Task Setup()
+		{
+			Globals globals = await MongoService.GetGlobalsAsync();
+			globals.ComputeClusters.Add(_clusterConfig);
+			Assert.IsTrue(await MongoService.TryUpdateSingletonAsync(globals));
+		}
+		
+		[TestMethod]
+		public async Task QueueCountForPoolNoMatchingCondition()
+		{
+			await CreateTaskWithRequirement("memory >= 10");
+			IPool pool = await CreatePool(new () { {"cpu", "5"} });
+			Assert.AreEqual(0, await ComputeService.GetNumQueuedTasksForPool(_clusterId, pool));
+		}
+
+		[TestMethod]
+		public async Task QueueCountForPoolSingleCondition()
+		{
+			await CreateTaskWithRequirement("memory >= 10");
+			IPool pool = await CreatePool(new () { {"memory", "20"} });
+			Assert.AreEqual(1, await ComputeService.GetNumQueuedTasksForPool(_clusterId, pool));
+		}
+		
+		[TestMethod]
+		public async Task QueueCountForPoolTwoConditions()
+		{
+			await CreateTaskWithRequirement("memory >= 10 && cpu < 5");
+			IPool pool = await CreatePool(new () { {"memory", "15"}, {"cpu", "4"} });
+			Assert.AreEqual(1, await ComputeService.GetNumQueuedTasksForPool(_clusterId, pool));
+		}
+		
+		[TestMethod]
+		public async Task QueueCountForPoolMultipleTasks()
+		{
+			await CreateTaskWithRequirement("memory >= 10 && cpu < 5");
+			await CreateTaskWithRequirement("memory >= 15 && cpu < 5");
+			await CreateTaskWithRequirement("memory >= 20 && cpu < 5");
+
+			IPool pool1 = await CreatePool(new () { {"memory", "20"}, {"cpu", "4"} });
+			IPool pool2 = await CreatePool(new () { {"memory", "15"}, {"cpu", "4"} });
+
+			Assert.AreEqual(3, await ComputeService.GetNumQueuedTasksForPool(_clusterId, pool1));
+			Assert.AreEqual(2, await ComputeService.GetNumQueuedTasksForPool(_clusterId, pool2));
+		}
+
+		private async Task CreateTaskWithRequirement(string condition, ChannelId? channelId = null, RefId? refId = null)
+		{
+			if (channelId == null)
+			{
+				channelId = new ChannelId($"bogusChannelId-{s_channelId}");
+				Interlocked.Increment(ref s_channelId);
+			}
+
+			if (refId == null)
+			{
+				refId = new RefId($"bogusRefId-{s_refId}");
+				Interlocked.Increment(ref s_refId);
+			}
+
+			Requirements requirements = new (condition);
+			IoHash reqHash = await StorageClient.WriteBlobAsync(_namespaceId, requirements, CancellationToken.None);
+			await ComputeService.AddTasksAsync(_clusterId, channelId.Value, new List<RefId> { refId.Value }, new CbObjectAttachment(reqHash));
+		}
+
+		private async Task<IPool> CreatePool(Dictionary<string, string> properties)
+		{
+			string poolName = $"bogusPool-{s_poolId}";
+			Interlocked.Increment(ref s_poolId);
+			return await PoolService.CreatePoolAsync(poolName, null, true, 0, 0, properties: properties);
+		}
+	}
+}
\ No newline at end of file
diff --git a/Engine/Source/Programs/Horde/Horde.Build.Tests/TestSetup.cs b/Engine/Source/Programs/Horde/Horde.Build.Tests/TestSetup.cs
--- a/Engine/Source/Programs/Horde/Horde.Build.Tests/TestSetup.cs
+++ b/Engine/Source/Programs/Horde/Horde.Build.Tests/TestSetup.cs
@@ -58,6 +58,7 @@ using Horde.Build.Tasks;
 using Horde.Build.Auditing;
 using Horde.Build.Storage.Backends;
 using EpicGames.Horde.Storage.Bundles;
+using Horde.Build.Compute;
 using Horde.Build.Devices;
 
 namespace Horde.Build.Tests
@@ -96,6 +97,7 @@ namespace Horde.Build.Tests
 		public AclService AclService => ServiceProvider.GetRequiredService<AclService>();
 		public AgentSoftwareService AgentSoftwareService => ServiceProvider.GetRequiredService<AgentSoftwareService>();
 		public AgentService AgentService => ServiceProvider.GetRequiredService<AgentService>();
+		public ComputeService ComputeService => ServiceProvider.GetRequiredService<ComputeService>();
 		public MongoService MongoService => ServiceProvider.GetRequiredService<MongoService>();
 		public ITemplateCollection TemplateCollection => ServiceProvider.GetRequiredService<ITemplateCollection>();
 		internal PerforceServiceStub PerforceService => (PerforceServiceStub)ServiceProvider.GetRequiredService<IPerforceService>();
@@ -200,6 +202,7 @@ namespace Horde.Build.Tests
 			services.AddSingleton<AwsReuseFleetManager, AwsReuseFleetManager>();
 			services.AddSingleton<ConsistencyService>();
 			services.AddSingleton<ConfigCollection>();
+			services.AddSingleton<ComputeService>();
 			services.AddSingleton<RequestTrackerService>();
 			services.AddSingleton<CredentialService>();
 			services.AddSingleton<JobTaskSource>();
diff --git a/Engine/Source/Programs/Horde/Horde.Build/Agents/Pools/IPool.cs b/Engine/Source/Programs/Horde/Horde.Build/Agents/Pools/IPool.cs
--- a/Engine/Source/Programs/Horde/Horde.Build/Agents/Pools/IPool.cs
+++ b/Engine/Source/Programs/Horde/Horde.Build/Agents/Pools/IPool.cs
@@ -3,6 +3,7 @@
 using System;
 using System.Collections.Generic;
 using EpicGames.Horde.Common;
+using EpicGames.Horde.Compute;
 using Horde.Build.Agents.Fleet;
 using Horde.Build.Utilities;
 
@@ -105,4 +106,47 @@ namespace Horde.Build.Agents.Pools
 		/// </summary>
 		public int UpdateIndex { get; }
 	}
+
+	/// <summary>
+	/// Extension methods for IPool
+	/// </summary>
+	public static class PoolExtensions
+	{
+		/// <summary>
+		/// Evaluates a condition against a pool
+		/// </summary>
+		/// <param name="pool">The pool to evaluate</param>
+		/// <param name="condition">The condition to evaluate</param>
+		/// <returns>True if the pool satisfies the condition</returns>
+		public static bool SatisfiesCondition(this IPool pool, Condition condition)
+		{
+			return condition.Evaluate(propKey =>
+			{
+				if (pool.Properties.TryGetValue(propKey, out string? propValue))
+				{
+					return new[] { propValue };
+				}
+
+				return Array.Empty<string>();
+			});
+		}
+
+		/// <summary>
+		/// Determine whether a pool offers agents that are compatible with the given requirements.
+		/// Since a pool does not have beforehand knowledge of what type of agents will be created, their expected
+		/// properties must be specified on the IPool. 
+		/// </summary>
+		/// <param name="pool">The pool to verify</param>
+		/// <param name="requirements">Requirements</param>
+		/// <returns>True if the pool satisfies the requirements</returns>
+		public static bool MeetsRequirements(this IPool pool, Requirements requirements)
+		{
+			if (requirements.Condition != null && !pool.SatisfiesCondition(requirements.Condition))
+			{
+				return false;
+			}
+
+			return true;
+		}
+	}
 }
diff --git a/Engine/Source/Programs/Horde/Horde.Build/Compute/ComputeService.cs b/Engine/Source/Programs/Horde/Horde.Build/Compute/ComputeService.cs
--- a/Engine/Source/Programs/Horde/Horde.Build/Compute/ComputeService.cs
+++ b/Engine/Source/Programs/Horde/Horde.Build/Compute/ComputeService.cs
@@ -14,6 +14,7 @@ using Google.Protobuf;
 using Google.Protobuf.WellKnownTypes;
 using Horde.Build.Agents;
 using Horde.Build.Agents.Leases;
+using Horde.Build.Agents.Pools;
 using Horde.Build.Server;
 using Horde.Build.Tasks;
 using Horde.Build.Utilities;
@@ -333,6 +334,18 @@ namespace Horde.Build.Compute
 		{
 			return await _messageQueue.ReadMessagesAsync(GetMessageQueueId(clusterId, channelId));
 		}
+		
+		/// <summary>
+		/// Get number of queued tasks that are compatible for a given pool
+		/// </summary>
+		/// <param name="clusterId">Cluster to query</param>
+		/// <param name="pool">Pool to check for compatibility</param>
+		/// <param name="cancellationToken">Cancellation token for the operation</param> 
+		/// <returns>Number of tasks in queue</returns>
+		public async Task<int> GetNumQueuedTasksForPool(ClusterId clusterId, IPool pool, CancellationToken cancellationToken = default)
+		{
+			return await _taskScheduler.GetNumQueuedTasksAsync(queueKey => CheckRequirements(pool, queueKey), cancellationToken);
+		}
 
 		/// <summary>
 		/// Dequeue completed items from a queue
@@ -412,6 +425,22 @@ namespace Horde.Build.Compute
 			}
 			return agent.MeetsRequirements(requirements);
 		}
+		
+		/// <summary>
+		/// Checks that a pool matches the necessary criteria to execute a task
+		/// </summary>
+		/// <param name="pool"></param>
+		/// <param name="queueKey"></param>
+		/// <returns></returns>
+		async ValueTask<bool> CheckRequirements(IPool pool, QueueKey queueKey)
+		{
+			Requirements? requirements = await GetCachedRequirementsAsync(queueKey);
+			if (requirements == null)
+			{
+				return false;
+			}
+			return pool.MeetsRequirements(requirements);
+		}
 
 		/// <summary>
 		/// Gets the requirements object from the CAS
diff --git a/Engine/Source/Programs/Horde/Horde.Build/Compute/RedisTaskScheduler.cs b/Engine/Source/Programs/Horde/Horde.Build/Compute/RedisTaskScheduler.cs
--- a/Engine/Source/Programs/Horde/Horde.Build/Compute/RedisTaskScheduler.cs
+++ b/Engine/Source/Programs/Horde/Horde.Build/Compute/RedisTaskScheduler.cs
@@ -4,7 +4,6 @@ using System;
 using System.Collections.Generic;
 using System.Diagnostics;
 using System.Linq;
-using System.Runtime.CompilerServices;
 using System.Threading;
 using System.Threading.Channels;
 using System.Threading.Tasks;
@@ -54,6 +53,15 @@ namespace Horde.Build.Compute
 		/// </summary>
 		/// <returns></returns>
 		Task<List<TQueueId>> GetInactiveQueuesAsync();
+		
+		/// <summary>
+		/// Get number of tasks that a given pool/agent can execute
+		/// A read-only operation and will not affect any queue.
+		/// </summary>
+		/// <param name="predicate">Predicate for determining which queues to read</param>
+		/// <param name="token">Cancellation token for the operation</param>
+		/// <returns>Number of tasks matching the supplied predicate</returns>
+		Task<int> GetNumQueuedTasksAsync(Func<TQueueId, ValueTask<bool>> predicate, CancellationToken token = default);
 	}
 
 	/// <summary>
@@ -375,6 +383,21 @@ namespace Horde.Build.Compute
 			return keys.ToList();
 		}
 
+		public async Task<int> GetNumQueuedTasksAsync(Func<TQueueId, ValueTask<bool>> predicate, CancellationToken token = default)
+		{
+			HashSet<TQueueId> queueIds = new (await _queueIndex.MembersAsync());
+			long totalTaskCount = 0;
+			foreach (TQueueId queueId in queueIds)
+			{
+				if (await predicate(queueId))
+				{
+					totalTaskCount += await GetQueue(queueId).LengthAsync();
+				}
+			}
+
+			return (int)totalTaskCount;
+		}
+
 		async Task UpdateQueuesAsync(CancellationToken cancellationToken)
 		{
 			Channel<TQueueId> newQueues = Channel.CreateUnbounded<TQueueId>();
