Horde: Additional device and pool telemetry for stream usage
#rnx

diff --git a/Engine/Source/Programs/Horde/Horde.Build.Tests/DeviceServiceTests.cs b/Engine/Source/Programs/Horde/Horde.Build.Tests/DeviceServiceTests.cs
--- a/Engine/Source/Programs/Horde/Horde.Build.Tests/DeviceServiceTests.cs
+++ b/Engine/Source/Programs/Horde/Horde.Build.Tests/DeviceServiceTests.cs
@@ -153,6 +153,7 @@ namespace Horde.Build.Tests
 			List<GetDeviceTelemetryResponse> telemetry = (await DeviceController!.GetDeviceTelemetry()).Value!;
 			Assert.AreEqual(telemetry.Count, 1);
 			Assert.AreEqual(telemetry[0].Telemetry.Count, 1);
+			Assert.AreEqual(telemetry[0].Telemetry[0].StreamId, "ue5-main");
 			Assert.AreEqual(telemetry[0].Telemetry[0].StepId, "abcd");
 		}
 
@@ -208,6 +209,8 @@ namespace Horde.Build.Tests
 			Assert.AreEqual(1, reservation.DeviceNames.Length);
 			Assert.AreNotSame(maintenanceDevice!.Name, reservation.DeviceNames[0]);
 
+			GetLegacyDeviceResponse reservedDevice = ResultToValue(await DeviceController!.GetDeviceV1Async(reservation.DeviceNames[0]));
+
 			await DeviceService.TickForTestingAsync();
 
 			List<GetDevicePoolTelemetryResponse> telemetry = (await DeviceController!.GetDevicePoolTelemetry()).Value!;
@@ -219,6 +222,10 @@ namespace Horde.Build.Tests
 					Assert.AreEqual(0, t.Available);
 					Assert.AreEqual(1, t.Reserved);
 					Assert.AreEqual(1, t.Maintenance);
+
+					Assert.IsTrue(t.StreamDevices.ContainsKey("ue5-main"));
+					Assert.AreEqual(t.StreamDevices["ue5-main"].Count, 1);
+					Assert.AreEqual(t.StreamDevices["ue5-main"][0], reservedDevice.Id);
 				}
 				else if (t.PlatformId == "testdeviceplatform3")
 				{
diff --git a/Engine/Source/Programs/Horde/Horde.Build/Devices/DeviceCollection.cs b/Engine/Source/Programs/Horde/Horde.Build/Devices/DeviceCollection.cs
--- a/Engine/Source/Programs/Horde/Horde.Build/Devices/DeviceCollection.cs
+++ b/Engine/Source/Programs/Horde/Horde.Build/Devices/DeviceCollection.cs
@@ -7,6 +7,7 @@ using System.Threading.Tasks;
 using Horde.Build.Acls;
 using Horde.Build.Projects;
 using Horde.Build.Server;
+using Horde.Build.Streams;
 using Horde.Build.Users;
 using Horde.Build.Utilities;
 using MongoDB.Bson;
@@ -20,6 +21,7 @@ namespace Horde.Build.Devices
 	using DevicePoolId = StringId<IDevicePool>;
 	using ProjectId = StringId<IProject>;
 	using UserId = ObjectId<IUser>;
+	using StreamId = StringId<IStream>;
 
 	/// <summary>
 	/// Collection of device documents
@@ -102,6 +104,9 @@ namespace Horde.Build.Devices
 			[BsonRequired]
 			public DevicePoolId PoolId { get; set; }
 
+			[BsonIgnoreIfNull]
+			public string? StreamId { get; set; }
+
 			[BsonIgnoreIfNull]
 			public string? JobId { get; set; }
 
@@ -146,7 +151,7 @@ namespace Horde.Build.Devices
 
 			}
 
-			public DeviceReservationDocument(ObjectId id, DevicePoolId poolId, List<DeviceId> devices, List<string> requestedDevicePlatforms, DateTime createTimeUtc, string? hostname, string? reservationDetails, string? jobId, string? stepId)
+			public DeviceReservationDocument(ObjectId id, DevicePoolId poolId, List<DeviceId> devices, List<string> requestedDevicePlatforms, DateTime createTimeUtc, string? hostname, string? reservationDetails, string? streamId, string? jobId, string? stepId)
 			{
 				Id = id;
 				PoolId = poolId;
@@ -156,6 +161,7 @@ namespace Horde.Build.Devices
 				UpdateTimeUtc = createTimeUtc;
 				Hostname = hostname;
 				ReservationDetails = reservationDetails;
+				StreamId = streamId;
 				JobId = jobId;
 				StepId = stepId;
 
@@ -263,6 +269,12 @@ namespace Horde.Build.Devices
 			[BsonRequired]
 			public DateTime CreateTimeUtc { get; set; }
 
+			/// <summary>
+			/// The stream id which utilized device
+			/// </summary>
+			[BsonIgnoreIfNull]
+			public string? StreamId { get; set; }
+
 			/// <summary>
 			/// The job id which utilized device
 			/// </summary>
@@ -304,7 +316,7 @@ namespace Horde.Build.Devices
 			{
 			}
 
-			public DeviceTelemetryDocument(DeviceId deviceId, ObjectId? reservationId = null, DateTime? reservationStartTime = null, string? jobId = null, string? stepId = null)
+			public DeviceTelemetryDocument(DeviceId deviceId, ObjectId? reservationId = null, DateTime? reservationStartTime = null, string? streamId = null, string? jobId = null, string? stepId = null)
 			{
 				TelemetryId = ObjectId.GenerateNewId();
 				CreateTimeUtc = DateTime.UtcNow;
@@ -313,6 +325,7 @@ namespace Horde.Build.Devices
 				ReservationId = reservationId;
 				ReservationStartUtc = reservationStartTime;
 
+				StreamId = streamId;
 				JobId = jobId;
 				StepId = stepId;
 			}
@@ -369,13 +382,19 @@ namespace Horde.Build.Devices
 			[BsonRequired]
 			public int Disabled { get; set; } = 0;
 
+			/// <summary>
+			/// Stream devices that are reserved
+			/// </summary>
+			[BsonIgnoreIfNull]
+			public Dictionary<StreamId, List<DeviceId>>? StreamDevices { get; set; } = new Dictionary<StreamId, List<DeviceId>>();
+			IReadOnlyDictionary<StreamId, IReadOnlyList<DeviceId>> IDevicePlatformTelemetry.StreamDevices => StreamDevices?.ToDictionary(kvp => kvp.Key, kvp => kvp.Value as IReadOnlyList<DeviceId>) ?? new Dictionary<StreamId, IReadOnlyList<DeviceId>>();
 
 			[BsonConstructor]
 			private DevicePlatformTelemetryDocument()
 			{
 			}
 
-			public DevicePlatformTelemetryDocument(DevicePlatformId platformId, int available, int reserved, int maintenance, int problem, int disabled)
+			public DevicePlatformTelemetryDocument(DevicePlatformId platformId, int available, int reserved, int maintenance, int problem, int disabled, Dictionary<StreamId, List<DeviceId>> streamDevices)
 			{
 				TelemetryId = ObjectId.GenerateNewId();
 				CreateTimeUtc = DateTime.UtcNow;
@@ -385,6 +404,7 @@ namespace Horde.Build.Devices
 				Maintenance = maintenance;	
 				Problem = problem;
 				Disabled = disabled;
+				StreamDevices = streamDevices;
 			}
 
 		}
@@ -742,7 +762,7 @@ namespace Horde.Build.Devices
 		}
 
 		/// <inheritdoc/>
-		public async Task<IDeviceReservation?> TryAddReservationAsync(DevicePoolId poolId, List<DeviceRequestData> request, string? hostname, string? reservationDetails, string? jobId, string? stepId)
+		public async Task<IDeviceReservation?> TryAddReservationAsync(DevicePoolId poolId, List<DeviceRequestData> request, string? hostname, string? reservationDetails, string? streamId, string? jobId, string? stepId)
 		{
 
 			if (request.Count == 0)
@@ -868,14 +888,14 @@ namespace Horde.Build.Devices
 			List<string> requestedPlatforms = deviceIds.Select(x => platformRequestMap[x]).ToList();
 
 			// Create new reservation
-			DeviceReservationDocument newReservation = new DeviceReservationDocument(ObjectId.GenerateNewId(), poolId, deviceIds, requestedPlatforms, reservationTimeUtc, hostname, reservationDetails, jobId, stepId);
+			DeviceReservationDocument newReservation = new DeviceReservationDocument(ObjectId.GenerateNewId(), poolId, deviceIds, requestedPlatforms, reservationTimeUtc, hostname, reservationDetails, streamId, jobId, stepId);
 			await _reservations.InsertOneAsync(newReservation);
 
 			// Create device telemetry data for reservation
 			List<DeviceTelemetryDocument> telemetry = new List<DeviceTelemetryDocument>();
 			foreach (DeviceId deviceId in deviceIds)
 			{
-				telemetry.Add(new DeviceTelemetryDocument(deviceId, newReservation.Id, newReservation.CreateTimeUtc, jobId, stepId));
+				telemetry.Add(new DeviceTelemetryDocument(deviceId, newReservation.Id, newReservation.CreateTimeUtc, streamId, jobId, stepId));
 			}
 
 			if (telemetry.Count > 0)
@@ -1052,6 +1072,7 @@ namespace Horde.Build.Devices
 			public int maintenance = 0;
 			public int problem = 0;
 			public int disabled = 0;
+			public Dictionary<StreamId, List<DeviceId>> streamDevices = new Dictionary<StreamId, List<DeviceId>>();
 
 		}
 
@@ -1103,6 +1124,18 @@ namespace Horde.Build.Devices
 					if (reservedDevices.Contains(device))
 					{
 						helper.reserved++;
+						IDeviceReservation? reservation = reservations.FirstOrDefault(x => x.Devices.Contains(device.Id));
+						if (reservation != null && reservation.StreamId != null)
+						{
+							StreamId streamId = new StreamId(reservation.StreamId);
+							List<DeviceId>? deviceIds;
+							if (!helper.streamDevices.TryGetValue(streamId, out deviceIds))
+							{
+								deviceIds = new List<DeviceId>();
+								helper.streamDevices[streamId] = deviceIds;
+							}
+							deviceIds.Add(device.Id);
+						}
 						continue;
 					}
 
@@ -1131,7 +1164,7 @@ namespace Horde.Build.Devices
 				foreach (KeyValuePair<DevicePlatformId, DevicePoolTelemetryHelper> platform in helpers)
 				{
 					DevicePoolTelemetryHelper helper = platform.Value;
-					platformTelemtry.Add(new DevicePlatformTelemetryDocument(platform.Key, helper.available, helper.reserved, helper.maintenance, helper.problem, helper.disabled));
+					platformTelemtry.Add(new DevicePlatformTelemetryDocument(platform.Key, helper.available, helper.reserved, helper.maintenance, helper.problem, helper.disabled, helper.streamDevices));
 				}
 
 				poolTelemetry[pool.Id] = platformTelemtry;
diff --git a/Engine/Source/Programs/Horde/Horde.Build/Devices/DeviceMessages.cs b/Engine/Source/Programs/Horde/Horde.Build/Devices/DeviceMessages.cs
--- a/Engine/Source/Programs/Horde/Horde.Build/Devices/DeviceMessages.cs
+++ b/Engine/Source/Programs/Horde/Horde.Build/Devices/DeviceMessages.cs
@@ -557,6 +557,11 @@ namespace Horde.Build.Devices
 		[Required]
 		public DateTime CreateTimeUtc { get; set; }
 
+		/// <summary>
+		/// The stream id which utilized device
+		/// </summary>
+		public string? StreamId { get; set; }
+
 		/// <summary>
 		/// The job id which utilized device
 		/// </summary>
@@ -589,6 +594,7 @@ namespace Horde.Build.Devices
 		public GetTelemetryInfoResponse(IDeviceTelemetry data)
 		{
 			CreateTimeUtc = data.CreateTimeUtc;
+			StreamId = data.StreamId;
 			JobId = data.JobId;
 			StepId = data.StepId;
 			ReservationStartUtc = data.ReservationStartUtc;
@@ -663,10 +669,15 @@ namespace Horde.Build.Devices
 		/// </summary>
 		public int Disabled { get; set; }
 
+		/// <summary>
+		/// StreamId to reserved device ids
+		/// </summary>
+		public Dictionary<string, List<string>> StreamDevices { get; set; } = null!;
+
 		/// <summary>
 		/// Constructor
 		/// </summary>
-		public GetDevicePlatformTelemetryResponse(string platformId, int available, int reserved, int maintenance, int problem, int disabled)
+		public GetDevicePlatformTelemetryResponse(string platformId, int available, int reserved, int maintenance, int problem, int disabled, Dictionary<string, List<string>> streamDevices)
 		{
 			PlatformId = platformId;
 			Available = available;
@@ -674,6 +685,7 @@ namespace Horde.Build.Devices
 			Maintenance = maintenance;
 			Problem = problem;
 			Disabled = disabled;
+			StreamDevices = streamDevices;
 		}
 	}
 
@@ -814,6 +826,12 @@ namespace Horde.Build.Devices
 	/// </summary>
 	public class GetLegacyDeviceResponse
 	{
+		/// <summary>
+		/// The  id of the reserved device
+		/// </summary>
+		[Required]
+		public string Id { get; set; } = null!;
+
 		/// <summary>
 		/// The  name of the reserved device
 		/// </summary>
diff --git a/Engine/Source/Programs/Horde/Horde.Build/Devices/DeviceService.cs b/Engine/Source/Programs/Horde/Horde.Build/Devices/DeviceService.cs
--- a/Engine/Source/Programs/Horde/Horde.Build/Devices/DeviceService.cs
+++ b/Engine/Source/Programs/Horde/Horde.Build/Devices/DeviceService.cs
@@ -131,7 +131,7 @@ namespace Horde.Build.Devices
 		/// </summary>
 		readonly IDeviceCollection _devices;
 		readonly ITicker _ticker;		
-		static readonly int _poolTelemetryMinutes = 15;
+		const int _poolTelemetryMinutes = 10;
 		int _poolTelemetryTick = _poolTelemetryMinutes;
 
 		/// <summary>
@@ -388,9 +388,20 @@ namespace Horde.Build.Devices
 		/// <summary>
 		/// Try to create a reservation satisfying the specified device platforms and models
 		/// </summary>
-		public Task<IDeviceReservation?> TryCreateReservationAsync(DevicePoolId pool, List<DeviceRequestData> request, string? hostname = null, string? reservationDetails = null, string? jobId = null, string? stepId = null)
+		public async Task<IDeviceReservation?> TryCreateReservationAsync(DevicePoolId pool, List<DeviceRequestData> request, string? hostname = null, string? reservationDetails = null, string? jobId = null, string? stepId = null)
 		{
-			return _devices.TryAddReservationAsync(pool, request, hostname, reservationDetails, jobId, stepId);
+			string? streamId = null;
+
+			if (jobId != null)
+			{
+				IJob? job = await _jobService.GetJobAsync(new JobId(jobId));
+				if (job != null)
+				{
+					streamId = job.StreamId.ToString();
+				}
+			}
+
+			return await _devices.TryAddReservationAsync(pool, request, hostname, reservationDetails, streamId, jobId, stepId);
 		}
 
 		/// <summary>
diff --git a/Engine/Source/Programs/Horde/Horde.Build/Devices/DevicesController.cs b/Engine/Source/Programs/Horde/Horde.Build/Devices/DevicesController.cs
--- a/Engine/Source/Programs/Horde/Horde.Build/Devices/DevicesController.cs
+++ b/Engine/Source/Programs/Horde/Horde.Build/Devices/DevicesController.cs
@@ -524,7 +524,12 @@ namespace Horde.Build.Devices
 
 					foreach (IDevicePlatformTelemetry telemetry in pool.Value)
 					{
-						platformTelemetry.Add(new GetDevicePlatformTelemetryResponse(telemetry.PlatformId.ToString(), telemetry.Available, telemetry.Reserved, telemetry.Maintenance, telemetry.Problem, telemetry.Disabled));
+						Dictionary<string, List<string>> streamDevices = new Dictionary<string, List<string>>();
+						foreach(KeyValuePair<StringId<Streams.IStream>, IReadOnlyList<DeviceId>> kvp in telemetry.StreamDevices)
+						{
+							streamDevices[kvp.Key.ToString()] = kvp.Value.Select(x => x.ToString()).ToList();
+						}
+						platformTelemetry.Add(new GetDevicePlatformTelemetryResponse(telemetry.PlatformId.ToString(), telemetry.Available, telemetry.Reserved, telemetry.Maintenance, telemetry.Problem, telemetry.Disabled, streamDevices));
 					}
 
 					poolData[pool.Key.ToString()] = platformTelemetry;
@@ -1069,6 +1074,7 @@ namespace Horde.Build.Devices
 
 			GetLegacyDeviceResponse response = new GetLegacyDeviceResponse();
 
+			response.Id = device.Id.ToString();
 			response.Name = device.Name;
 			response.Type = platformName;
 			response.IPOrHostName = device.Address ?? "";
diff --git a/Engine/Source/Programs/Horde/Horde.Build/Devices/IDevice.cs b/Engine/Source/Programs/Horde/Horde.Build/Devices/IDevice.cs
--- a/Engine/Source/Programs/Horde/Horde.Build/Devices/IDevice.cs
+++ b/Engine/Source/Programs/Horde/Horde.Build/Devices/IDevice.cs
@@ -4,6 +4,7 @@ using System;
 using System.Collections.Generic;
 using Horde.Build.Acls;
 using Horde.Build.Projects;
+using Horde.Build.Streams;
 using Horde.Build.Users;
 using Horde.Build.Utilities;
 using MongoDB.Bson;
@@ -16,6 +17,7 @@ namespace Horde.Build.Devices
 	using DevicePoolId = StringId<IDevicePool>;
 	using ProjectId = StringId<IProject>;
 	using UserId = ObjectId<IUser>;
+	using StreamId = StringId<IStream>;
 
 	/// <summary>
 	/// A reservation containing one or more devices
@@ -32,6 +34,11 @@ namespace Horde.Build.Devices
 		/// </summary>
 		public DevicePoolId PoolId { get; }
 
+		/// <summary>
+		/// Strwam id holding reservation
+		/// </summary>
+		public string? StreamId { get; }
+
 		/// <summary>
 		/// JobID holding reservation
 		/// </summary>
@@ -209,6 +216,11 @@ namespace Horde.Build.Devices
 		/// </summary>
 		public DeviceId DeviceId { get; }
 
+		/// <summary>
+		/// The stream id which utilized device
+		/// </summary>
+		public string? StreamId { get; }
+
 		/// <summary>
 		/// The job id which utilized device
 		/// </summary>
@@ -358,6 +370,11 @@ namespace Horde.Build.Devices
 		/// </summary>
 		public int Disabled { get; }
 
+		/// <summary>
+		/// Devices reserved in streams
+		/// </summary>
+		public IReadOnlyDictionary<StreamId, IReadOnlyList<DeviceId>> StreamDevices { get; }
+
 	}
 
 	/// <summary>
diff --git a/Engine/Source/Programs/Horde/Horde.Build/Devices/IDeviceCollection.cs b/Engine/Source/Programs/Horde/Horde.Build/Devices/IDeviceCollection.cs
--- a/Engine/Source/Programs/Horde/Horde.Build/Devices/IDeviceCollection.cs
+++ b/Engine/Source/Programs/Horde/Horde.Build/Devices/IDeviceCollection.cs
@@ -174,18 +174,19 @@ namespace Horde.Build.Devices
         /// <returns></returns>
         Task CheckoutDeviceAsync(DeviceId deviceId, UserId? checkedOutByUserId);
 
-        // RESERVATIONS
+		// RESERVATIONS
 
-        /// <summary>
-        /// Create a new reseveration in the pool with the specified devices
-        /// </summary>
-        /// <param name="poolId">The pool of devices to use for the new reservation</param>
-        /// <param name="request">The requested devices for the reservation</param>
-        /// <param name="hostname">The hostname of the machine making the reservation</param>
-        /// <param name="reservationDetails">The details of the reservation</param>
-        /// <param name="jobId">The Job Id associated with the job</param>
-        /// <param name="stepId">The Step Id associated with the job</param>
-        Task<IDeviceReservation?> TryAddReservationAsync(DevicePoolId poolId, List<DeviceRequestData> request, string? hostname, string? reservationDetails, string? jobId, string? stepId);
+		/// <summary>
+		/// Create a new reseveration in the pool with the specified devices
+		/// </summary>
+		/// <param name="poolId">The pool of devices to use for the new reservation</param>
+		/// <param name="request">The requested devices for the reservation</param>
+		/// <param name="hostname">The hostname of the machine making the reservation</param>
+		/// <param name="reservationDetails">The details of the reservation</param>
+		/// <param name="streamId">The Stream Id associated with the job</param>
+		/// <param name="jobId">The Job Id associated with the job</param>
+		/// <param name="stepId">The Step Id associated with the job</param>
+		Task<IDeviceReservation?> TryAddReservationAsync(DevicePoolId poolId, List<DeviceRequestData> request, string? hostname, string? reservationDetails, string? streamId, string? jobId, string? stepId);
 
 		/// <summary>
 		/// Gets a reservation by guid for legacy clients
