Horde: Fix issue with log view options rendering offscreen with display scaling
#rnx
#skipci

diff --git a/Engine/Source/Programs/Horde/HordeDashboard/src/components/LogView.tsx b/Engine/Source/Programs/Horde/HordeDashboard/src/components/LogView.tsx
--- a/Engine/Source/Programs/Horde/HordeDashboard/src/components/LogView.tsx
+++ b/Engine/Source/Programs/Horde/HordeDashboard/src/components/LogView.tsx
@@ -928,8 +928,8 @@ export const LogList: React.FC<{ logId: string }> = observer(({ logId }) => {
       <Stack tokens={{ childrenGap: 0 }} style={{ backgroundColor: "#FFFFFF", paddingTop: 12 }}>
          <Stack horizontal >
             <div key={`windowsize_logview1_${windowSize.width}_${windowSize.height}`} style={{ width: vw / 2 - 880, flexShrink: 0 }} />
-            <Stack tokens={{ childrenGap: 0 }} styles={{ root: { backgroundColor: "#FFFFFF", paddingLeft: 4, paddingRight: 24, paddingTop: 12 } }}>
-               <Stack horizontal styles={{ root: { paddingBottom: 4, width: 1768 } }}>
+            <Stack tokens={{ childrenGap: 0, maxWidth: 1780 }} disableShrink={true} styles={{ root: { width: "100%", backgroundColor: "#FFFFFF", paddingLeft: 4, paddingRight: 24, paddingTop: 12 } }}>
+               <Stack horizontal style={{paddingBottom: 4}}>
                   <Stack className={hordeClasses.button} horizontal horizontalAlign={"start"} verticalAlign="center" tokens={{ childrenGap: 8 }}>
                      <Stack horizontal tokens={{ childrenGap: 2 }}>
                         <DefaultButton disabled={!errors.length} className={errors.length ? handler.style.errorButton : handler.style.errorButtonDisabled}
