Horde: Replication performance improvements.

* Rewritten BuzHash inner loop as unsafe code to avoid bounds checks.
* Multiple paths can now be synced in a single command.
* Added stats for P4 output buffer being stalled, and throughput of replication service.
* Fixed OOB array access when serializing files with SHA1 hashes.

diff --git a/Engine/Source/Programs/Horde/Horde.Build/Perforce/ReplicationService.cs b/Engine/Source/Programs/Horde/Horde.Build/Perforce/ReplicationService.cs
--- a/Engine/Source/Programs/Horde/Horde.Build/Perforce/ReplicationService.cs
+++ b/Engine/Source/Programs/Horde/Horde.Build/Perforce/ReplicationService.cs
@@ -658,8 +658,15 @@ namespace Horde.Build.Perforce
 					await perforce.SyncQuietAsync(SyncOptions.Force | SyncOptions.KeepWorkspaceFiles, -1, flushPath, cancellationToken);
 				}
 
-				// Do a sync preview to find everything that's left, and sort the remaining list of paths
+				// Add the root directory from the filter to the list of files to sync. This prevents traversing above it.
 				Dictionary<Utf8String, DirectoryToSync> pathToDirectory = new Dictionary<Utf8String, DirectoryToSync>(clientInfo.ServerInfo.Utf8PathComparer);
+				if (filter != null)
+				{
+					string rootPath = filter.Substring(0, filter.LastIndexOf('/') + 1);
+					pathToDirectory.Add(rootPath.Replace('/', Path.DirectorySeparatorChar), new DirectoryToSync(rootPath, clientInfo.ServerInfo.Utf8PathComparer));
+				}
+
+				// Do a sync preview to find everything that's left, and sort the remaining list of paths
 				await foreach (PerforceResponse<SyncRecord> response in perforce.StreamCommandAsync<SyncRecord>("sync", new[] { "-n" }, new string[] { $"{queryPath}@{change}" }, null, cancellationToken))
 				{
 					PerforceError? error = response.Error;
@@ -683,13 +690,7 @@ namespace Horde.Build.Perforce
 
 					Utf8String directoryPath = path.Slice(clientRoot.Length, fileIdx - clientRoot.Length);
 
-					DirectoryToSync? directory;
-					if (!pathToDirectory.TryGetValue(directoryPath, out directory))
-					{
-						directory = new DirectoryToSync(NormalizePathSeparators(directoryPath), clientInfo.ServerInfo.Utf8PathComparer);
-						pathToDirectory.Add(directoryPath, directory);
-					}
-
+					DirectoryToSync directory = FindOrAddDirectoryTree(pathToDirectory, directoryPath, clientInfo.ServerInfo.Utf8PathComparer);
 					directory._fileNameToSize.Add(path.Slice(fileIdx).Clone(), response.Data.FileSize);
 					directory._size += response.Data.FileSize;
 				}
@@ -711,29 +712,48 @@ namespace Horde.Build.Perforce
 					// Save the incremental state
 					if (syncedSize > 0)
 					{
+						Stopwatch flushTimer = Stopwatch.StartNew();
 						await writer.FlushAsync(syncNode, cancellationToken);
+						flushTimer.Stop();
 						deleteIncRef = true;
 					}
 
-					// Find the next path to sync
-					const long MaxBatchSize = 10L * 1024 * 1024 * 1024;
-					(int dirIdx, Utf8String path, long size) = GetSyncBatch(directories, MaxBatchSize, clientInfo.ServerInfo.Utf8PathComparer, _logger);
-					syncedSize += size;
+					// Find the next paths to sync
+					const long MaxBatchSize = 1L * 1024 * 1024 * 1024;
 
-					// If we're syncing the entire view, use the given filter instead. It will be tighter than the base path from all files.
-					if (filter != null && path.Equals("..."))
+					int dirIdx = directories.Count - 1;
+					long size = directories[dirIdx]._size;
+
+					for (; dirIdx > 0; dirIdx--)
 					{
-						path = filter;
+						long nextSize = size + directories[dirIdx - 1]._size;
+						if (size > 0 && nextSize > MaxBatchSize)
+						{
+							break;
+						}
+						size = nextSize;
 					}
 
-					const long FlushInterval = 1024 * 1024 * 256;
+					syncedSize += size;
+					double syncPct = (totalSize == 0) ? 100.0 : (syncedSize * 100.0) / totalSize;
+					_logger.LogInformation("Syncing {StreamId} to {Change} [{SyncPct:n1}%] ({Size:n1}mb)", stream.Id, change, syncPct, size / (1024.0 * 1024.0));
+
+					// Copy them to a separate list and remove any redundant paths
+					List<string> syncPaths = new List<string>();
+					for (int idx = dirIdx; idx < directories.Count; idx++)
+					{
+						Utf8String basePath = directories[idx]._path;
+						syncPaths.Add($"//{clientInfo.Client.Name}/{basePath}...@{change}");
 
-					long dataSize = 0;
-					long flushDataSize = FlushInterval;
+						long dirSize = directories[idx]._size;
+						while (idx + 1 < directories.Count && directories[idx + 1]._path.StartsWith(basePath, clientInfo.ServerInfo.Utf8PathComparer))
+						{
+							dirSize += directories[idx + 1]._size;
+							idx++;
+						}
 
-					string syncPath = $"//{clientInfo.Client.Name}/{path}@{change}";
-					double syncPct = (totalSize == 0) ? 100.0 : (syncedSize * 100.0) / totalSize;
-					_logger.LogInformation("Syncing {StreamId} to {Change} [{SyncPct:n1}%]: {Path} ({Size:n1}mb)", stream.Id, change, syncPct, path, size / (1024.0 * 1024.0));
+						_logger.LogInformation("  {Directory} ({Size:n1}mb)", directories[idx]._path + "...", dirSize / (1024.0 * 1024.0));
+					}
 
 					Stopwatch syncTimer = Stopwatch.StartNew();
 					Stopwatch processTimer = new Stopwatch();
@@ -742,7 +762,7 @@ namespace Horde.Build.Perforce
 					Dictionary<int, FileWriter> handles = new Dictionary<int, FileWriter>();
 					try
 					{
-						await foreach (PerforceResponse response in perforce.StreamCommandAsync("sync", Array.Empty<string>(), new string[] { syncPath }, null, typeof(SyncRecord), true, default))
+						await foreach (PerforceResponse response in perforce.StreamCommandAsync("sync", Array.Empty<string>(), syncPaths, null, typeof(SyncRecord), true, default))
 						{
 							PerforceError? error = response.Error;
 							if (error != null)
@@ -777,7 +797,6 @@ namespace Horde.Build.Perforce
 								{
 									FileWriter file = handles[io.File];
 									await file.AppendAsync(io.Payload, Options.Chunking, writer, cancellationToken);
-									dataSize += io.Payload.Length;
 								}
 								else if (io.Command == PerforceIoCommand.Close)
 								{
@@ -815,57 +834,40 @@ namespace Horde.Build.Perforce
 							handle.Dispose();
 						}
 					}
-					_logger.LogInformation("Completed batch in {TimeSeconds:n1}s ({ProcessTimeSeconds:n1}s processing)", syncTimer.Elapsed.TotalSeconds, processTimer.Elapsed.TotalSeconds);
+
+					TimeSpan stallTime = (perforce as NativePerforceConnection)?.StallTime ?? TimeSpan.Zero;
+					_logger.LogInformation("Completed batch in {TimeSeconds:n1}s ({ProcessTimeSeconds:n1}s processing, {StallTimeSeconds:n1}s stalled, {Throughput:n1}mb/s)", syncTimer.Elapsed.TotalSeconds, processTimer.Elapsed.TotalSeconds, stallTime.TotalSeconds, size / (1024.0 * 1024.0 * syncTimer.Elapsed.TotalSeconds));
 
 					// Combine all the existing sync paths together with a new wildcard.
-					List<Utf8String> paths = syncNode.Paths;
-					if (paths.Count > 0)
+					while (directories.Count > dirIdx)
 					{
-						Utf8String lastPath = syncNode.Paths[^1];
-						for (int endIdx = 0; endIdx < lastPath.Length; endIdx++)
+						Utf8String nextPath = directories[^1]._path;
+						if (syncNode.Paths.Count > 0)
 						{
-							if (lastPath[endIdx] == '/')
+							Utf8String lastPath = syncNode.Paths[^1];
+							for (int endIdx = 0; endIdx < lastPath.Length; endIdx++)
 							{
-								Utf8String prefix = lastPath.Substring(0, endIdx + 1);
-								if (!path.StartsWith(prefix, clientInfo.ServerInfo.Utf8PathComparer))
+								if (lastPath[endIdx] == '/')
 								{
-									// Find all the paths that start with this prefix
-									int pathIdx = paths.Count - 1;
-									while (pathIdx > 0 && paths[pathIdx - 1].StartsWith(prefix, clientInfo.ServerInfo.Utf8PathComparer))
+									Utf8String prefix = lastPath.Substring(0, endIdx + 1);
+									if (!nextPath.StartsWith(prefix, clientInfo.ServerInfo.Utf8PathComparer))
 									{
-										pathIdx--;
-									}
+										// Remove any paths that start with this prefix
+										while (syncNode.Paths.Count > 0 && syncNode.Paths[^1].StartsWith(prefix, clientInfo.ServerInfo.Utf8PathComparer))
+										{
+											syncNode.Paths.RemoveAt(syncNode.Paths.Count - 1);
+										}
 
-									// Replace the paths with a wildcard
-									paths.RemoveRange(pathIdx, paths.Count - pathIdx);
-									paths.Add(prefix + "...");
-									break;
+										// Replace it with a wildcard
+										syncNode.Paths.Add(prefix + "...");
+										break;
+									}
 								}
 							}
 						}
+						syncNode.Paths.Add(nextPath + "...");
+						directories.RemoveAt(directories.Count - 1);
 					}
-					paths.Add(path);
-
-					// Print the new state
-					if (dataSize > flushDataSize)
-					{
-						_logger.LogInformation("Queuing flush of working set after receiving {Size:n1}mb.", dataSize / (1024 * 1024));
-						foreach (Utf8String nextPath in paths)
-						{
-							_logger.LogInformation("Filter: {Path}", nextPath);
-						}
-
-						Stopwatch flushTimer = Stopwatch.StartNew();
-						await flushTask;
-						flushTask = writer.FlushAsync(syncNode, cancellationToken);
-						flushTimer.Stop();
-
-						flushDataSize = dataSize + FlushInterval;
-						_logger.LogInformation("Flush took {Time:n1}s. Next flush at {Size:n1}mb.", flushTimer.Elapsed.TotalSeconds, flushDataSize / (1024 * 1024));
-					}
-
-					// Remove the paths we've synced
-					directories.RemoveRange(dirIdx, directories.Count - dirIdx);
 				}
 				await flushTask;
 			}
@@ -894,55 +896,44 @@ namespace Horde.Build.Perforce
 			return fileIdx;
 		}
 
-		static Utf8String NormalizePathSeparators(Utf8String path)
+		static DirectoryToSync FindOrAddDirectoryTree(Dictionary<Utf8String, DirectoryToSync> pathToDirectory, Utf8String directoryPath, Utf8StringComparer comparer)
 		{
-			byte[] newPath = new byte[path.Length];
-			for (int idx = 0; idx < path.Length; idx++)
+			DirectoryToSync? directory;
+			if (!pathToDirectory.TryGetValue(directoryPath, out directory))
 			{
-				newPath[idx] = (path[idx] == '\\') ? (byte)'/' : path[idx];
+				// Add a new path
+				Utf8String normalizedPath = NormalizePathSeparators(directoryPath);
+				directory = new DirectoryToSync(normalizedPath, comparer);
+				pathToDirectory.Add(directoryPath, directory);
+
+				// Also add all the parent directories. This makes the logic for combining wildcards simpler.
+				for (int idx = normalizedPath.Length - 2; idx > 0; idx--)
+				{
+					if (normalizedPath[idx] == '/')
+					{
+						Utf8String parentDirectoryPath = directoryPath.Substring(0, idx + 1);
+						if (pathToDirectory.ContainsKey(parentDirectoryPath))
+						{
+							break;
+						}
+						else
+						{
+							pathToDirectory.Add(parentDirectoryPath, new DirectoryToSync(normalizedPath.Substring(0, idx + 1), comparer));
+						}
+					}
+				}
 			}
-			return new Utf8String(newPath);
+			return directory;
 		}
 
-		static (int, Utf8String, long) GetSyncBatch(List<DirectoryToSync> directories, long maxSize, Utf8StringComparer comparer, ILogger logger)
+		static Utf8String NormalizePathSeparators(Utf8String path)
 		{
-			int idx = directories.Count - 1;
-			long size = directories[idx]._size;
-
-			Utf8String prefix = directories[idx]._path;
-			while (prefix.Length > 0)
+			byte[] newPath = new byte[path.Length];
+			for (int idx = 0; idx < path.Length; idx++)
 			{
-				// Get the length of the parent directory
-				int endIdx = prefix.Length - 1;
-				while (endIdx > 0 && prefix[endIdx - 1] != (byte)'/' && prefix[endIdx - 1] != (byte)'\\')
-				{
-					endIdx--;
-				}
-
-				// Get the parent directory prefix, ending with a slash
-				Utf8String nextPrefix = prefix.Substring(0, endIdx);
-
-				// Include all files with this prefix
-				int nextIdx = idx;
-				long nextSize = size;
-				while (nextIdx > 0 && directories[nextIdx - 1]._path.StartsWith(nextPrefix, comparer))
-				{
-					nextSize += directories[nextIdx - 1]._size;
-					nextIdx--;
-				}
-				if (nextSize > maxSize)
-				{
-					logger.LogInformation("Next filter {NextFilter} would exceed max size ({NextSize:n1}mb > {MaxSize:n1}mb); using {Filter} ({Size:n1}mb) instead.", nextPrefix + "...", nextSize / (1024.0 * 1024.0), maxSize / (1024.0 * 1024.0), prefix, size / (1024.0 * 1024.0));
-					return (idx, prefix + "...", size);
-				}
-				size = nextSize;
-				idx = nextIdx;
-
-				// Update the prefix
-				prefix = nextPrefix;
+				newPath[idx] = (path[idx] == '\\') ? (byte)'/' : path[idx];
 			}
-
-			return (idx, prefix + "...", size);
+			return new Utf8String(newPath);
 		}
 
 		async Task FlushWorkspaceAsync(ReplicationClient clientInfo, IPerforceConnection perforce, int change)
diff --git a/Engine/Source/Programs/Shared/EpicGames.Core/BuzHash.cs b/Engine/Source/Programs/Shared/EpicGames.Core/BuzHash.cs
--- a/Engine/Source/Programs/Shared/EpicGames.Core/BuzHash.cs
+++ b/Engine/Source/Programs/Shared/EpicGames.Core/BuzHash.cs
@@ -1,6 +1,9 @@
 // Copyright Epic Games, Inc. All Rights Reserved.
 
 using System;
+using System.Buffers;
+using System.Diagnostics;
+using System.Runtime.InteropServices;
 
 namespace EpicGames.Core
 {
@@ -37,25 +40,22 @@ namespace EpicGames.Core
 			0xf9c18d66, 0x593ade65, 0xd95ddf11,
 		};
 
-		int _count;
 		uint _state;
 
 		/// <summary>
 		/// Constructor
 		/// </summary>
 		public BuzHash()
-			: this(0, 0)
+			: this(0)
 		{
 		}
 
 		/// <summary>
 		/// Constructor
 		/// </summary>
-		/// <param name="count"></param>
 		/// <param name="state"></param>
-		public BuzHash(int count, uint state)
+		public BuzHash(uint state)
 		{
-			_count = count;
 			_state = state;
 		}
 
@@ -64,7 +64,6 @@ namespace EpicGames.Core
 		/// </summary>
 		public void Reset()
 		{ 
-			_count = 0; 
 			_state = 0; 
 		}
 
@@ -75,7 +74,6 @@ namespace EpicGames.Core
 		public void Add(byte value)
 		{
 			_state = Add(_state, value);
-			_count++;
 		}
 
 		/// <summary>
@@ -85,7 +83,6 @@ namespace EpicGames.Core
 		public void Add(ReadOnlySpan<byte> span)
 		{
 			_state = Add(_state, span);
-			_count += span.Length;
 		}
 
 		/// <summary>
@@ -105,11 +102,15 @@ namespace EpicGames.Core
 		/// <param name="state">The current hash value</param>
 		/// <param name="data">Data to append to the hash</param>
 		/// <returns>New hash value</returns>
-		public static uint Add(uint state, ReadOnlySpan<byte> data)
+		public static unsafe uint Add(uint state, ReadOnlySpan<byte> data)
 		{
-			for (int idx = 0; idx < data.Length; idx++)
+			fixed (uint* table = s_table)
+			fixed (byte* dataPtr = data)
 			{
-				state = Add(state, data[idx]);
+				for (int idx = 0; idx < data.Length; idx++)
+				{
+					state = Rol32(state, 1) ^ table[dataPtr[idx]];
+				}
 			}
 			return state;
 		}
@@ -118,10 +119,10 @@ namespace EpicGames.Core
 		/// Removes a byte from the start of the hash window
 		/// </summary>
 		/// <param name="value">Value at the start of the window</param>
-		public void Sub(byte value)
+		/// <param name="count">Number of items in the hash</param>
+		public void Sub(byte value, int count)
 		{
-			_state ^= Rol32(s_table[value], (int)(_count - 1));
-			_count--;
+			_state ^= Rol32(s_table[value], (int)(count - 1));
 		}
 
 		/// <summary>
@@ -147,6 +148,38 @@ namespace EpicGames.Core
 			}
 		}
 
+		/// <summary>
+		/// Consume data from a source until the hash is below a certain threshold 
+		/// </summary>
+		/// <param name="prev">Previous data added to the hash</param>
+		/// <param name="next">Next data to be added to the hash</param>
+		/// <param name="threshold">Threshold for hash below which to terminate the scan</param>
+		/// <returns>Number of bytes consumed from next</returns>
+		public static unsafe int Update(ReadOnlySpan<byte> prev, ReadOnlySpan<byte> next, uint threshold, ref uint state)
+		{
+			Debug.Assert(prev.Length == next.Length);
+
+			fixed (uint* table = s_table)
+			fixed (byte* prevPtr = prev)
+			fixed (byte* nextPtr = next)
+			{
+				uint hash = state;
+				for (int idx = 0; idx < prev.Length; idx++)
+				{
+					hash = Rol32(hash, 1) ^ table[nextPtr[idx]];
+					if (hash < threshold)
+					{
+						state = hash;
+						return idx;
+					}
+					hash = hash ^ Rol32(table[prevPtr[idx]], prev.Length - 1);
+				}
+				state = hash;
+			}
+
+			return -1;
+		}
+
 		/// <summary>
 		/// Gets the current hash value
 		/// </summary>
diff --git a/Engine/Source/Programs/Shared/EpicGames.Horde/Storage/Nodes/DirectoryNode.cs b/Engine/Source/Programs/Shared/EpicGames.Horde/Storage/Nodes/DirectoryNode.cs
--- a/Engine/Source/Programs/Shared/EpicGames.Horde/Storage/Nodes/DirectoryNode.cs
+++ b/Engine/Source/Programs/Shared/EpicGames.Horde/Storage/Nodes/DirectoryNode.cs
@@ -551,6 +551,10 @@ namespace EpicGames.Horde.Storage.Nodes
 			{
 				size += (directoryEntry.Name.Length + 1) + VarInt.MeasureUnsigned((ulong)directoryEntry.Length);
 			}
+			if ((flags & DirectoryFlags.WithGitHashes) != 0)
+			{
+				size += (directoryEntries.Count + fileEntries.Count) * IoHash.NumBytes;
+			}
 
 			// Allocate the buffer and copy the node to it
 			byte[] data = new byte[size];
diff --git a/Engine/Source/Programs/Shared/EpicGames.Horde/Storage/Nodes/FileNode.cs b/Engine/Source/Programs/Shared/EpicGames.Horde/Storage/Nodes/FileNode.cs
--- a/Engine/Source/Programs/Shared/EpicGames.Horde/Storage/Nodes/FileNode.cs
+++ b/Engine/Source/Programs/Shared/EpicGames.Horde/Storage/Nodes/FileNode.cs
@@ -343,43 +343,37 @@ namespace EpicGames.Horde.Storage.Nodes
 			uint newRollingHash = _rollingHash;
 			uint rollingHashThreshold = (uint)((1L << 32) / options.LeafOptions.TargetSize);
 
-			// Length of the data taken from the input span, updated as we step through it.
-			int length = 0;
+			// Offset within the head span, updated as we step through it.
+			int offset = 0;
 
 			// Step the window through the tail end of the existing payload window. In this state, update the hash to remove data from the current payload, and add data from the new payload.
-			int splitLength = Math.Min(headSpan.Length, windowSize);
-			ReadOnlySequence<byte> tailSequence = _writtenSequence.Slice(_writtenSequence.Length - windowSize, splitLength);
+			int tailLength = Math.Min(headSpan.Length, windowSize);
+			ReadOnlySequence<byte> tailSequence = _writtenSequence.Slice(_writtenSequence.Length - windowSize, tailLength);
 
 			foreach (ReadOnlyMemory<byte> tailSegment in tailSequence)
 			{
-				int baseLength = length;
-				int spanLength = length + tailSegment.Length;
-
-				ReadOnlySpan<byte> tailSpan = tailSegment.Span;
-				for (; length < spanLength; length++)
+				int count = BuzHash.Update(tailSegment.Span, headSpan.Slice(offset, tailSegment.Length), rollingHashThreshold, ref newRollingHash);
+				if (count != -1)
 				{
-					newRollingHash = BuzHash.Add(newRollingHash, headSpan[length]);
-					if (newRollingHash < rollingHashThreshold)
-					{
-						AppendLeafData(headSpan.Slice(0, length), newRollingHash);
-						_isReadOnly = true;
-						return length;
-					}
-					newRollingHash = BuzHash.Sub(newRollingHash, tailSpan[length - baseLength], windowSize);
+					offset += count;
+					AppendLeafData(headSpan.Slice(0, offset), newRollingHash);
+					_isReadOnly = true;
+					return offset;
 				}
+				offset += tailSegment.Length;
 			}
 
 			// Step through the new window until we get to a chunk boundary.
-			for (; length < headSpan.Length; length++)
+			if (offset < headSpan.Length)
 			{
-				newRollingHash = BuzHash.Add(newRollingHash, headSpan[length]);
-				if (newRollingHash < rollingHashThreshold)
+				int count = BuzHash.Update(headSpan.Slice(offset - windowSize, headSpan.Length - offset), headSpan.Slice(offset), rollingHashThreshold, ref newRollingHash);
+				if (count != -1)
 				{
-					AppendLeafData(headSpan.Slice(0, length), newRollingHash);
+					offset += count;
+					AppendLeafData(headSpan.Slice(0, offset), newRollingHash);
 					_isReadOnly = true;
-					return length;
+					return offset;
 				}
-				newRollingHash = BuzHash.Sub(newRollingHash, headSpan[length - windowSize], windowSize);
 			}
 
 			// Otherwise just append all the data.
diff --git a/Engine/Source/Programs/Shared/EpicGames.Perforce/NativePerforceConnection.cs b/Engine/Source/Programs/Shared/EpicGames.Perforce/NativePerforceConnection.cs
--- a/Engine/Source/Programs/Shared/EpicGames.Perforce/NativePerforceConnection.cs
+++ b/Engine/Source/Programs/Shared/EpicGames.Perforce/NativePerforceConnection.cs
@@ -268,6 +268,7 @@ namespace EpicGames.Perforce
 		readonly BlockingCollection<PinnedBuffer> _writeBuffers = new BlockingCollection<PinnedBuffer>();
 		Response? _currentResponse;
 		readonly ManualResetEvent _responseCompleteEvent;
+		readonly Stopwatch _stallTimer = new Stopwatch();
 
 		/// <inheritdoc/>
 		public IPerforceSettings Settings { get; }
@@ -335,6 +336,8 @@ namespace EpicGames.Perforce
 
 		void GetNextWriteBuffer(NativeWriteBuffer nativeWriteBuffer, int minSize)
 		{
+			_stallTimer.Start();
+
 			PinnedBuffer buffer = _writeBuffers.Take();
 			if (buffer.MaxLength < minSize)
 			{
@@ -344,6 +347,8 @@ namespace EpicGames.Perforce
 			nativeWriteBuffer._data = buffer.BasePtr;
 			nativeWriteBuffer._maxLength = buffer.Data.Length;
 			nativeWriteBuffer._maxCount = Int32.MaxValue;
+
+			_stallTimer.Stop();
 		}
 
 		/// <summary>
@@ -394,6 +399,11 @@ namespace EpicGames.Perforce
 			}
 		}
 
+		/// <summary>
+		/// Gets the amount of time stalled waiting for an output buffer in the last command
+		/// </summary>
+		public TimeSpan StallTime => _stallTimer.Elapsed;
+
 		/// <summary>
 		/// Initializes the connection, throwing an error on failure
 		/// </summary>
@@ -469,6 +479,7 @@ namespace EpicGames.Perforce
 				_currentResponse = request.Value.Response;
 				_responseCompleteEvent.Reset();
 
+				_stallTimer.Reset();
 				request.Value.Action();
 
 				_currentResponse._readBuffers.Writer.TryComplete();
