Horde: Add a command to synthesize a Git repo from data stored in a repository snapshot.

diff --git a/Engine/Source/Programs/Horde/Horde.Build/Commands/Bundles/RepoCommand.cs b/Engine/Source/Programs/Horde/Horde.Build/Commands/Bundles/RepoCommand.cs
new file mode 100644
--- /dev/null
+++ b/Engine/Source/Programs/Horde/Horde.Build/Commands/Bundles/RepoCommand.cs
@@ -0,0 +1,141 @@
+// Copyright Epic Games, Inc. All Rights Reserved.
+
+using System;
+using System.Buffers;
+using System.Collections.Generic;
+using System.IO;
+using System.IO.Compression;
+using System.Security.Cryptography;
+using System.Threading;
+using System.Threading.Tasks;
+using EpicGames.Core;
+using EpicGames.Horde.Storage;
+using EpicGames.Horde.Storage.Bundles;
+using EpicGames.Horde.Storage.Git;
+using EpicGames.Horde.Storage.Nodes;
+using Horde.Build.Perforce;
+using Microsoft.Extensions.Caching.Memory;
+using Microsoft.Extensions.Configuration;
+using Microsoft.Extensions.DependencyInjection;
+using Microsoft.Extensions.Logging;
+
+namespace Horde.Build.Commands.Bundles
+{
+	[Command("bundle", "repo", "Synthesizes a Git repo from a bundle to the local hard drive")]
+	class RepoCommand : Command
+	{
+		[CommandLine("-Ref=")]
+		public RefName RefName { get; set; } = new RefName("default-ref");
+
+		[CommandLine("-OutputDir=", Required = true)]
+		public DirectoryReference OutputDir { get; set; } = null!;
+
+		protected readonly IConfiguration _configuration;
+		protected readonly ILoggerProvider _loggerProvider;
+
+		public RepoCommand(IConfiguration configuration, ILoggerProvider loggerProvider)
+		{
+			_configuration = configuration;
+			_loggerProvider = loggerProvider;
+		}
+
+		public override async Task<int> ExecuteAsync(ILogger logger)
+		{
+			using ServiceProvider serviceProvider = Startup.CreateServiceProvider(_configuration, _loggerProvider);
+
+			ITreeStore store = serviceProvider.GetRequiredService<ITreeStore<ReplicationService>>();
+
+			ReplicationNode node = await store.ReadTreeAsync<ReplicationNode>(RefName);
+			DirectoryNode root = await node.Contents.ExpandAsync();
+
+			Sha1Hash tree = await WriteTreeRecursiveAsync(root, CancellationToken.None);
+			Sha1Hash commit = await WriteCommitObjectAsync(tree, CancellationToken.None);
+			logger.LogInformation("Commit: {Sha}", commit);
+
+			await WriteTextFileAsync("HEAD", "ref: refs/heads/main");
+			await WriteTextFileAsync("refs/heads/main", commit.ToString());
+
+			return 0;
+		}
+
+		async Task<Sha1Hash> WriteTreeRecursiveAsync(DirectoryNode root, CancellationToken cancellationToken)
+		{
+			Sha1Hash hash = await WriteTreeObjectAsync(root, cancellationToken);
+
+			foreach (DirectoryEntry entry in root.Directories)
+			{
+				DirectoryNode dir = await entry.ExpandAsync(cancellationToken);
+				await WriteTreeRecursiveAsync(dir, cancellationToken);
+			}
+
+			foreach (FileEntry entry in root.Files)
+			{
+				await WriteBlobObjectAsync(entry, cancellationToken);
+			}
+
+			return hash;
+		}
+
+		async Task<Sha1Hash> WriteCommitObjectAsync(Sha1Hash tree, CancellationToken cancellationToken)
+		{
+			GitCommit commit = new GitCommit(tree, "Test");
+			byte[] data = commit.Serialize();
+			return await WriteObjectAsync(data, cancellationToken);
+		}
+
+		async Task<Sha1Hash> WriteTreeObjectAsync(DirectoryNode dir, CancellationToken cancellationToken)
+		{
+			GitTree tree = dir.AsGitTree();
+			byte[] data = tree.Serialize();
+			return await WriteObjectAsync(data, cancellationToken);
+		}
+
+		async Task WriteBlobObjectAsync(FileEntry entry, CancellationToken cancellationToken)
+		{
+			FileReference output = GetBlobPath(OutputDir, entry.GitHash);
+			DirectoryReference.CreateDirectory(output.Directory);
+
+			using (FileStream stream = FileReference.Open(output, FileMode.Create))
+			{
+				using (ZLibStream compressed = new ZLibStream(stream, CompressionLevel.Fastest, true))
+				{
+					await GitObject.WriteHeaderAsync(compressed, GitObjectType.Blob, entry.Length, cancellationToken);
+
+					FileNode node = await entry.ExpandAsync(cancellationToken);
+					await node.CopyToStreamAsync(compressed, cancellationToken);
+				}
+			}
+		}
+
+		async Task<Sha1Hash> WriteObjectAsync(byte[] data, CancellationToken cancellationToken)
+		{
+			Sha1Hash hash = Sha1Hash.Compute(data);
+
+			FileReference output = GetBlobPath(OutputDir, hash);
+			DirectoryReference.CreateDirectory(output.Directory);
+
+			using (FileStream stream = FileReference.Open(output, FileMode.Create))
+			{
+				using (ZLibStream compressed = new ZLibStream(stream, CompressionLevel.Fastest, true))
+				{
+					await compressed.WriteAsync(data, cancellationToken);
+				}
+			}
+
+			return hash;
+		}
+
+		static FileReference GetBlobPath(DirectoryReference rootDir, Sha1Hash hash)
+		{
+			string str = hash.ToString();
+			return FileReference.Combine(rootDir, $".git/objects/{str.Substring(0, 2)}/{str.Substring(2)}");
+		}
+
+		async Task WriteTextFileAsync(string path, string text)
+		{
+			FileReference file = FileReference.Combine(OutputDir, ".git", path);
+			DirectoryReference.CreateDirectory(file.Directory);
+			await FileReference.WriteAllTextAsync(file, text);
+		}
+	}
+}
diff --git a/Engine/Source/Programs/Horde/Horde.Build/Perforce/ReplicationService.cs b/Engine/Source/Programs/Horde/Horde.Build/Perforce/ReplicationService.cs
--- a/Engine/Source/Programs/Horde/Horde.Build/Perforce/ReplicationService.cs
+++ b/Engine/Source/Programs/Horde/Horde.Build/Perforce/ReplicationService.cs
@@ -107,10 +107,10 @@ namespace Horde.Build.Perforce
 		}
 
 		/// <inheritdoc/>
-		public override NewTreeBlob Serialize()
+		public override async Task<ITreeBlob> SerializeAsync(ITreeWriter writer, CancellationToken cancellationToken)
 		{
-			List<TreeNodeRef> references = new List<TreeNodeRef>();
-			references.Add(Contents);
+			List<ITreeBlobRef> references = new List<ITreeBlobRef>();
+			references.Add(await Contents.CollapseAsync(writer, cancellationToken));
 
 			ByteArrayBuilder builder = new ByteArrayBuilder();
 			builder.WriteVariableLengthArray(Paths, x => builder.WriteUtf8String(x));
diff --git a/Engine/Source/Programs/Shared/EpicGames.Horde/Storage/Git/GitCommit.cs b/Engine/Source/Programs/Shared/EpicGames.Horde/Storage/Git/GitCommit.cs
new file mode 100644
--- /dev/null
+++ b/Engine/Source/Programs/Shared/EpicGames.Horde/Storage/Git/GitCommit.cs
@@ -0,0 +1,155 @@
+// Copyright Epic Games, Inc. All Rights Reserved.
+
+using System;
+using System.Collections.Generic;
+using System.Diagnostics;
+using System.Linq;
+using System.Security.Cryptography;
+using EpicGames.Core;
+
+namespace EpicGames.Horde.Storage.Git
+{
+	/// <summary>
+	/// Property names for git commits
+	/// </summary>
+	public static class GitCommitProperty
+	{
+		/// <summary>
+		/// The root of the tree for this commit
+		/// </summary>
+		public static Utf8String Tree { get; } = new Utf8String("tree");
+
+		/// <summary>
+		/// Author of the commit
+		/// </summary>
+		public static Utf8String Author { get; } = new Utf8String("author");
+
+		/// <summary>
+		/// Person that committed the change
+		/// </summary>
+		public static Utf8String Committer { get; } = new Utf8String("committer");
+
+		/// <summary>
+		/// Parent commit
+		/// </summary>
+		public static Utf8String Parent { get; } = new Utf8String("parent");
+	}
+
+	/// <summary>
+	/// Representation of a Git commit object
+	/// </summary>
+	public class GitCommit
+	{
+		/// <summary>
+		/// Properties for the commit
+		/// </summary>
+		public List<KeyValuePair<Utf8String, Utf8String>> Properties { get; } = new List<KeyValuePair<Utf8String, Utf8String>>();
+			
+		/// <summary>
+		/// Commit messages
+		/// </summary>
+		public Utf8String Message { get; set; }
+
+		/// <summary>
+		/// Constructor
+		/// </summary>
+		/// <param name="tree"></param>
+		/// <param name="message"></param>
+		public GitCommit(Sha1Hash tree, Utf8String message)
+		{
+			AddProperty(GitCommitProperty.Tree, tree.ToUtf8String());
+			Message = message;
+		}
+
+		/// <summary>
+		/// Adds a new property to the collection
+		/// </summary>
+		/// <param name="name"></param>
+		/// <param name="value"></param>
+		public void AddProperty(Utf8String name, Utf8String value)
+		{
+			Properties.Add(new KeyValuePair<Utf8String, Utf8String>(name, value));
+		}
+
+		/// <summary>
+		/// Gets a property value with the given key name
+		/// </summary>
+		/// <param name="name">Name of the property</param>
+		/// <returns>The property value</returns>
+		public Utf8String GetProperty(Utf8String name)
+		{
+			foreach (KeyValuePair<Utf8String, Utf8String> pair in Properties)
+			{
+				if (pair.Key == name)
+				{
+					return pair.Value;
+				}
+			}
+			return Utf8String.Empty;
+		}
+
+		/// <summary>
+		/// Gets property values with the given key name
+		/// </summary>
+		/// <param name="name">Name of the property</param>
+		/// <returns>The property value</returns>
+		public IEnumerable<Utf8String> FindProperties(Utf8String name)
+		{
+			foreach (KeyValuePair<Utf8String, Utf8String> pair in Properties)
+			{
+				if (pair.Key == name)
+				{
+					yield return pair.Value;
+				}
+			}
+		}
+
+		/// <summary>
+		/// The tree for this commit
+		/// </summary>
+		public Sha1Hash GetTree() => Sha1Hash.Parse(GetProperty(GitCommitProperty.Tree));
+
+		/// <summary>
+		/// Parents of this commit
+		/// </summary>
+		public IEnumerable<Sha1Hash> GetParents() => FindProperties(GitCommitProperty.Parent).Select(x => Sha1Hash.Parse(x));
+
+		/// <summary>
+		/// Serializes this object
+		/// </summary>
+		/// <returns></returns>
+		public byte[] Serialize()
+		{
+			int size = Message.Length + 1 + Properties.Sum(x => x.Key.Length + x.Value.Length + 2);
+
+			Span<byte> header = GitObject.WriteHeader(GitObjectType.Commit, size, stackalloc byte[32]);
+
+			byte[] data = new byte[header.Length + size];
+			header.CopyTo(data);
+
+			Span<byte> output = data.AsSpan(header.Length);
+			foreach (KeyValuePair<Utf8String, Utf8String> pair in Properties)
+			{
+				pair.Key.Span.CopyTo(output);
+				output = output.Slice(pair.Key.Length);
+
+				output[0] = (byte)' ';
+				output = output.Slice(1);
+
+				pair.Value.Span.CopyTo(output);
+				output = output.Slice(pair.Value.Length);
+
+				output[0] = (byte)'\n';
+				output = output.Slice(1);
+			}
+
+			output[0] = (byte)'\n';
+			output = output.Slice(1);
+
+			Message.Span.CopyTo(output);
+			Debug.Assert(output.Length == Message.Length);
+
+			return data;
+		}
+	}
+}
diff --git a/Engine/Source/Programs/Shared/EpicGames.Horde/Storage/Git/GitObject.cs b/Engine/Source/Programs/Shared/EpicGames.Horde/Storage/Git/GitObject.cs
--- a/Engine/Source/Programs/Shared/EpicGames.Horde/Storage/Git/GitObject.cs
+++ b/Engine/Source/Programs/Shared/EpicGames.Horde/Storage/Git/GitObject.cs
@@ -1,7 +1,10 @@
 // Copyright Epic Games, Inc. All Rights Reserved.
 
 using System;
+using System.IO;
 using System.Security.Cryptography;
+using System.Threading;
+using System.Threading.Tasks;
 using EpicGames.Core;
 
 namespace EpicGames.Horde.Storage.Git
@@ -20,6 +23,11 @@ namespace EpicGames.Horde.Storage.Git
 		/// Tree objects
 		/// </summary>
 		public static Utf8String Tree { get; } = "tree";
+
+		/// <summary>
+		/// Commit objects
+		/// </summary>
+		public static Utf8String Commit { get; } = "commit";
 	}
 
 	/// <summary>
@@ -27,6 +35,21 @@ namespace EpicGames.Horde.Storage.Git
 	/// </summary>
 	public static class GitObject
 	{
+		/// <summary>
+		/// Writes the header for an object to a stream
+		/// </summary>
+		/// <param name="stream"></param>
+		/// <param name="type">Type of the object</param>
+		/// <param name="size">Size of the object</param>
+		/// <param name="cancellationToken"></param>
+		/// <returns></returns>
+		public static async Task WriteHeaderAsync(Stream stream, Utf8String type, long size, CancellationToken cancellationToken = default)
+		{
+			byte[] header = new byte[32];
+			int length = GitObject.WriteHeader(type, size, header).Length;
+			await stream.WriteAsync(header.AsMemory(0, length), cancellationToken);
+		}
+
 		/// <summary>
 		/// Appends data for a header to the hash
 		/// </summary>
diff --git a/Engine/Source/Programs/Shared/EpicGames.Horde/Storage/Nodes/DirectoryNode.cs b/Engine/Source/Programs/Shared/EpicGames.Horde/Storage/Nodes/DirectoryNode.cs
--- a/Engine/Source/Programs/Shared/EpicGames.Horde/Storage/Nodes/DirectoryNode.cs
+++ b/Engine/Source/Programs/Shared/EpicGames.Horde/Storage/Nodes/DirectoryNode.cs
@@ -517,23 +517,24 @@ namespace EpicGames.Horde.Storage.Nodes
 		#endregion
 
 		/// <inheritdoc/>
-		public override NewTreeBlob Serialize()
+		public override async Task<ITreeBlob> SerializeAsync(ITreeWriter writer, CancellationToken cancellationToken)
 		{
-			List<TreeNodeRef> refs = new List<TreeNodeRef>();
+			List<ITreeBlobRef> refs = new List<ITreeBlobRef>();
 
 			List<FileEntry> fileEntries = _nameToFileEntry.Values.OrderBy(x => x.Name).ToList();
 			foreach (FileEntry fileEntry in fileEntries)
 			{
-				refs.Add(fileEntry);
+				refs.Add(await fileEntry.CollapseAsync(writer, cancellationToken));
 			}
 
 			List<DirectoryEntry> directoryEntries = _nameToDirectoryEntry.Values.OrderBy(x => x.Name).ToList();
 			foreach (DirectoryEntry directoryEntry in directoryEntries)
 			{
-				refs.Add(directoryEntry);
+				refs.Add(await directoryEntry.CollapseAsync(writer, cancellationToken));
 			}
 
-			return new NewTreeBlob(SerializeData(Flags, fileEntries, directoryEntries), refs);
+			ReadOnlySequence<byte> data = SerializeData(Flags, fileEntries, directoryEntries);
+			return new NewTreeBlob(data, refs);
 		}
 
 		static ReadOnlySequence<byte> SerializeData(DirectoryFlags flags, List<FileEntry> fileEntries, List<DirectoryEntry> directoryEntries)
diff --git a/Engine/Source/Programs/Shared/EpicGames.Horde/Storage/Nodes/FileNode.cs b/Engine/Source/Programs/Shared/EpicGames.Horde/Storage/Nodes/FileNode.cs
--- a/Engine/Source/Programs/Shared/EpicGames.Horde/Storage/Nodes/FileNode.cs
+++ b/Engine/Source/Programs/Shared/EpicGames.Horde/Storage/Nodes/FileNode.cs
@@ -293,7 +293,7 @@ namespace EpicGames.Horde.Storage.Nodes
 		public override IReadOnlyList<TreeNodeRef> GetReferences() => Array.Empty<TreeNodeRef>();
 
 		/// <inheritdoc/>
-		public override NewTreeBlob Serialize() => new NewTreeBlob(_writtenSequence, Array.Empty<TreeNodeRef>());
+		public override Task<ITreeBlob> SerializeAsync(ITreeWriter writer, CancellationToken cancellationToken) => Task.FromResult<ITreeBlob>(new NewTreeBlob(_writtenSequence, Array.Empty<ITreeBlobRef>()));
 
 		/// <inheritdoc/>
 		public override ValueTask<ReadOnlyMemory<byte>> AppendDataAsync(ReadOnlyMemory<byte> newData, ChunkingOptions options, ITreeWriter writer, CancellationToken cancellationToken)
@@ -478,7 +478,7 @@ namespace EpicGames.Horde.Storage.Nodes
 		public override IReadOnlyList<TreeNodeRef> GetReferences() => _children;
 
 		/// <inheritdoc/>
-		public override NewTreeBlob Serialize()
+		public override async Task<ITreeBlob> SerializeAsync(ITreeWriter writer, CancellationToken cancellationToken)
 		{
 			byte[] data = new byte[1 + sizeof(uint) + VarInt.MeasureUnsigned((ulong)_length) + 1];
 			data[0] = TypeId;
@@ -495,7 +495,14 @@ namespace EpicGames.Horde.Storage.Nodes
 			pos++;
 
 			Debug.Assert(pos == data.Length);
-			return new NewTreeBlob(data, _children);
+
+			List<ITreeBlobRef> childRefs = new List<ITreeBlobRef>();
+			foreach (TreeNodeRef typedRef in _children)
+			{
+				childRefs.Add(await typedRef.CollapseAsync(writer, cancellationToken));
+			}
+
+			return new NewTreeBlob(data, childRefs);
 		}
 
 		/// <inheritdoc/>
diff --git a/Engine/Source/Programs/Shared/EpicGames.Horde/Storage/TreeNode.cs b/Engine/Source/Programs/Shared/EpicGames.Horde/Storage/TreeNode.cs
--- a/Engine/Source/Programs/Shared/EpicGames.Horde/Storage/TreeNode.cs
+++ b/Engine/Source/Programs/Shared/EpicGames.Horde/Storage/TreeNode.cs
@@ -52,7 +52,7 @@ namespace EpicGames.Horde.Storage
 		/// Serialize a node to a block of memory
 		/// </summary>
 		/// <returns>New data to be stored into a blob</returns>
-		public abstract NewTreeBlob Serialize();
+		public abstract Task<ITreeBlob> SerializeAsync(ITreeWriter writer, CancellationToken cancellationToken);
 
 		/// <summary>
 		/// Deserialize a node from data
@@ -69,7 +69,7 @@ namespace EpicGames.Horde.Storage
 	/// <summary>
 	/// Data to be stored in a tree blob
 	/// </summary>
-	public struct NewTreeBlob
+	public struct NewTreeBlob : ITreeBlob
 	{
 		/// <summary>
 		/// The opaque data payload
@@ -79,14 +79,14 @@ namespace EpicGames.Horde.Storage
 		/// <summary>
 		/// References to other tree nodes
 		/// </summary>
-		public IReadOnlyList<TreeNodeRef> Refs { get; }
+		public IReadOnlyList<ITreeBlobRef> Refs { get; }
 
 		/// <summary>
 		/// Constructor
 		/// </summary>
 		/// <param name="data">Data payload</param>
 		/// <param name="refs">References to other nodes</param>
-		public NewTreeBlob(ReadOnlyMemory<byte> data, IReadOnlyList<TreeNodeRef> refs)
+		public NewTreeBlob(ReadOnlyMemory<byte> data, IReadOnlyList<ITreeBlobRef> refs)
 		{
 			Data = new ReadOnlySequence<byte>(data);
 			Refs = refs;
@@ -97,7 +97,7 @@ namespace EpicGames.Horde.Storage
 		/// </summary>
 		/// <param name="data">Data payload</param>
 		/// <param name="refs">References to other nodes</param>
-		public NewTreeBlob(ReadOnlySequence<byte> data, IReadOnlyList<TreeNodeRef> refs)
+		public NewTreeBlob(ReadOnlySequence<byte> data, IReadOnlyList<ITreeBlobRef> refs)
 		{
 			Data = data;
 			Refs = refs;
@@ -148,30 +148,10 @@ namespace EpicGames.Horde.Storage
 		/// <param name="cancellationToken">Cancellation token for the operation</param>
 		public static async Task FlushAsync(this ITreeWriter writer, TreeNode root, CancellationToken cancellationToken = default)
 		{
-			ITreeBlob rootBlob = await writer.WriteTreeAsync(root, cancellationToken);
+			ITreeBlob rootBlob = await root.SerializeAsync(writer, cancellationToken);
 			await writer.FlushAsync(rootBlob.Data, rootBlob.Refs, cancellationToken);
 		}
 
-		/// <summary>
-		/// Write a tree of nodes to an <see cref="ITreeWriter"/>
-		/// </summary>
-		/// <param name="writer">Writer to output the nodes to</param>
-		/// <param name="node">Root node to serialize</param>
-		/// <param name="cancellationToken">Cancellation token for the operation</param>
-		/// <returns>Reference to the new tree of nodes</returns>
-		public static async Task<ITreeBlob> WriteTreeAsync(this ITreeWriter writer, TreeNode node, CancellationToken cancellationToken = default)
-		{
-			NewTreeBlob blob = node.Serialize();
-
-			List<ITreeBlobRef> targetRefs = new List<ITreeBlobRef>();
-			foreach (TreeNodeRef typedRef in blob.Refs)
-			{
-				targetRefs.Add(await typedRef.CollapseAsync(writer, cancellationToken));
-			}
-
-			return TreeBlob.Create(blob.Data, targetRefs);
-		}
-
 		/// <inheritdoc/>
 		public static async Task<T?> TryReadTreeAsync<T>(this ITreeStore store, RefName name, CancellationToken cancellationToken = default) where T : TreeNode
 		{
diff --git a/Engine/Source/Programs/Shared/EpicGames.Horde/Storage/TreeNodeRef.cs b/Engine/Source/Programs/Shared/EpicGames.Horde/Storage/TreeNodeRef.cs
--- a/Engine/Source/Programs/Shared/EpicGames.Horde/Storage/TreeNodeRef.cs
+++ b/Engine/Source/Programs/Shared/EpicGames.Horde/Storage/TreeNodeRef.cs
@@ -181,15 +181,8 @@ namespace EpicGames.Horde.Storage
 			ITreeBlobRef? target = _target;
 			if (target == null)
 			{
-				NewTreeBlob blob = Node!.Serialize();
-
-				List<ITreeBlobRef> targetRefs = new List<ITreeBlobRef>();
-				foreach (TreeNodeRef typedRef in blob.Refs)
-				{
-					targetRefs.Add(await typedRef.CollapseAsync(writer, cancellationToken));
-				}
-
-				target = await writer.WriteNodeAsync(blob.Data, targetRefs, cancellationToken);
+				ITreeBlob blob = await Node!.SerializeAsync(writer, cancellationToken);
+				target = await writer.WriteNodeAsync(blob.Data, blob.Refs, cancellationToken);
 				MarkAsClean(target);
 			}
 			return target;
