Horde: Fix new build issue where default text value is used when user had deleted it
#rnx
#skipci

diff --git a/Engine/Source/Programs/Horde/HordeDashboard/src/components/NewBuild.tsx b/Engine/Source/Programs/Horde/HordeDashboard/src/components/NewBuild.tsx
--- a/Engine/Source/Programs/Horde/HordeDashboard/src/components/NewBuild.tsx
+++ b/Engine/Source/Programs/Horde/HordeDashboard/src/components/NewBuild.tsx
@@ -405,7 +405,8 @@ class BuildParameters {
 
          let value = (this.values[this.paramKey(p)] as string)?.trim();
 
-         value = value ? value : param.default;
+         // default already populated in text entry
+         // value = value ? value : param.default;
 
          if (!value) {
             return;
