Horde: Fix exception when syncing a folder that already exists in the workspace.

diff --git a/Engine/Source/Programs/Horde/Horde.Build/Perforce/ReplicationService.cs b/Engine/Source/Programs/Horde/Horde.Build/Perforce/ReplicationService.cs
--- a/Engine/Source/Programs/Horde/Horde.Build/Perforce/ReplicationService.cs
+++ b/Engine/Source/Programs/Horde/Horde.Build/Perforce/ReplicationService.cs
@@ -767,7 +767,14 @@ namespace Horde.Build.Perforce
 							PerforceError? error = response.Error;
 							if (error != null)
 							{
-								throw new ReplicationException($"Perforce error while replicating content: {error.Data}");
+								if (error.Generic == PerforceGenericCode.Empty)
+								{
+									continue;
+								}
+								else
+								{
+									throw new ReplicationException($"Perforce error while replicating content - {error}");
+								}
 							}
 
 							processTimer.Start();
