Horde: Extract an interface for ComputeService for testing

diff --git a/Engine/Source/Programs/Horde/Horde.Build.Tests/Compute/ComputeServiceTest.cs b/Engine/Source/Programs/Horde/Horde.Build.Tests/Compute/ComputeServiceTest.cs
--- a/Engine/Source/Programs/Horde/Horde.Build.Tests/Compute/ComputeServiceTest.cs
+++ b/Engine/Source/Programs/Horde/Horde.Build.Tests/Compute/ComputeServiceTest.cs
@@ -46,7 +46,7 @@ namespace Horde.Build.Tests.Compute
 		{
 			await CreateTaskWithRequirement("memory >= 10");
 			IPool pool = await CreatePool(new () { {"cpu", "5"} });
-			Assert.AreEqual(0, await ComputeService.GetNumQueuedTasksForPool(_clusterId, pool));
+			Assert.AreEqual(0, await ComputeService.GetNumQueuedTasksForPoolAsync(_clusterId, pool));
 		}
 
 		[TestMethod]
@@ -54,7 +54,7 @@ namespace Horde.Build.Tests.Compute
 		{
 			await CreateTaskWithRequirement("memory >= 10");
 			IPool pool = await CreatePool(new () { {"memory", "20"} });
-			Assert.AreEqual(1, await ComputeService.GetNumQueuedTasksForPool(_clusterId, pool));
+			Assert.AreEqual(1, await ComputeService.GetNumQueuedTasksForPoolAsync(_clusterId, pool));
 		}
 		
 		[TestMethod]
@@ -62,7 +62,7 @@ namespace Horde.Build.Tests.Compute
 		{
 			await CreateTaskWithRequirement("memory >= 10 && cpu < 5");
 			IPool pool = await CreatePool(new () { {"memory", "15"}, {"cpu", "4"} });
-			Assert.AreEqual(1, await ComputeService.GetNumQueuedTasksForPool(_clusterId, pool));
+			Assert.AreEqual(1, await ComputeService.GetNumQueuedTasksForPoolAsync(_clusterId, pool));
 		}
 		
 		[TestMethod]
@@ -75,8 +75,8 @@ namespace Horde.Build.Tests.Compute
 			IPool pool1 = await CreatePool(new () { {"memory", "20"}, {"cpu", "4"} });
 			IPool pool2 = await CreatePool(new () { {"memory", "15"}, {"cpu", "4"} });
 
-			Assert.AreEqual(3, await ComputeService.GetNumQueuedTasksForPool(_clusterId, pool1));
-			Assert.AreEqual(2, await ComputeService.GetNumQueuedTasksForPool(_clusterId, pool2));
+			Assert.AreEqual(3, await ComputeService.GetNumQueuedTasksForPoolAsync(_clusterId, pool1));
+			Assert.AreEqual(2, await ComputeService.GetNumQueuedTasksForPoolAsync(_clusterId, pool2));
 		}
 
 		private async Task CreateTaskWithRequirement(string condition, ChannelId? channelId = null, RefId? refId = null)
diff --git a/Engine/Source/Programs/Horde/Horde.Build/Compute/ComputeService.cs b/Engine/Source/Programs/Horde/Horde.Build/Compute/ComputeService.cs
--- a/Engine/Source/Programs/Horde/Horde.Build/Compute/ComputeService.cs
+++ b/Engine/Source/Programs/Horde/Horde.Build/Compute/ComputeService.cs
@@ -63,7 +63,7 @@ namespace Horde.Build.Compute
 	/// <summary>
 	/// Dispatches remote actions. Does not implement any cross-pod communication to satisfy leases; only agents connected to this server instance will be stored.
 	/// </summary>
-	public class ComputeService : TaskSourceBase<ComputeTaskMessage>, IHostedService, IDisposable
+	public class ComputeService : TaskSourceBase<ComputeTaskMessage>, IHostedService, IDisposable, IComputeService
 	{
 		[RedisConverter(typeof(QueueKeySerializer))]
 		class QueueKey
@@ -200,10 +200,7 @@ namespace Horde.Build.Compute
 			_requirementsCache.Dispose();
 		}
 
-		/// <summary>
-		/// Gets information about a compute cluster
-		/// </summary>
-		/// <param name="clusterId">Cluster to use for execution</param>
+		/// <inheritdoc/>
 		public async Task<IComputeClusterInfo> GetClusterInfoAsync(ClusterId clusterId)
 		{
 			ComputeClusterConfig? config = await GetClusterAsync(clusterId);
@@ -214,14 +211,7 @@ namespace Horde.Build.Compute
 			return new ClusterInfo(config);
 		}
 
-		/// <summary>
-		/// Post tasks to be executed to a channel
-		/// </summary>
-		/// <param name="clusterId">Cluster to use for execution</param>
-		/// <param name="channelId">Unique identifier of the client</param>
-		/// <param name="taskRefIds">List of task hashes</param>
-		/// <param name="requirementsHash">Requirements document for execution</param>
-		/// <returns>Async task</returns>
+		/// <inheritdoc/>
 		public async Task AddTasksAsync(ClusterId clusterId, ChannelId channelId, List<RefId> taskRefIds, CbObjectAttachment requirementsHash)
 		{
 			List<Task> tasks = new List<Task>();
@@ -324,36 +314,19 @@ namespace Horde.Build.Compute
 			return _taskScheduler.EnqueueAsync(new QueueKey(clusterId, new IoHash(message.RequirementsHash.ToByteArray())), taskInfo, true);
 		}
 
-		/// <summary>
-		/// Dequeue completed items from a queue and return immediately
-		/// </summary>
-		/// <param name="clusterId">Cluster containing the channel</param>
-		/// <param name="channelId">Queue to remove items from</param>
-		/// <returns>List of status updates</returns>
+		/// <inheritdoc/>
 		public async Task<List<ComputeTaskStatus>> GetTaskUpdatesAsync(ClusterId clusterId, ChannelId channelId)
 		{
 			return await _messageQueue.ReadMessagesAsync(GetMessageQueueId(clusterId, channelId));
 		}
 		
-		/// <summary>
-		/// Get number of queued tasks that are compatible for a given pool
-		/// </summary>
-		/// <param name="clusterId">Cluster to query</param>
-		/// <param name="pool">Pool to check for compatibility</param>
-		/// <param name="cancellationToken">Cancellation token for the operation</param> 
-		/// <returns>Number of tasks in queue</returns>
-		public async Task<int> GetNumQueuedTasksForPool(ClusterId clusterId, IPool pool, CancellationToken cancellationToken = default)
+		/// <inheritdoc/>
+		public async Task<int> GetNumQueuedTasksForPoolAsync(ClusterId clusterId, IPool pool, CancellationToken cancellationToken = default)
 		{
 			return await _taskScheduler.GetNumQueuedTasksAsync(queueKey => CheckRequirements(pool, queueKey), cancellationToken);
 		}
 
-		/// <summary>
-		/// Dequeue completed items from a queue
-		/// </summary>
-		/// <param name="clusterId">Cluster containing the channel</param>
-		/// <param name="channelId">Queue to remove items from</param>
-		/// <param name="cancellationToken">Cancellation token to stop waiting for items</param>
-		/// <returns>List of status updates</returns>
+		/// <inheritdoc/>
 		public async Task<List<ComputeTaskStatus>> WaitForTaskUpdatesAsync(ClusterId clusterId, ChannelId channelId, CancellationToken cancellationToken)
 		{
 			return await _messageQueue.WaitForMessagesAsync(GetMessageQueueId(clusterId, channelId), cancellationToken);
diff --git a/Engine/Source/Programs/Horde/Horde.Build/Compute/IComputeService.cs b/Engine/Source/Programs/Horde/Horde.Build/Compute/IComputeService.cs
new file mode 100644
--- /dev/null
+++ b/Engine/Source/Programs/Horde/Horde.Build/Compute/IComputeService.cs
@@ -0,0 +1,59 @@
+// Copyright Epic Games, Inc. All Rights Reserved.
+
+using System.Collections.Generic;
+using System.Threading;
+using System.Threading.Tasks;
+using EpicGames.Horde.Compute;
+using EpicGames.Horde.Storage;
+using EpicGames.Serialization;
+using Horde.Build.Agents.Pools;
+
+namespace Horde.Build.Compute;
+
+/// <summary>
+/// Interface for running and scheduling compute tasks (primarily for remote execution)
+/// </summary>
+public interface IComputeService
+{
+	/// <summary>
+	/// Post tasks to be executed to a channel
+	/// </summary>
+	/// <param name="clusterId">Cluster to use for execution</param>
+	/// <param name="channelId">Unique identifier of the client</param>
+	/// <param name="taskRefIds">List of task hashes</param>
+	/// <param name="requirementsHash">Requirements document for execution</param>
+	/// <returns>Async task</returns>
+	Task AddTasksAsync(ClusterId clusterId, ChannelId channelId, List<RefId> taskRefIds, CbObjectAttachment requirementsHash);
+
+	/// <summary>
+	/// Gets information about a compute cluster
+	/// </summary>
+	/// <param name="clusterId">Cluster to use for execution</param>
+	Task<IComputeClusterInfo> GetClusterInfoAsync(ClusterId clusterId);
+
+	/// <summary>
+	/// Dequeue completed items from a queue and return immediately
+	/// </summary>
+	/// <param name="clusterId">Cluster containing the channel</param>
+	/// <param name="channelId">Queue to remove items from</param>
+	/// <returns>List of status updates</returns>
+	Task<List<ComputeTaskStatus>> GetTaskUpdatesAsync(ClusterId clusterId, ChannelId channelId);
+
+	/// <summary>
+	/// Get number of queued tasks that are compatible for a given pool
+	/// </summary>
+	/// <param name="clusterId">Cluster to query</param>
+	/// <param name="pool">Pool to check for compatibility</param>
+	/// <param name="cancellationToken">Cancellation token for the operation</param> 
+	/// <returns>Number of tasks in queue</returns>
+	Task<int> GetNumQueuedTasksForPoolAsync(ClusterId clusterId, IPool pool, CancellationToken cancellationToken = default);
+
+	/// <summary>
+	/// Dequeue completed items from a queue
+	/// </summary>
+	/// <param name="clusterId">Cluster containing the channel</param>
+	/// <param name="channelId">Queue to remove items from</param>
+	/// <param name="cancellationToken">Cancellation token to stop waiting for items</param>
+	/// <returns>List of status updates</returns>
+	Task<List<ComputeTaskStatus>> WaitForTaskUpdatesAsync(ClusterId clusterId, ChannelId channelId, CancellationToken cancellationToken);
+}
\ No newline at end of file
