Horde: Add a stream setting to allow using virtual streams for replication.

diff --git a/Engine/Source/Programs/Horde/Horde.Build/Perforce/ReplicationService.cs b/Engine/Source/Programs/Horde/Horde.Build/Perforce/ReplicationService.cs
--- a/Engine/Source/Programs/Horde/Horde.Build/Perforce/ReplicationService.cs
+++ b/Engine/Source/Programs/Horde/Horde.Build/Perforce/ReplicationService.cs
@@ -1036,7 +1036,7 @@ namespace Horde.Build.Perforce
 				newClient.Description = "Created to mirror Perforce content to Horde Storage";
 				newClient.Owner = perforce.Settings.UserName;
 				newClient.Host = serverInfo.ClientHost;
-				newClient.Stream = stream.Name;
+				newClient.Stream = stream.Config.ReplicationStream ?? stream.Name;
 				await perforce.CreateClientAsync(newClient);
 				_logger.LogInformation("Created client {ClientName} for {StreamName}", newClient.Name, stream.Name);
 
diff --git a/Engine/Source/Programs/Horde/Horde.Build/Server/Config.cs b/Engine/Source/Programs/Horde/Horde.Build/Server/Config.cs
--- a/Engine/Source/Programs/Horde/Horde.Build/Server/Config.cs
+++ b/Engine/Source/Programs/Horde/Horde.Build/Server/Config.cs
@@ -269,6 +269,11 @@ namespace Horde.Build.Server
 		/// </summary>
 		public string? ReplicationFilter { get; set; }
 
+		/// <summary>
+		/// Stream to use for replication, if different to the default.
+		/// </summary>
+		public string? ReplicationStream { get; set; }
+
 		/// <summary>
 		/// Workflows for dealing with new issues
 		/// </summary>
