Horde: Add no-op pool sizing strategy that emits AWS CloudWatch metrics for auto-scaling groups

diff --git a/Engine/Source/Programs/Horde/Horde.Build.Tests/ControllerIntegrationTest.cs b/Engine/Source/Programs/Horde/Horde.Build.Tests/ControllerIntegrationTest.cs
--- a/Engine/Source/Programs/Horde/Horde.Build.Tests/ControllerIntegrationTest.cs
+++ b/Engine/Source/Programs/Horde/Horde.Build.Tests/ControllerIntegrationTest.cs
@@ -4,6 +4,7 @@ using System;
 using System.Collections.Generic;
 using System.Net.Http;
 using System.Threading.Tasks;
+using Amazon.CloudWatch;
 using Horde.Build.Jobs;
 using Horde.Build.Server;
 using Horde.Build.Perforce;
@@ -19,6 +20,7 @@ using Horde.Build.Jobs.Templates;
 using Horde.Build.Configuration;
 using Horde.Build.Jobs.Graphs;
 using Horde.Build.Jobs.Artifacts;
+using Moq;
 
 namespace Horde.Build.Tests;
 
@@ -42,8 +44,14 @@ public class TestWebApplicationFactory<TStartup> : WebApplicationFactory<TStartu
 			{ "Horde:OidcAuthority", null },
 			{ "Horde:OidcClientId", null }
 		};
+
+		Mock<IAmazonCloudWatch> cloudWatchMock = new (MockBehavior.Strict);
 		builder.ConfigureAppConfiguration((hostingContext, config) => { config.AddInMemoryCollection(dict); });
-		builder.ConfigureTestServices(collection => collection.AddSingleton<IPerforceService, PerforceServiceStub>());
+		builder.ConfigureTestServices(collection =>
+		{
+			collection.AddSingleton<IPerforceService, PerforceServiceStub>();
+			collection.AddSingleton<IAmazonCloudWatch>(x => cloudWatchMock.Object);
+		});
 	}
 }
 
diff --git a/Engine/Source/Programs/Horde/Horde.Build.Tests/Fleet/AutoscaleServiceV2Test.cs b/Engine/Source/Programs/Horde/Horde.Build.Tests/Fleet/AutoscaleServiceV2Test.cs
--- a/Engine/Source/Programs/Horde/Horde.Build.Tests/Fleet/AutoscaleServiceV2Test.cs
+++ b/Engine/Source/Programs/Horde/Horde.Build.Tests/Fleet/AutoscaleServiceV2Test.cs
@@ -81,7 +81,8 @@ namespace Horde.Build.Tests.Fleet
 			PoolSizeStrategySpy leaseUtilizationSpy = new();
 			PoolSizeStrategySpy jobQueueSpy = new();
 			PoolSizeStrategySpy noOpSpy = new();
-			service.OverridePoolSizeStrategiesDuringTesting(leaseUtilizationSpy, jobQueueSpy, noOpSpy);
+			PoolSizeStrategySpy computeQueueAwsMetricSpy = new();
+			service.OverridePoolSizeStrategiesDuringTesting(leaseUtilizationSpy, jobQueueSpy, noOpSpy, computeQueueAwsMetricSpy);
 
 			await service.TickLeaderAsync(CancellationToken.None);
 			
@@ -161,7 +162,8 @@ namespace Horde.Build.Tests.Fleet
 
 			LeaseUtilizationStrategy leaseUtilizationStrategy = new(AgentCollection, PoolCollection, LeaseCollection, Clock);
 			JobQueueStrategy jobQueueStrategy = new(JobCollection, GraphCollection, StreamService, Clock);
-			AutoscaleServiceV2 service = new AutoscaleServiceV2(leaseUtilizationStrategy,jobQueueStrategy, new NoOpPoolSizeStrategy(), AgentCollection, PoolCollection, fleetManager, _dogStatsD, Clock, serverSettingsOpt, logger);
+			ComputeQueueAwsMetricStrategy computeQueueAwsMetricStrategy = new(null!, null!, null!);
+			AutoscaleServiceV2 service = new AutoscaleServiceV2(leaseUtilizationStrategy,jobQueueStrategy, new NoOpPoolSizeStrategy(), computeQueueAwsMetricStrategy, AgentCollection, PoolCollection, fleetManager, _dogStatsD, Clock, serverSettingsOpt, logger);
 			return service;
 		}
 	}
diff --git a/Engine/Source/Programs/Horde/Horde.Build/Agents/Fleet/AutoscaleServiceV2.cs b/Engine/Source/Programs/Horde/Horde.Build/Agents/Fleet/AutoscaleServiceV2.cs
--- a/Engine/Source/Programs/Horde/Horde.Build/Agents/Fleet/AutoscaleServiceV2.cs
+++ b/Engine/Source/Programs/Horde/Horde.Build/Agents/Fleet/AutoscaleServiceV2.cs
@@ -28,6 +28,7 @@ namespace Horde.Build.Agents.Fleet
 		private IPoolSizeStrategy _leaseUtilizationStrategy;
 		private IPoolSizeStrategy _jobQueueStrategy;
 		private IPoolSizeStrategy _noOpStrategy;
+		private IPoolSizeStrategy _computeQueueAwsMetricStrategy;
 		private readonly IAgentCollection _agentCollection;
 		private readonly IPoolCollection _poolCollection;
 		private readonly IFleetManager _fleetManager;
@@ -42,11 +43,23 @@ namespace Horde.Build.Agents.Fleet
 		/// <summary>
 		/// Constructor
 		/// </summary>
-		public AutoscaleServiceV2(LeaseUtilizationStrategy leaseUtilizationStrategy, JobQueueStrategy jobQueueStrategy, NoOpPoolSizeStrategy noOpStrategy, IAgentCollection agentCollection, IPoolCollection poolCollection, IFleetManager fleetManager, IDogStatsd dogStatsd, IClock clock, IOptions<ServerSettings> settings, ILogger<AutoscaleServiceV2> logger)
+		public AutoscaleServiceV2(
+			LeaseUtilizationStrategy leaseUtilizationStrategy,
+			JobQueueStrategy jobQueueStrategy,
+			NoOpPoolSizeStrategy noOpStrategy,
+			ComputeQueueAwsMetricStrategy computeQueueAwsMetricStrategy,
+			IAgentCollection agentCollection,
+			IPoolCollection poolCollection,
+			IFleetManager fleetManager,
+			IDogStatsd dogStatsd,
+			IClock clock,
+			IOptions<ServerSettings> settings,
+			ILogger<AutoscaleServiceV2> logger)
 		{
 			_leaseUtilizationStrategy = leaseUtilizationStrategy;
 			_jobQueueStrategy = jobQueueStrategy;
 			_noOpStrategy = noOpStrategy;
+			_computeQueueAwsMetricStrategy = computeQueueAwsMetricStrategy;
 			_agentCollection = agentCollection;
 			_poolCollection = poolCollection;
 			_fleetManager = fleetManager;
@@ -181,11 +194,17 @@ namespace Horde.Build.Agents.Fleet
 		/// <param name="leaseUtilizationStrategy"></param>
 		/// <param name="jobQueueStrategy"></param>
 		/// <param name="noOpStrategy"></param>
-		internal void OverridePoolSizeStrategiesDuringTesting(IPoolSizeStrategy leaseUtilizationStrategy, IPoolSizeStrategy jobQueueStrategy, IPoolSizeStrategy noOpStrategy)
+		/// <param name="computeQueueAwsMetricStrategy"></param>
+		internal void OverridePoolSizeStrategiesDuringTesting(
+			IPoolSizeStrategy leaseUtilizationStrategy,
+			IPoolSizeStrategy jobQueueStrategy,
+			IPoolSizeStrategy noOpStrategy,
+			IPoolSizeStrategy computeQueueAwsMetricStrategy)
 		{
 			_leaseUtilizationStrategy = leaseUtilizationStrategy;
 			_jobQueueStrategy = jobQueueStrategy;
 			_noOpStrategy = noOpStrategy;
+			_computeQueueAwsMetricStrategy = computeQueueAwsMetricStrategy;
 		}
 		
 		private IPoolSizeStrategy GetPoolSizeStrategy(PoolSizeStrategy type)
@@ -195,6 +214,7 @@ namespace Horde.Build.Agents.Fleet
 				PoolSizeStrategy.LeaseUtilization => _leaseUtilizationStrategy,
 				PoolSizeStrategy.JobQueue => _jobQueueStrategy,
 				PoolSizeStrategy.NoOp => _noOpStrategy,
+				PoolSizeStrategy.ComputeQueueAwsMetric => _computeQueueAwsMetricStrategy,
 				_ => throw new Exception($"Unknown strategy: {type}")
 			};
 		}
diff --git a/Engine/Source/Programs/Horde/Horde.Build/Agents/Fleet/ComputeQueueAwsMetricStrategy.cs b/Engine/Source/Programs/Horde/Horde.Build/Agents/Fleet/ComputeQueueAwsMetricStrategy.cs
new file mode 100644
--- /dev/null
+++ b/Engine/Source/Programs/Horde/Horde.Build/Agents/Fleet/ComputeQueueAwsMetricStrategy.cs
@@ -0,0 +1,146 @@
+// Copyright Epic Games, Inc. All Rights Reserved.
+
+using System;
+using System.Collections.Generic;
+using System.Linq;
+using System.Net;
+using System.Threading.Tasks;
+using Amazon.CloudWatch;
+using Amazon.CloudWatch.Model;
+using EpicGames.Horde.Compute;
+using Horde.Build.Agents.Pools;
+using Horde.Build.Compute;
+using Horde.Build.Utilities;
+using Microsoft.Extensions.Logging;
+using OpenTracing;
+using OpenTracing.Util;
+
+namespace Horde.Build.Agents.Fleet
+{
+	using PoolId = StringId<IPool>;
+
+	/// <summary>
+	/// Settings for <see cref="ComputeQueueAwsMetricStrategy" />
+	/// </summary>
+	public class ComputeQueueAwsMetricSettings
+	{
+		/// <summary>
+		/// Compute cluster ID to observe
+		/// </summary>
+		public string ComputeClusterId { get; set; }
+		
+		/// <summary>
+		/// AWS CloudWatch namespace to write metrics in
+		/// </summary>
+		public string Namespace { get; set; }
+		
+		/// <summary>
+		/// Constructor
+		/// </summary>
+		/// <param name="computeClusterId"></param>
+		/// <param name="cloudWatchNamespace"></param>
+		public ComputeQueueAwsMetricSettings(string computeClusterId, string cloudWatchNamespace = "HordeBuild")
+		{
+			ComputeClusterId = computeClusterId;
+			Namespace = cloudWatchNamespace;
+		}
+	}
+	
+	/// <summary>
+	/// A no-op strategy that reports size of compute task queue for a given pool as AWS CloudWatch metrics
+	/// A metric that later can be used as a source for AWS-controlled auto-scaling policies. 
+	/// </summary>
+	public class ComputeQueueAwsMetricStrategy : IPoolSizeStrategy
+	{
+		private readonly IAmazonCloudWatch _cloudWatch;
+		private readonly IComputeService _computeService;
+		private readonly ILogger<ComputeQueueAwsMetricStrategy> _logger;
+
+		/// <summary>
+		/// Constructor
+		/// </summary>
+		/// <param name="cloudWatch"></param>
+		/// <param name="computeService"></param>
+		/// <param name="logger"></param>
+		public ComputeQueueAwsMetricStrategy(IAmazonCloudWatch cloudWatch, IComputeService computeService, ILogger<ComputeQueueAwsMetricStrategy> logger)
+		{
+			_cloudWatch = cloudWatch;
+			_computeService = computeService;
+			_logger = logger;
+		}
+
+		/// <inheritdoc/>
+		public string Name { get; } = "ComputeQueueAwsMetric";
+
+		/// <inheritdoc/>
+		public async Task<List<PoolSizeData>> CalcDesiredPoolSizesAsync(List<PoolSizeData> pools)
+		{
+			using IScope scope = GlobalTracer.Instance.BuildSpan("ComputeQueueAwsMetricStrategy.CalcDesiredPoolSizesAsync").StartActive();
+
+			Dictionary<string, List<MetricDatum>> metricsPerCloudWatchNamespace = new();
+			foreach (PoolSizeData poolSizeData in pools)
+			{
+				ComputeQueueAwsMetricSettings? settings = poolSizeData.Pool.ComputeQueueAwsMetricSettings;
+				if (settings == null)
+				{
+					_logger.LogWarning("Pool {PoolId} is configured to use strategy {Strategy} but has no settings configured", poolSizeData.Pool.Id, nameof(ComputeQueueAwsMetricStrategy));
+					continue;
+				}
+				
+				int numAgents = poolSizeData.Agents.Count;
+				int totalCpuCores = poolSizeData.Agents.Select(x => x.Resources.TryGetValue(KnownPropertyNames.LogicalCores, out int numCpuCores) ? numCpuCores : 0).Sum();
+				int numQueuedComputeTasks = await _computeService.GetNumQueuedTasksForPoolAsync(new ClusterId(settings.ComputeClusterId), poolSizeData.Pool);
+				
+				double numQueuedTasksPerAgent = numQueuedComputeTasks / (double)numAgents;
+				double numQueuedTasksPerCores = numQueuedComputeTasks / (double)totalCpuCores;
+
+				DateTime now = DateTime.UtcNow;
+				
+				if (!metricsPerCloudWatchNamespace.TryGetValue(settings.Namespace, out List<MetricDatum>? metricDatums))
+				{
+					metricDatums = new List<MetricDatum>();
+					metricsPerCloudWatchNamespace[settings.Namespace] = metricDatums;
+				}
+				
+				List<Dimension> dimensions = new ()
+				{
+					new () { Name = "Pool", Value = poolSizeData.Pool.Id.ToString() },
+					new () { Name = "ComputeCluster", Value = settings.ComputeClusterId },
+				};
+				metricDatums.Add(new ()
+				{
+					MetricName = "NumComputeTasksPerAgent",
+					Dimensions = dimensions,
+					Unit = StandardUnit.Count,
+					Value = numQueuedTasksPerAgent,
+					TimestampUtc = now
+				});
+				metricDatums.Add(new ()
+				{
+					MetricName = "NumComputeTasksPerCpuCores",
+					Dimensions = dimensions,
+					Unit = StandardUnit.Count,
+					Value = numQueuedTasksPerCores,
+					TimestampUtc = now
+				});
+			}
+
+			foreach ((string ns, List<MetricDatum> metricDatums) in metricsPerCloudWatchNamespace)
+			{
+				using IScope cwScope = GlobalTracer.Instance.BuildSpan("Putting CloudWatch metrics").StartActive();
+				cwScope.Span.SetTag("namespace", ns);
+				
+				PutMetricDataRequest request = new() { Namespace = ns, MetricData = metricDatums };
+				PutMetricDataResponse response = await _cloudWatch.PutMetricDataAsync(request);
+				scope.Span.SetTag("res.statusCode", (int)response.HttpStatusCode);
+				if (response.HttpStatusCode != HttpStatusCode.OK)
+				{
+					_logger.LogError("Unable to put CloudWatch metrics");
+				}
+			}
+
+			// Return the input pool size data as-is since we are only observing and not modifying
+			return pools;
+		}
+	}
+}
diff --git a/Engine/Source/Programs/Horde/Horde.Build/Agents/Fleet/IPoolSizeStrategy.cs b/Engine/Source/Programs/Horde/Horde.Build/Agents/Fleet/IPoolSizeStrategy.cs
--- a/Engine/Source/Programs/Horde/Horde.Build/Agents/Fleet/IPoolSizeStrategy.cs
+++ b/Engine/Source/Programs/Horde/Horde.Build/Agents/Fleet/IPoolSizeStrategy.cs
@@ -28,7 +28,13 @@ namespace Horde.Build.Agents.Fleet
 		/// No-op strategy used as fallback/default behavior
 		/// <see cref="NoOpPoolSizeStrategy"/> 
 		/// </summary>
-		NoOp
+		NoOp,
+		
+		/// <summary>
+		/// A no-op strategy that reports metrics to let an external AWS auto-scaling policy scale the fleet
+		/// <see cref="ComputeQueueAwsMetric"/> 
+		/// </summary>
+		ComputeQueueAwsMetric
 	}
 
 	/// <summary>
diff --git a/Engine/Source/Programs/Horde/Horde.Build/Agents/Pools/IPool.cs b/Engine/Source/Programs/Horde/Horde.Build/Agents/Pools/IPool.cs
--- a/Engine/Source/Programs/Horde/Horde.Build/Agents/Pools/IPool.cs
+++ b/Engine/Source/Programs/Horde/Horde.Build/Agents/Pools/IPool.cs
@@ -100,6 +100,11 @@ namespace Horde.Build.Agents.Pools
 		/// Settings for job queue pool sizing strategy (if used)
 		/// </summary>
 		public JobQueueSettings? JobQueueSettings { get; }
+		
+		/// <summary>
+		/// Settings for job queue pool sizing strategy (if used)
+		/// </summary>
+		public ComputeQueueAwsMetricSettings? ComputeQueueAwsMetricSettings { get; }
 
 		/// <summary>
 		/// Update index for this document
diff --git a/Engine/Source/Programs/Horde/Horde.Build/Agents/Pools/IPoolCollection.cs b/Engine/Source/Programs/Horde/Horde.Build/Agents/Pools/IPoolCollection.cs
--- a/Engine/Source/Programs/Horde/Horde.Build/Agents/Pools/IPoolCollection.cs
+++ b/Engine/Source/Programs/Horde/Horde.Build/Agents/Pools/IPoolCollection.cs
@@ -31,6 +31,7 @@ namespace Horde.Build.Agents.Pools
 		/// <param name="sizeStrategy">Pool sizing strategy</param>
 		/// <param name="leaseUtilizationSettings">Settings for lease utilization strategy</param>
 		/// <param name="jobQueueSettings">Settings for job queue strategy</param>
+		/// <param name="computeQueueAwsMetricSettings">Settings for compute queue AWS metric strategy</param>
 		/// <param name="properties">Properties for the pool</param>
 		/// <returns>The new pool document</returns>
 		Task<IPool> AddAsync(
@@ -46,6 +47,7 @@ namespace Horde.Build.Agents.Pools
 			PoolSizeStrategy? sizeStrategy = null,
 			LeaseUtilizationSettings? leaseUtilizationSettings = null,
 			JobQueueSettings? jobQueueSettings = null,
+			ComputeQueueAwsMetricSettings? computeQueueAwsMetricSettings = null,
 			IEnumerable<KeyValuePair<string, string>>? properties = null);
 
 		/// <summary>
@@ -94,6 +96,7 @@ namespace Horde.Build.Agents.Pools
 		/// <param name="sizeStrategy">Pool sizing strategy</param>
 		/// <param name="leaseUtilizationSettings">Settings for lease utilization strategy</param>
 		/// <param name="jobQueueSettings">Settings for job queue strategy</param>
+		/// <param name="computeQueueAwsMetricSettings">Settings for compute queue AWS metric strategy</param>
 		/// <param name="useDefaultStrategy">Whether to use the default strategy</param>
 		/// <returns>Async task</returns>
 		Task<IPool?> TryUpdateAsync(
@@ -114,6 +117,7 @@ namespace Horde.Build.Agents.Pools
 			PoolSizeStrategy? sizeStrategy = null,
 			LeaseUtilizationSettings? leaseUtilizationSettings = null,
 			JobQueueSettings? jobQueueSettings = null,
+			ComputeQueueAwsMetricSettings? computeQueueAwsMetricSettings = null,
 			bool? useDefaultStrategy = null);
 	}
 }
diff --git a/Engine/Source/Programs/Horde/Horde.Build/Agents/Pools/PoolCollection.cs b/Engine/Source/Programs/Horde/Horde.Build/Agents/Pools/PoolCollection.cs
--- a/Engine/Source/Programs/Horde/Horde.Build/Agents/Pools/PoolCollection.cs
+++ b/Engine/Source/Programs/Horde/Horde.Build/Agents/Pools/PoolCollection.cs
@@ -46,6 +46,7 @@ namespace Horde.Build.Agents.Pools
 			PoolSizeStrategy IPool.SizeStrategy => SizeStrategy ?? defaultStrategy!.Value;
 			public LeaseUtilizationSettings? LeaseUtilizationSettings { get; set; }
 			public JobQueueSettings? JobQueueSettings { get; set; }
+			public ComputeQueueAwsMetricSettings? ComputeQueueAwsMetricSettings { get; set; }
 			public int UpdateIndex { get; set; }
 
 			// Read-only wrappers
@@ -75,6 +76,7 @@ namespace Horde.Build.Agents.Pools
 				SizeStrategy = other.SizeStrategy;
 				LeaseUtilizationSettings = other.LeaseUtilizationSettings;
 				JobQueueSettings = other.JobQueueSettings;
+				ComputeQueueAwsMetricSettings = other.ComputeQueueAwsMetricSettings;
 				UpdateIndex = other.UpdateIndex;
 			}
 		}
@@ -111,6 +113,7 @@ namespace Horde.Build.Agents.Pools
 			PoolSizeStrategy? sizeStrategy,
 			LeaseUtilizationSettings? leaseUtilizationSettings,
 			JobQueueSettings? jobQueueSettings,
+			ComputeQueueAwsMetricSettings? computeQueueAwsMetricSettings,
 			IEnumerable<KeyValuePair<string, string>>? properties)
 		{
 			PoolDocument pool = new PoolDocument();
@@ -134,6 +137,7 @@ namespace Horde.Build.Agents.Pools
 			pool.SizeStrategy = sizeStrategy;
 			pool.LeaseUtilizationSettings = leaseUtilizationSettings;
 			pool.JobQueueSettings = jobQueueSettings;
+			pool.ComputeQueueAwsMetricSettings = computeQueueAwsMetricSettings;
 			await _pools.InsertOneAsync(pool);
 			return pool;
 		}
@@ -226,6 +230,7 @@ namespace Horde.Build.Agents.Pools
 			PoolSizeStrategy? sizeStrategy,
 			LeaseUtilizationSettings? leaseUtilizationSettings,
 			JobQueueSettings? jobQueueSettings, 
+			ComputeQueueAwsMetricSettings? computeQueueAwsMetricSettings, 
 			bool? useDefaultStrategy)
 		{
 			TransactionBuilder<PoolDocument> transaction = new TransactionBuilder<PoolDocument>();
@@ -333,6 +338,10 @@ namespace Horde.Build.Agents.Pools
 			{
 				transaction.Set(x => x.JobQueueSettings, jobQueueSettings);
 			}
+			if (computeQueueAwsMetricSettings != null)
+			{
+				transaction.Set(x => x.ComputeQueueAwsMetricSettings, computeQueueAwsMetricSettings);
+			}
 			return TryUpdateAsync(pool, transaction);
 		}
 	}
diff --git a/Engine/Source/Programs/Horde/Horde.Build/Agents/Pools/PoolMessages.cs b/Engine/Source/Programs/Horde/Horde.Build/Agents/Pools/PoolMessages.cs
--- a/Engine/Source/Programs/Horde/Horde.Build/Agents/Pools/PoolMessages.cs
+++ b/Engine/Source/Programs/Horde/Horde.Build/Agents/Pools/PoolMessages.cs
@@ -75,6 +75,39 @@ namespace Horde.Build.Agents.Pools
 		}
 	}
 	
+	/// <see cref="ComputeQueueAwsMetricSettings" />
+	public class ComputeQueueAwsMetricSettingsMessage
+	{
+		/// <summary>
+		/// Compute cluster ID to observe
+		/// </summary>
+		public string ComputeClusterId { get; set; }
+
+		/// <summary>
+		/// AWS CloudWatch namespace to write metrics in
+		/// </summary>
+		public string Namespace { get; set; }
+
+		/// <summary>
+		/// Construct a public REST API representation from the internal one
+		/// </summary>
+		/// <param name="settings"></param>
+		public ComputeQueueAwsMetricSettingsMessage(ComputeQueueAwsMetricSettings settings)
+		{
+			ComputeClusterId = settings.ComputeClusterId;
+			Namespace = settings.Namespace;
+		}
+		
+		/// <summary>
+		/// Convert this public REST API object to an internal representation
+		/// </summary>
+		/// <returns></returns>
+		public ComputeQueueAwsMetricSettings Convert()
+		{
+			return new ComputeQueueAwsMetricSettings(ComputeClusterId, Namespace);
+		}
+	}
+	
 	/// <summary>
 	/// Parameters to create a new pool
 	/// </summary>
@@ -125,9 +158,14 @@ namespace Horde.Build.Agents.Pools
 		/// Settings for job queue pool sizing strategy (if used) 
 		/// </summary>
 		public JobQueueSettingsMessage? JobQueueSettings { get; set; }
+		
+		/// <summary>
+		/// Settings for compute queue pool sizing strategy for AWS metrics (if used) 
+		/// </summary>
+		public ComputeQueueAwsMetricSettingsMessage? ComputeQueueAwsMetricSettings { get; set; }
 
 		/// <summary>
-		/// The minimum nunmber of agents to retain in this pool
+		/// The minimum number of agents to retain in this pool
 		/// </summary>
 		public int? MinAgents { get; set; }
 
@@ -216,9 +254,14 @@ namespace Horde.Build.Agents.Pools
 		/// Settings for job queue pool sizing strategy (if used) 
 		/// </summary>
 		public JobQueueSettingsMessage? JobQueueSettings { get; set; }
+		
+		/// <summary>
+		/// Settings for compute queue AWS metric pool sizing strategy (if used) 
+		/// </summary>
+		public ComputeQueueAwsMetricSettingsMessage? ComputeQueueAwsMetricSettings { get; set; }
 
 		/// <summary>
-		/// The minimum nunmber of agents to retain in this pool
+		/// The minimum number of agents to retain in this pool
 		/// </summary>
 		public int? MinAgents { get; set; }
 
@@ -310,6 +353,11 @@ namespace Horde.Build.Agents.Pools
 		/// </summary>
 		public JobQueueSettingsMessage? JobQueueSettings { get; set; }
 
+		/// <summary>
+		/// Settings for compute queue pool sizing strategy for AWS metrics (if used) 
+		/// </summary>
+		public ComputeQueueAwsMetricSettingsMessage? ComputeQueueAwsMetricSettings { get; set; }
+
 		/// <summary>
 		/// The minimum nunmber of agents to retain in this pool
 		/// </summary>
@@ -346,6 +394,7 @@ namespace Horde.Build.Agents.Pools
 			SizeStrategy = pool.SizeStrategy;
 			LeaseUtilizationSettings = pool.LeaseUtilizationSettings == null ? null : new LeaseUtilizationSettingsMessage(pool.LeaseUtilizationSettings);
 			JobQueueSettings = pool.JobQueueSettings == null ? null : new JobQueueSettingsMessage(pool.JobQueueSettings);
+			ComputeQueueAwsMetricSettings = pool.ComputeQueueAwsMetricSettings == null ? null : new ComputeQueueAwsMetricSettingsMessage(pool.ComputeQueueAwsMetricSettings);
 			MinAgents = pool.MinAgents;
 			NumReserveAgents = pool.NumReserveAgents;
 			Workspaces = pool.Workspaces.Select(x => new GetAgentWorkspaceResponse(x)).ToList();
diff --git a/Engine/Source/Programs/Horde/Horde.Build/Agents/Pools/PoolService.cs b/Engine/Source/Programs/Horde/Horde.Build/Agents/Pools/PoolService.cs
--- a/Engine/Source/Programs/Horde/Horde.Build/Agents/Pools/PoolService.cs
+++ b/Engine/Source/Programs/Horde/Horde.Build/Agents/Pools/PoolService.cs
@@ -67,6 +67,7 @@ namespace Horde.Build.Agents.Pools
 		/// <param name="sizeStrategy">Pool sizing strategy</param>
 		/// <param name="leaseUtilizationSettings">Settings for lease utilization strategy</param>
 		/// <param name="jobQueueSettings">Settings for job queue strategy</param>
+		/// <param name="computeQueueAwsMetricSettings">Settings for compute queue AWS metric strategy</param>
 		/// <param name="properties">Properties for the new pool</param>
 		/// <returns>The new pool document</returns>
 		public Task<IPool> CreatePoolAsync(
@@ -81,6 +82,7 @@ namespace Horde.Build.Agents.Pools
 			PoolSizeStrategy? sizeStrategy = null,
 			LeaseUtilizationSettings? leaseUtilizationSettings = null,
 			JobQueueSettings? jobQueueSettings = null,
+			ComputeQueueAwsMetricSettings? computeQueueAwsMetricSettings = null,
 			Dictionary<string, string>? properties = null)
 		{
 			return _pools.AddAsync(
@@ -96,6 +98,7 @@ namespace Horde.Build.Agents.Pools
 				sizeStrategy,
 				leaseUtilizationSettings,
 				jobQueueSettings,
+				computeQueueAwsMetricSettings,
 				properties);
 		}
 
@@ -125,6 +128,7 @@ namespace Horde.Build.Agents.Pools
 		/// <param name="sizeStrategy">New pool sizing strategy for the pool</param>
 		/// <param name="leaseUtilizationSettings">Settings for lease utilization-based strategy</param>
 		/// <param name="jobQueueSettings">Settings for job queue-based strategy</param>
+		/// <param name="computeQueueAwsMetricSettings">Settings for compute queue AWS metric strategy</param>
 		/// <param name="useDefaultStrategy">Whether to use the default strategy</param>
 		/// <returns>Async task object</returns>
 		public async Task<IPool?> UpdatePoolAsync(
@@ -141,6 +145,7 @@ namespace Horde.Build.Agents.Pools
 			PoolSizeStrategy? sizeStrategy = null,
 			LeaseUtilizationSettings? leaseUtilizationSettings = null,
 			JobQueueSettings? jobQueueSettings = null,
+			ComputeQueueAwsMetricSettings? computeQueueAwsMetricSettings = null,
 			bool? useDefaultStrategy = null)
 		{
 			for (; pool != null; pool = await _pools.GetAsync(pool.Id))
@@ -159,6 +164,7 @@ namespace Horde.Build.Agents.Pools
 					sizeStrategy: sizeStrategy,
 					leaseUtilizationSettings: leaseUtilizationSettings,
 					jobQueueSettings: jobQueueSettings,
+					computeQueueAwsMetricSettings: computeQueueAwsMetricSettings,
 					useDefaultStrategy: useDefaultStrategy);
 				
 				if (newPool != null)
diff --git a/Engine/Source/Programs/Horde/Horde.Build/Agents/Pools/PoolsController.cs b/Engine/Source/Programs/Horde/Horde.Build/Agents/Pools/PoolsController.cs
--- a/Engine/Source/Programs/Horde/Horde.Build/Agents/Pools/PoolsController.cs
+++ b/Engine/Source/Programs/Horde/Horde.Build/Agents/Pools/PoolsController.cs
@@ -58,6 +58,7 @@ namespace Horde.Build.Agents.Pools
 
 			LeaseUtilizationSettings? luSettings = create.LeaseUtilizationSettings?.Convert();
 			JobQueueSettings? jqSettings = create.JobQueueSettings?.Convert();
+			ComputeQueueAwsMetricSettings? cqamSettings = create.ComputeQueueAwsMetricSettings?.Convert();
 
 			TimeSpan? conformInterval = create.ConformInterval == null ? null : TimeSpan.FromHours(create.ConformInterval.Value);
 			TimeSpan? scaleOutCooldown = create.ScaleOutCooldown == null ? null : TimeSpan.FromSeconds(create.ScaleOutCooldown.Value);
@@ -65,7 +66,7 @@ namespace Horde.Build.Agents.Pools
 
 			IPool newPool = await _poolService.CreatePoolAsync(
 				create.Name, create.Condition, create.EnableAutoscaling, create.MinAgents, create.NumReserveAgents, conformInterval,
-				scaleOutCooldown, scaleInCooldown, create.SizeStrategy, luSettings, jqSettings, create.Properties);
+				scaleOutCooldown, scaleInCooldown, create.SizeStrategy, luSettings, jqSettings, cqamSettings, create.Properties);
 			return new CreatePoolResponse(newPool.Id.ToString());
 		}
 
@@ -150,7 +151,7 @@ namespace Horde.Build.Agents.Pools
 
 			await _poolService.UpdatePoolAsync(pool, update.Name, update.Condition, update.EnableAutoscaling,
 				update.MinAgents, update.NumReserveAgents, update.Properties, conformInterval, scaleOutCooldown, scaleInCooldown, update.SizeStrategy,
-				update.LeaseUtilizationSettings?.Convert(), update.JobQueueSettings?.Convert(), update.UseDefaultStrategy);
+				update.LeaseUtilizationSettings?.Convert(), update.JobQueueSettings?.Convert(), update.ComputeQueueAwsMetricSettings?.Convert(), update.UseDefaultStrategy);
 			return new OkResult();
 		}
 
diff --git a/Engine/Source/Programs/Horde/Horde.Build/Horde.Build.csproj b/Engine/Source/Programs/Horde/Horde.Build/Horde.Build.csproj
--- a/Engine/Source/Programs/Horde/Horde.Build/Horde.Build.csproj
+++ b/Engine/Source/Programs/Horde/Horde.Build/Horde.Build.csproj
@@ -20,6 +20,7 @@
 
   <ItemGroup>
     <PackageReference Include="AWSSDK.AutoScaling" Version="3.7.10.43" />
+    <PackageReference Include="AWSSDK.CloudWatch" Version="3.7.6.3" />
     <PackageReference Include="AWSSDK.EC2" Version="3.7.75.1" />
     <PackageReference Include="AWSSDK.Extensions.NETCore.Setup" Version="3.7.2" />
     <PackageReference Include="AWSSDK.S3" Version="3.7.9.19" />
diff --git a/Engine/Source/Programs/Horde/Horde.Build/Startup.cs b/Engine/Source/Programs/Horde/Horde.Build/Startup.cs
--- a/Engine/Source/Programs/Horde/Horde.Build/Startup.cs
+++ b/Engine/Source/Programs/Horde/Horde.Build/Startup.cs
@@ -12,9 +12,8 @@ using System.Text.Json;
 using System.Text.Json.Serialization;
 using System.Threading;
 using System.Threading.Tasks;
+using Amazon.CloudWatch;
 using Amazon.Extensions.NETCore.Setup;
-using Amazon.Runtime;
-using Amazon.SecurityToken.Model;
 using EpicGames.AspNet;
 using EpicGames.Core;
 using EpicGames.Horde.Storage;
@@ -94,7 +93,7 @@ using Horde.Build.Agents.Fleet.Providers;
 using Horde.Build.Server.Notices;
 using Horde.Build.Notifications.Sinks;
 using EpicGames.Horde.Storage.Bundles;
-using Microsoft.Extensions.Caching.Memory;
+using StatusCode = Grpc.Core.StatusCode;
 
 namespace Horde.Build
 {
@@ -389,6 +388,7 @@ namespace Horde.Build
 			services.AddSingleton<LeaseUtilizationStrategy>();
 			services.AddSingleton<JobQueueStrategy>();
 			services.AddSingleton<NoOpPoolSizeStrategy>();
+			services.AddSingleton<ComputeQueueAwsMetricStrategy>();
 			
 			switch (settings.FleetManager)
 			{
@@ -464,6 +464,10 @@ namespace Horde.Build
 			services.AddSingleton<ITreeStore<ReplicationService>>(sp => CreateTreeStore(sp, settings.CommitStorage).ForType<ReplicationService>());
 
 			services.AddHordeStorage(settings => configSection.GetSection("Storage").Bind(settings));
+			
+			AWSOptions awsOptions = Configuration.GetAWSOptions();
+			services.AddDefaultAWSOptions(awsOptions);
+			services.AddAWSService<IAmazonCloudWatch>();
 
 			ConfigureLogStorage(services);
 
@@ -636,6 +640,7 @@ namespace Horde.Build
 			services.AddSingleton<ConformTaskSource>();
 			services.AddHostedService<ConformTaskSource>(provider => provider.GetRequiredService<ConformTaskSource>());
 			services.AddSingleton<ComputeService>();
+			services.AddSingleton<IComputeService, ComputeService>();
 
 			services.AddSingleton<ITaskSource, UpgradeTaskSource>();
 			services.AddSingleton<ITaskSource, ShutdownTaskSource>();
