Horde: Fix bug with missing empty constructor for pool size strategy

diff --git a/Engine/Source/Programs/Horde/Horde.Build/Agents/Pools/PoolMessages.cs b/Engine/Source/Programs/Horde/Horde.Build/Agents/Pools/PoolMessages.cs
--- a/Engine/Source/Programs/Horde/Horde.Build/Agents/Pools/PoolMessages.cs
+++ b/Engine/Source/Programs/Horde/Horde.Build/Agents/Pools/PoolMessages.cs
@@ -81,13 +81,20 @@ namespace Horde.Build.Agents.Pools
 		/// <summary>
 		/// Compute cluster ID to observe
 		/// </summary>
-		public string ComputeClusterId { get; set; }
+		public string ComputeClusterId { get; set; } = null!;
 
 		/// <summary>
 		/// AWS CloudWatch namespace to write metrics in
 		/// </summary>
-		public string Namespace { get; set; }
+		public string Namespace { get; set; } = null!;
 
+		/// <summary>
+		/// Constructor
+		/// </summary>
+		public ComputeQueueAwsMetricSettingsMessage()
+		{
+		}
+		
 		/// <summary>
 		/// Construct a public REST API representation from the internal one
 		/// </summary>
