Horde: Skip CloudWatch metrics calculated as NaN

diff --git a/Engine/Source/Programs/Horde/Horde.Build/Agents/Fleet/ComputeQueueAwsMetricStrategy.cs b/Engine/Source/Programs/Horde/Horde.Build/Agents/Fleet/ComputeQueueAwsMetricStrategy.cs
--- a/Engine/Source/Programs/Horde/Horde.Build/Agents/Fleet/ComputeQueueAwsMetricStrategy.cs
+++ b/Engine/Source/Programs/Horde/Horde.Build/Agents/Fleet/ComputeQueueAwsMetricStrategy.cs
@@ -94,6 +94,12 @@ namespace Horde.Build.Agents.Fleet
 				double numQueuedTasksPerAgent = numQueuedComputeTasks / (double)numAgents;
 				double numQueuedTasksPerCores = numQueuedComputeTasks / (double)totalCpuCores;
 
+				if (Double.IsNaN(numQueuedTasksPerAgent) || Double.IsNaN(numQueuedTasksPerCores))
+				{
+					_logger.LogWarning("Pool {PoolId} has no agents or unable to count CPU cores. numQueuedTasksPerAgent={NumQueuedTasksPerAgent} numQueuedTasksPerCores={NumQueuedTasksPerCores}", poolSizeData.Pool.Id, numQueuedTasksPerAgent, numQueuedTasksPerCores);
+					continue;
+				}
+
 				DateTime now = DateTime.UtcNow;
 				
 				if (!metricsPerCloudWatchNamespace.TryGetValue(settings.Namespace, out List<MetricDatum>? metricDatums))
