Horde: Another fix to zero agent or core counts in CloudWatch reporting

diff --git a/Engine/Source/Programs/Horde/Horde.Build/Agents/Fleet/ComputeQueueAwsMetricStrategy.cs b/Engine/Source/Programs/Horde/Horde.Build/Agents/Fleet/ComputeQueueAwsMetricStrategy.cs
--- a/Engine/Source/Programs/Horde/Horde.Build/Agents/Fleet/ComputeQueueAwsMetricStrategy.cs
+++ b/Engine/Source/Programs/Horde/Horde.Build/Agents/Fleet/ComputeQueueAwsMetricStrategy.cs
@@ -91,14 +91,8 @@ namespace Horde.Build.Agents.Fleet
 				int totalCpuCores = poolSizeData.Agents.Select(x => x.Resources.TryGetValue(KnownPropertyNames.LogicalCores, out int numCpuCores) ? numCpuCores : 0).Sum();
 				int numQueuedComputeTasks = await _computeService.GetNumQueuedTasksForPoolAsync(new ClusterId(settings.ComputeClusterId), poolSizeData.Pool);
 				
-				double numQueuedTasksPerAgent = numQueuedComputeTasks / (double)numAgents;
-				double numQueuedTasksPerCores = numQueuedComputeTasks / (double)totalCpuCores;
-
-				if (Double.IsNaN(numQueuedTasksPerAgent) || Double.IsNaN(numQueuedTasksPerCores))
-				{
-					_logger.LogWarning("Pool {PoolId} has no agents or unable to count CPU cores. numQueuedTasksPerAgent={NumQueuedTasksPerAgent} numQueuedTasksPerCores={NumQueuedTasksPerCores}", poolSizeData.Pool.Id, numQueuedTasksPerAgent, numQueuedTasksPerCores);
-					continue;
-				}
+				double numQueuedTasksPerAgent = numQueuedComputeTasks / (double)Math.Max(numAgents, 1);
+				double numQueuedTasksPerCores = numQueuedComputeTasks / (double)Math.Max(totalCpuCores, 1);
 
 				DateTime now = DateTime.UtcNow;
 				
