Horde: Add span tracking to device service tick and some additional logging
#rnx

diff --git a/Engine/Source/Programs/Horde/Horde.Build/Devices/DeviceService.cs b/Engine/Source/Programs/Horde/Horde.Build/Devices/DeviceService.cs
--- a/Engine/Source/Programs/Horde/Horde.Build/Devices/DeviceService.cs
+++ b/Engine/Source/Programs/Horde/Horde.Build/Devices/DeviceService.cs
@@ -19,6 +19,9 @@ using HordeCommon;
 using Microsoft.Extensions.Hosting;
 using Microsoft.Extensions.Logging;
 using MongoDB.Bson;
+using OpenTracing;
+using OpenTracing.Util;
+
 
 namespace Horde.Build.Devices
 {
@@ -172,11 +175,13 @@ namespace Horde.Build.Devices
 		/// </summary>
 		async ValueTask TickAsync(CancellationToken stoppingToken)
 		{
-
 			if (!stoppingToken.IsCancellationRequested)
 			{
+				using IScope scope = GlobalTracer.Instance.BuildSpan("DeviceService.TickAsync").StartActive();
+
 				try
 				{
+					_logger.LogInformation("Expiring reservations");
 					await _devices.ExpireReservationsAsync();
 				}
 				catch (Exception ex)
@@ -188,6 +193,8 @@ namespace Horde.Build.Devices
 				// expire shared devices and notifications
 				try
 				{
+
+					_logger.LogInformation("Sending shared device notifications");
 					List<(UserId, IDevice)>? expireNotifications = await _devices.ExpireNotificatonsAsync();
 					if (expireNotifications != null && expireNotifications.Count > 0)
 					{
@@ -197,6 +204,7 @@ namespace Horde.Build.Devices
 						}
 					}
 
+					_logger.LogInformation("Expiring shared device checkouts");
 					List<(UserId, IDevice)>? expireCheckouts = await _devices.ExpireCheckedOutAsync();
 					if (expireCheckouts != null && expireCheckouts.Count > 0)
 					{
@@ -218,6 +226,7 @@ namespace Horde.Build.Devices
 					if (_poolTelemetryTick >= _poolTelemetryMinutes)
 					{
 						_poolTelemetryTick = 0;
+						_logger.LogInformation("Updating pool telemetry");
 						await _devices.CreatePoolTelemetrySnapshot();						
 					}
 					else
