Horde: Fix stream device telemetry as mongo doesn't supported dotted names in keys
#rnx

diff --git a/Engine/Source/Programs/Horde/Horde.Build/Devices/DeviceCollection.cs b/Engine/Source/Programs/Horde/Horde.Build/Devices/DeviceCollection.cs
--- a/Engine/Source/Programs/Horde/Horde.Build/Devices/DeviceCollection.cs
+++ b/Engine/Source/Programs/Horde/Horde.Build/Devices/DeviceCollection.cs
@@ -332,6 +332,16 @@ namespace Horde.Build.Devices
 
 		}
 
+		class StreamDeviceTelemetry : IStreamDeviceTelemetry
+		{
+			/// <inheritdoc/>
+			public StreamId StreamId { get; set; } = new StreamId();
+
+			/// <inheritdoc/>
+			public List<DeviceId> DeviceIds { get; set; } = new List<DeviceId>();
+		}
+
+
 		class DevicePlatformTelemetryDocument : IDevicePlatformTelemetry
 		{
 			/// <summary>
@@ -386,15 +396,15 @@ namespace Horde.Build.Devices
 			/// Stream devices that are reserved
 			/// </summary>
 			[BsonIgnoreIfNull]
-			public Dictionary<StreamId, List<DeviceId>>? StreamDevices { get; set; } = new Dictionary<StreamId, List<DeviceId>>();
-			IReadOnlyDictionary<StreamId, IReadOnlyList<DeviceId>> IDevicePlatformTelemetry.StreamDevices => StreamDevices?.ToDictionary(kvp => kvp.Key, kvp => kvp.Value as IReadOnlyList<DeviceId>) ?? new Dictionary<StreamId, IReadOnlyList<DeviceId>>();
+			public List<StreamDeviceTelemetry> StreamDevices { get; set; } = new List<StreamDeviceTelemetry>();
+			IReadOnlyList<IStreamDeviceTelemetry> IDevicePlatformTelemetry.StreamDevices => StreamDevices;
 
 			[BsonConstructor]
 			private DevicePlatformTelemetryDocument()
 			{
 			}
 
-			public DevicePlatformTelemetryDocument(DevicePlatformId platformId, int available, int reserved, int maintenance, int problem, int disabled, Dictionary<StreamId, List<DeviceId>> streamDevices)
+			public DevicePlatformTelemetryDocument(DevicePlatformId platformId, int available, int reserved, int maintenance, int problem, int disabled, List<StreamDeviceTelemetry> streamDevices)
 			{
 				TelemetryId = ObjectId.GenerateNewId();
 				CreateTimeUtc = DateTime.UtcNow;
@@ -1164,7 +1174,13 @@ namespace Horde.Build.Devices
 				foreach (KeyValuePair<DevicePlatformId, DevicePoolTelemetryHelper> platform in helpers)
 				{
 					DevicePoolTelemetryHelper helper = platform.Value;
-					platformTelemtry.Add(new DevicePlatformTelemetryDocument(platform.Key, helper.available, helper.reserved, helper.maintenance, helper.problem, helper.disabled, helper.streamDevices));
+					List<StreamDeviceTelemetry> streamDevices = new List<StreamDeviceTelemetry>();
+					foreach (KeyValuePair<StreamId, List<DeviceId>> i in helper.streamDevices)
+					{
+						streamDevices.Add(new StreamDeviceTelemetry() { StreamId = i.Key, DeviceIds = i.Value });
+					}
+
+					platformTelemtry.Add(new DevicePlatformTelemetryDocument(platform.Key, helper.available, helper.reserved, helper.maintenance, helper.problem, helper.disabled, streamDevices));
 				}
 
 				poolTelemetry[pool.Id] = platformTelemtry;
diff --git a/Engine/Source/Programs/Horde/Horde.Build/Devices/DevicesController.cs b/Engine/Source/Programs/Horde/Horde.Build/Devices/DevicesController.cs
--- a/Engine/Source/Programs/Horde/Horde.Build/Devices/DevicesController.cs
+++ b/Engine/Source/Programs/Horde/Horde.Build/Devices/DevicesController.cs
@@ -525,9 +525,9 @@ namespace Horde.Build.Devices
 					foreach (IDevicePlatformTelemetry telemetry in pool.Value)
 					{
 						Dictionary<string, List<string>> streamDevices = new Dictionary<string, List<string>>();
-						foreach(KeyValuePair<StringId<Streams.IStream>, IReadOnlyList<DeviceId>> kvp in telemetry.StreamDevices)
+						foreach(IStreamDeviceTelemetry i in telemetry.StreamDevices)
 						{
-							streamDevices[kvp.Key.ToString()] = kvp.Value.Select(x => x.ToString()).ToList();
+							streamDevices[i.StreamId.ToString()] = i.DeviceIds.Select(x => x.ToString()).ToList();
 						}
 						platformTelemetry.Add(new GetDevicePlatformTelemetryResponse(telemetry.PlatformId.ToString(), telemetry.Available, telemetry.Reserved, telemetry.Maintenance, telemetry.Problem, telemetry.Disabled, streamDevices));
 					}
diff --git a/Engine/Source/Programs/Horde/Horde.Build/Devices/IDevice.cs b/Engine/Source/Programs/Horde/Horde.Build/Devices/IDevice.cs
--- a/Engine/Source/Programs/Horde/Horde.Build/Devices/IDevice.cs
+++ b/Engine/Source/Programs/Horde/Horde.Build/Devices/IDevice.cs
@@ -335,6 +335,24 @@ namespace Horde.Build.Devices
 		public Acl? Acl { get; }
 	}
 
+	/// <summary>
+	/// Stream device telemetry for pool snapshot
+	/// </summary>
+	public interface IStreamDeviceTelemetry
+	{
+		/// <summary>
+		/// Associated stream
+		/// </summary>
+		public StreamId StreamId { get;  }
+
+		/// <summary>
+		/// Device ids contained in the snapshot
+		/// </summary>
+		public List<DeviceId> DeviceIds { get; }
+
+	}
+
+
 	/// <summary>
 	/// Platform telemetry for a device pool
 	/// </summary>
@@ -373,7 +391,7 @@ namespace Horde.Build.Devices
 		/// <summary>
 		/// Devices reserved in streams
 		/// </summary>
-		public IReadOnlyDictionary<StreamId, IReadOnlyList<DeviceId>> StreamDevices { get; }
+		public IReadOnlyList<IStreamDeviceTelemetry> StreamDevices { get; }
 
 	}
 
