Horde: Allow marking issues as fixed/acknowledged directly from Slack. Uses modal API to prompt for confirmation. Also includes refactored Slack library.

diff --git a/Engine/Source/Programs/Horde/Horde.Build/Horde.Build.csproj b/Engine/Source/Programs/Horde/Horde.Build/Horde.Build.csproj
--- a/Engine/Source/Programs/Horde/Horde.Build/Horde.Build.csproj
+++ b/Engine/Source/Programs/Horde/Horde.Build/Horde.Build.csproj
@@ -110,6 +110,7 @@
     <ProjectReference Include="..\..\Shared\EpicGames.Redis\EpicGames.Redis.csproj" />
     <ProjectReference Include="..\..\Shared\EpicGames.Serialization\EpicGames.Serialization.csproj" />
     <ProjectReference Include="..\..\..\ThirdParty\Perforce\P4Api.Net\p4api.net\p4api.net.csproj" />
+    <ProjectReference Include="..\..\Shared\EpicGames.Slack\EpicGames.Slack.csproj" />
   </ItemGroup>
 
   <ItemGroup Condition="$(WithHordeStorage)">
diff --git a/Engine/Source/Programs/Horde/Horde.Build/Notifications/Sinks/SlackNotificationSink.cs b/Engine/Source/Programs/Horde/Horde.Build/Notifications/Sinks/SlackNotificationSink.cs
--- a/Engine/Source/Programs/Horde/Horde.Build/Notifications/Sinks/SlackNotificationSink.cs
+++ b/Engine/Source/Programs/Horde/Horde.Build/Notifications/Sinks/SlackNotificationSink.cs
@@ -15,6 +15,9 @@ using System.Text.RegularExpressions;
 using System.Threading;
 using System.Threading.Tasks;
 using EpicGames.Core;
+using EpicGames.Slack;
+using EpicGames.Slack.Blocks;
+using EpicGames.Slack.Elements;
 using Horde.Build.Configuration;
 using Horde.Build.Devices;
 using Horde.Build.Issues;
@@ -27,7 +30,6 @@ using Horde.Build.Server;
 using Horde.Build.Streams;
 using Horde.Build.Users;
 using Horde.Build.Utilities;
-using Horde.Build.Utilities.Slack.BlockKit;
 using HordeCommon;
 using Microsoft.AspNetCore.Hosting;
 using Microsoft.Extensions.Caching.Memory;
@@ -52,15 +54,19 @@ namespace Horde.Build.Notifications.Sinks
 	/// </summary>
 	public class SlackNotificationSink : BackgroundService, INotificationSink, IAvatarService
 	{
-		const string AddReactionUrl = "https://slack.com/api/reactions.add";
-		const string RemoveReactionUrl = "https://slack.com/api/reactions.remove";
-		const string ConversationsInviteUrl = "https://slack.com/api/conversations.invite";
-		const string PostMessageUrl = "https://slack.com/api/chat.postMessage";
-		const string UpdateMessageUrl = "https://slack.com/api/chat.update";
-		const string GetPermalinkUrl = "https://slack.com/api/chat.getPermalink";
-
 		const bool defaultAllowMentions = true;
 
+		const int MaxLineLength = 2048;
+
+		// The color to use for error messages.
+		const string ErrorColor = "#ec4c47";
+		
+		// The color to use for warning messages.
+		const string WarningColor = "#f7d154";
+
+		// The color to use for success messages.
+		const string SuccessColor = "#4db507";
+
 		class SocketResponse
 		{
 			[JsonPropertyName("ok")]
@@ -82,25 +88,13 @@ namespace Horde.Build.Notifications.Sinks
 			public EventPayload? Payload { get; set; }
 		}
 
-		class ReactionMessage
-		{
-			[JsonPropertyName("channel")]
-			public string? Channel { get; set; }
-
-			[JsonPropertyName("timestamp")]
-			public string? Ts { get; set; }
-
-			[JsonPropertyName("name")]
-			public string? Name { get; set; }
-		}
-
-		class InviteMessage
+		class EventResponse
 		{
-			[JsonPropertyName("channel")]
-			public string? Channel { get; set; }
+			[JsonPropertyName("envelope_id")]
+			public string? EnvelopeId { get; set; }
 
-			[JsonPropertyName("users")]
-			public string? Users { get; set; } // Comma separated list of ids
+			[JsonPropertyName("payload")]
+			public object? Payload { get; set; }
 		}
 
 		class EventPayload
@@ -108,56 +102,32 @@ namespace Horde.Build.Notifications.Sinks
 			[JsonPropertyName("type")]
 			public string? Type { get; set; }
 
+			[JsonPropertyName("trigger_id")]
+			public string? TriggerId { get; set; }
+
 			[JsonPropertyName("user")]
-			public UserInfo? User { get; set; }
+			public SlackUser? User { get; set; }
 
 			[JsonPropertyName("response_url")]
 			public string? ResponseUrl { get; set; }
 
 			[JsonPropertyName("actions")]
 			public List<ActionInfo> Actions { get; set; } = new List<ActionInfo>();
-		}
-
-		class UserResponse
-		{
-			[JsonPropertyName("ok")]
-			public bool Ok { get; set; }
 
-			[JsonPropertyName("user")]
-			public UserInfo? User { get; set; }
+			[JsonPropertyName("view")]
+			public SlackView? View { get; set; }
 		}
 
-		class UserInfo
-		{
-			[JsonPropertyName("id")]
-			public string? UserName { get; set; }
-
-			[JsonPropertyName("profile")]
-			public UserProfile? Profile { get; set; }
-		}
-
-		class UserProfile
+		class ActionInfo
 		{
-			[JsonPropertyName("is_custom_image")]
-			public bool IsCustomImage { get; set; }
-
-			[JsonPropertyName("image_24")]
-			public string? Image24 { get; set; }
+			[JsonPropertyName("blockId")]
+			public string BlockId { get; set; } = String.Empty;
 
-			[JsonPropertyName("image_32")]
-			public string? Image32 { get; set; }
-
-			[JsonPropertyName("image_48")]
-			public string? Image48 { get; set; }
+			[JsonPropertyName("action_id")]
+			public string ActionId { get; set; } = String.Empty;
 
-			[JsonPropertyName("image_72")]
-			public string? Image72 { get; set; }
-		}
-
-		class ActionInfo
-		{
 			[JsonPropertyName("value")]
-			public string? Value { get; set; }
+			public string Value { get; set; } = String.Empty;
 		}
 
 		class MessageStateDocument
@@ -187,7 +157,7 @@ namespace Horde.Build.Notifications.Sinks
 			public string? Permalink { get; set; }
 		}
 
-		class SlackUser : IAvatar
+		class SlackUserDocument : IAvatar
 		{
 			public const int CurrentVersion = 2;
 			
@@ -214,14 +184,14 @@ namespace Horde.Build.Notifications.Sinks
 			[BsonElement("v")]
 			public int Version { get; set; }
 
-			private SlackUser()
+			private SlackUserDocument()
 			{
 			}
 
-			public SlackUser(UserId id, UserInfo? info)
+			public SlackUserDocument(UserId id, SlackUser? info)
 			{
 				Id = id;
-				SlackUserId = info?.UserName;
+				SlackUserId = info?.Id;
 				if (info != null && info.Profile != null && info.Profile.IsCustomImage)
 				{
 					Image24 = info.Profile.Image24;
@@ -238,16 +208,18 @@ namespace Horde.Build.Notifications.Sinks
 		readonly IUserCollection _userCollection;
 		readonly ILogFileService _logFileService;
 		readonly StreamService _streamService;
-		readonly ConfigCollection _configCollection;
 		readonly IWebHostEnvironment _environment;
 		readonly ServerSettings _settings;
 		readonly IMongoCollection<MessageStateDocument> _messageStates;
-		readonly IMongoCollection<SlackUser> _slackUsers;
+		readonly IMongoCollection<SlackUserDocument> _slackUsers;
 		readonly HashSet<string>? _allowUsers;
 		readonly IExternalIssueService _externalIssueService;
 		readonly JsonSerializerOptions _jsonSerializerOptions;
 		readonly ILogger _logger;
 
+		readonly HttpClient _httpClient;
+		readonly SlackClient _slackClient;
+
 		/// <summary>
 		/// Map of email address to Slack user ID.
 		/// </summary>
@@ -256,20 +228,23 @@ namespace Horde.Build.Notifications.Sinks
 		/// <summary>
 		/// Constructor
 		/// </summary>
-		public SlackNotificationSink(MongoService mongoService, IssueService issueService, IUserCollection userCollection, ILogFileService logFileService, StreamService streamService, ConfigCollection configCollection, IExternalIssueService externalIssueService, IWebHostEnvironment environment, IOptions<ServerSettings> settings, ILogger<SlackNotificationSink> logger)
+		public SlackNotificationSink(MongoService mongoService, IssueService issueService, IUserCollection userCollection, ILogFileService logFileService, StreamService streamService, IExternalIssueService externalIssueService, IWebHostEnvironment environment, IOptions<ServerSettings> settings, ILogger<SlackNotificationSink> logger)
 		{
 			_issueService = issueService;
 			_userCollection = userCollection;
 			_logFileService = logFileService;
 			_streamService = streamService;
-			_configCollection = configCollection;
 			_externalIssueService = externalIssueService;
 			_environment = environment;
 			_settings = settings.Value;
 			_messageStates = mongoService.Database.GetCollection<MessageStateDocument>("Slack");
-			_slackUsers = mongoService.Database.GetCollection<SlackUser>("Slack.UsersV2");
+			_slackUsers = mongoService.Database.GetCollection<SlackUserDocument>("Slack.UsersV2");
 			_logger = logger;
 
+			_httpClient = new HttpClient();
+			_httpClient.DefaultRequestHeaders.Add("Authorization", $"Bearer {_settings.SlackToken ?? ""}");
+			_slackClient = new SlackClient(_httpClient, _logger);
+
 			_jsonSerializerOptions = new JsonSerializerOptions();
 			Startup.ConfigureJsonSerializer(_jsonSerializerOptions);
 
@@ -364,20 +339,20 @@ namespace Horde.Build.Notifications.Sinks
 				string jobUrl = _settings.DashboardUrl + "/job/" + notification.JobId;
 				sb.AppendLine($"Job `{notification.JobName}` with ID <{jobUrl}|{notification.JobId}> in pool `{notification.PoolName}`");
 			}
-			
-			Color outcomeColor = BlockKitAttachmentColors.Warning;
-			BlockKitAttachment attachment = new BlockKitAttachment();
+
+			string outcomeColor = WarningColor;
+			SlackAttachment attachment = new SlackAttachment();
 			attachment.Color = outcomeColor;
 			attachment.FallbackText = $"Job(s) scheduled in an auto-scaled pool with no agents online. Job IDs {jobIds}";
 
-			attachment.Blocks.Add(new HeaderBlock($"Jobs scheduled in empty pool", true));
-			attachment.Blocks.Add(new SectionBlock($"One or more jobs were scheduled in an auto-scaled pool but with no current agents online."));
+			attachment.AddHeader($"Jobs scheduled in empty pool", true);
+			attachment.AddSection($"One or more jobs were scheduled in an auto-scaled pool but with no current agents online.");
 			if (sb.Length > 0)
 			{
-				attachment.Blocks.Add(new SectionBlock(sb.ToString()));
+				attachment.AddSection(sb.ToString());
 			}
-			
-			await SendMessageAsync(recipient, attachments: new[] { attachment });
+
+			await SendMessageAsync(recipient, attachment);
 		}
 
 		#region Job Complete
@@ -440,12 +415,12 @@ namespace Horde.Build.Notifications.Sinks
 
 			Uri jobLink = new Uri($"{_settings.DashboardUrl}job/{job.Id}");
 
-			Color outcomeColor = jobOutcome == JobStepOutcome.Failure ? BlockKitAttachmentColors.Error : jobOutcome == JobStepOutcome.Warnings ? BlockKitAttachmentColors.Warning : BlockKitAttachmentColors.Success;
+			string outcomeColor = jobOutcome == JobStepOutcome.Failure ? ErrorColor : jobOutcome == JobStepOutcome.Warnings ? WarningColor : SuccessColor;
 
-			BlockKitAttachment attachment = new BlockKitAttachment();
+			SlackAttachment attachment = new SlackAttachment();
 			attachment.FallbackText = $"{stream.Name} - {GetJobChangeText(job)} - {job.Name} - {jobOutcome}";
 			attachment.Color = outcomeColor;
-			attachment.Blocks.Add(new SectionBlock($"*<{jobLink}|{stream.Name} - {GetJobChangeText(job)} - {job.Name}>*"));
+			attachment.AddSection($"*<{jobLink}|{stream.Name} - {GetJobChangeText(job)} - {job.Name}>*");
 
 			if (!String.IsNullOrEmpty(job.PreflightDescription))
 			{
@@ -454,12 +429,12 @@ namespace Horde.Build.Notifications.Sinks
 				{
 					description += $"\n#preflight {job.Id}";
 				}
-				attachment.Blocks.Add(new SectionBlock($"```{description}```"));
+				attachment.AddSection($"```{description}```");
 			}
 
 			if (jobOutcome == JobStepOutcome.Success)
 			{
-				attachment.Blocks.Add(new SectionBlock($"*Job Succeeded*"));
+				attachment.AddSection($"*Job Succeeded*");
 			}
 			else
 			{
@@ -487,25 +462,25 @@ namespace Horde.Build.Notifications.Sinks
 				if (failedStepStrings.Any())
 				{
 					string msg = $"*Errors*\n{String.Join(", ", failedStepStrings)}";
-					attachment.Blocks.Add(new SectionBlock(msg.Substring(0, Math.Min(msg.Length, 3000))));
+					attachment.AddSection(msg.Substring(0, Math.Min(msg.Length, 3000)));
 				}
 				else if (warningStepStrings.Any())
 				{
 					string msg = $"*Warnings*\n{String.Join(", ", warningStepStrings)}";
-					attachment.Blocks.Add(new SectionBlock(msg.Substring(0, Math.Min(msg.Length, 3000))));
+					attachment.AddSection(msg.Substring(0, Math.Min(msg.Length, 3000)));
 				}
 			}
 
 			if (job.AutoSubmit)
 			{
-				attachment.Blocks.Add(new DividerBlock());
+				attachment.AddDivider();
 				if (job.AutoSubmitChange != null)
 				{
-					attachment.Blocks.Add(new SectionBlock($"Shelved files were submitted in CL {job.AutoSubmitChange}."));
+					attachment.AddSection($"Shelved files were submitted in CL {job.AutoSubmitChange}.");
 				}
 				else
 				{
-					attachment.Color = BlockKitAttachmentColors.Warning;
+					attachment.Color = WarningColor;
 
 					string autoSubmitMessage = String.Empty;
 					if (!String.IsNullOrEmpty(job.AutoSubmitMessage))
@@ -513,11 +488,11 @@ namespace Horde.Build.Notifications.Sinks
 						autoSubmitMessage = $"\n\n```{job.AutoSubmitMessage}```";
 					}
 
-					attachment.Blocks.Add(new SectionBlock($"Files in CL *{job.PreflightChange}* were *not submitted*. Please resolve the following issues and submit manually.{autoSubmitMessage}"));
+					attachment.AddSection($"Files in CL *{job.PreflightChange}* were *not submitted*. Please resolve the following issues and submit manually.{autoSubmitMessage}");
 				}
 			}
 
-			return SendMessageAsync(recipient, attachments: new[] { attachment });
+			return SendMessageAsync(recipient, attachment);
 		}
 
 		#endregion
@@ -550,15 +525,15 @@ namespace Horde.Build.Notifications.Sinks
 			Uri jobStepLink = new Uri($"{_settings.DashboardUrl}job/{job.Id}?step={step.Id}");
 			Uri jobStepLogLink = new Uri($"{_settings.DashboardUrl}log/{step.LogId}");
 
-			Color outcomeColor = step.Outcome == JobStepOutcome.Failure ? BlockKitAttachmentColors.Error : step.Outcome == JobStepOutcome.Warnings ? BlockKitAttachmentColors.Warning : BlockKitAttachmentColors.Success;
+			string outcomeColor = step.Outcome == JobStepOutcome.Failure ? ErrorColor : step.Outcome == JobStepOutcome.Warnings ? WarningColor : SuccessColor;
 
-			BlockKitAttachment attachment = new BlockKitAttachment();
+			SlackAttachment attachment = new SlackAttachment();
 			attachment.FallbackText = $"{stream.Name} - {GetJobChangeText(job)} - {job.Name} - {node.Name} - {step.Outcome}";
 			attachment.Color = outcomeColor;
-			attachment.Blocks.Add(new SectionBlock($"*<{jobStepLink}|{stream.Name} - {GetJobChangeText(job)} - {job.Name} - {node.Name}>*"));
+			attachment.AddSection($"*<{jobStepLink}|{stream.Name} - {GetJobChangeText(job)} - {job.Name} - {node.Name}>*");
 			if (step.Outcome == JobStepOutcome.Success)
 			{
-				attachment.Blocks.Add(new SectionBlock($"*Job Step Succeeded*"));
+				attachment.AddSection($"*Job Step Succeeded*");
 			}
 			else
 			{
@@ -568,10 +543,10 @@ namespace Horde.Build.Notifications.Sinks
 				if (errors.Any())
 				{
 					string errorSummary = errors.Count > MaxJobStepEvents ? $"*Errors (First {MaxJobStepEvents} shown)*" : $"*Errors*";
-					attachment.Blocks.Add(new SectionBlock(errorSummary));
+					attachment.AddSection(errorSummary);
 					foreach (ILogEventData error in errors.Take(MaxJobStepEvents))
 					{
-						attachment.Blocks.Add(QuoteBlock(error.Message));
+						attachment.AddSection(QuoteText(error.Message));
 					}
 				}
 				else if (warnings.Any())
@@ -580,14 +555,14 @@ namespace Horde.Build.Notifications.Sinks
 					eventStrings.Add(warningSummary);
 					foreach (ILogEventData warning in warnings.Take(MaxJobStepEvents))
 					{
-						attachment.Blocks.Add(QuoteBlock(warning.Message));
+						attachment.AddSection(QuoteText(warning.Message));
 					}
 				}
 
-				attachment.Blocks.Add(new SectionBlock($"<{jobStepLogLink}|View Job Step Log>"));
+				attachment.AddSection($"<{jobStepLogLink}|View Job Step Log>");
 			}
 
-			return SendMessageAsync(recipient, attachments: new[] { attachment });
+			return SendMessageAsync(recipient, attachment);
 		}
 
 		#endregion
@@ -610,15 +585,15 @@ namespace Horde.Build.Notifications.Sinks
 		{
 			Uri labelLink = new Uri($"{_settings.DashboardUrl}job/{job.Id}?label={labelIdx}");
 
-			Color outcomeColor = outcome == LabelOutcome.Failure ? BlockKitAttachmentColors.Error : outcome == LabelOutcome.Warnings ? BlockKitAttachmentColors.Warning : BlockKitAttachmentColors.Success;
+			string outcomeColor = outcome == LabelOutcome.Failure ? ErrorColor : outcome == LabelOutcome.Warnings ? WarningColor : SuccessColor;
 
-			BlockKitAttachment attachment = new BlockKitAttachment();
+			SlackAttachment attachment = new SlackAttachment();
 			attachment.FallbackText = $"{stream.Name} - {GetJobChangeText(job)} - {job.Name} - Label {label.DashboardName} - {outcome}";
 			attachment.Color = outcomeColor;
-			attachment.Blocks.Add(new SectionBlock($"*<{labelLink}|{stream.Name} - {GetJobChangeText(job)} - {job.Name} - Label {label.DashboardName}>*"));
+			attachment.AddSection($"*<{labelLink}|{stream.Name} - {GetJobChangeText(job)} - {job.Name} - Label {label.DashboardName}>*");
 			if (outcome == LabelOutcome.Success)
 			{
-				attachment.Blocks.Add(new SectionBlock($"*Label Succeeded*"));
+				attachment.AddSection($"*Label Succeeded*");
 			}
 			else
 			{
@@ -639,15 +614,15 @@ namespace Horde.Build.Notifications.Sinks
 
 				if (failedStepStrings.Any())
 				{
-					attachment.Blocks.Add(new SectionBlock($"*Errors*\n{String.Join(", ", failedStepStrings)}"));
+					attachment.AddSection($"*Errors*\n{String.Join(", ", failedStepStrings)}");
 				}
 				else if (warningStepStrings.Any())
 				{
-					attachment.Blocks.Add(new SectionBlock($"*Warnings*\n{String.Join(", ", warningStepStrings)}"));
+					attachment.AddSection($"*Warnings*\n{String.Join(", ", warningStepStrings)}");
 				}
 			}
 
-			return SendMessageAsync(recipient, attachments: new[] { attachment });
+			return SendMessageAsync(recipient, attachment);
 		}
 
 		#endregion
@@ -757,32 +732,6 @@ namespace Horde.Build.Notifications.Sinks
 			await UpdateReportsAsync(issue, details.Spans);
 		}
 
-		async Task AddReactionAsync(string channel, string ts, string name)
-		{
-			ReactionMessage message = new ReactionMessage();
-			message.Channel = channel;
-			message.Ts = ts;
-			message.Name = name;
-			await SendRequestAsync<SlackResponse>(AddReactionUrl, message);
-		}
-
-		async Task RemoveReactionAsync(string channel, string ts, string name)
-		{
-			ReactionMessage message = new ReactionMessage();
-			message.Channel = channel;
-			message.Ts = ts;
-			message.Name = name;
-			await SendRequestAsync<SlackResponse>(RemoveReactionUrl, message);
-		}
-
-		async Task InviteUsersAsync(string channel, List<string> userIds)
-		{
-			InviteMessage message = new InviteMessage();
-			message.Channel = channel;
-			message.Users = String.Join(",", userIds);
-			await SendRequestAsync<SlackResponse>(ConversationsInviteUrl, message);
-		}
-
 		async Task InviteUsersAsync(string channel, IEnumerable<UserId> userIds)
 		{
 			List<string> slackUserIds = new List<string>();
@@ -801,7 +750,7 @@ namespace Horde.Build.Notifications.Sinks
 
 			if (slackUserIds.Count > 0)
 			{
-				await InviteUsersAsync(channel, slackUserIds);
+				await _slackClient.InviteUsersAsync(channel, slackUserIds);
 			}
 		}
 
@@ -861,19 +810,31 @@ namespace Horde.Build.Notifications.Sinks
 					summary.AppendLine($"From {FormatJobStep(span.FirstFailure, span.NodeName)}:");
 					foreach (ILogEventData eventDataItem in eventDataItems)
 					{
-						summary.AppendLine(QuoteText(eventDataItem.Message, TextObject.MaxLength));
+						summary.AppendLine(QuoteText(eventDataItem.Message, MaxLineLength));
 					}
 					if (events.Count > eventDataItems.Count)
 					{
 						summary.AppendLine("```...```");
 					}
-					string? summaryTs = await SendMessageToThread(triageChannel, state.Ts, summary.ToString());
+					string? summaryTs = await _slackClient.PostMessageAsync(triageChannel, state.Ts, summary.ToString());
 
 					// Permalink to the summary text so we link inside the thread rather than just to the original message
 					string? permalink = null;
 					if (summaryTs != null)
 					{
-						permalink = await GetPermalinkAsync(state.Channel, summaryTs);
+						permalink = await _slackClient.GetPermalinkAsync(state.Channel, summaryTs);
+					}
+
+					// Post controls for interacting with the issue
+					{
+						SlackMessage message = new SlackMessage();
+
+						ActionsBlock actions = message.AddActions();
+						actions.AddButton("Acknowledge", value: $"issue_{issue.Id}_ack", style: ButtonStyle.Primary);
+						actions.AddButton("Not Me", value: $"issue_{issue.Id}_decline", style: ButtonStyle.Danger);
+						actions.AddButton("Mark Fixed", value: $"issue_{issue.Id}_markfixed");
+
+						await _slackClient.PostMessageAsync(triageChannel, state.Ts, message);
 					}
 
 					// If it has an owner, show that
@@ -881,7 +842,7 @@ namespace Horde.Build.Notifications.Sinks
 					if (issue.OwnerId != null)
 					{
 						string mention = await FormatMentionAsync(issue.OwnerId.Value, workflow.AllowMentions);
-						await SendMessageToThread(triageChannel, state.Ts, $"Assigned to {mention}");
+						await _slackClient.PostMessageAsync(triageChannel, state.Ts, $"Assigned to {mention}");
 						inviteUserIds.Add(issue.OwnerId.Value);
 					}
 					else
@@ -901,7 +862,7 @@ namespace Horde.Build.Notifications.Sinks
 							}
 
 							string suspectMessage = $"Possibly {StringUtils.FormatList(suspectList, "or")}.";
-							await SendMessageToThread(triageChannel, state.Ts, suspectMessage);
+							await _slackClient.PostMessageAsync(triageChannel, state.Ts, suspectMessage);
 						}
 					}
 
@@ -915,11 +876,11 @@ namespace Horde.Build.Notifications.Sinks
 
 				if (issue.AcknowledgedAt != null)
 				{
-					await AddReactionAsync(state.Channel, state.Ts, "eyes");
+					await _slackClient.AddReactionAsync(state.Channel, state.Ts, "eyes");
 				}
 				else
 				{
-					await RemoveReactionAsync(state.Channel, state.Ts, "eyes");
+					await _slackClient.RemoveReactionAsync(state.Channel, state.Ts, "eyes");
 				}
 
 				IIssueSpan? fixFailedSpan = null;
@@ -970,20 +931,20 @@ namespace Horde.Build.Notifications.Sinks
 
 				if (fixFailedSpan != null)
 				{
-					await AddReactionAsync(state.Channel, state.Ts, "x");
+					await _slackClient.AddReactionAsync(state.Channel, state.Ts, "x");
 				}
 				else
 				{
-					await RemoveReactionAsync(state.Channel, state.Ts, "x");
+					await _slackClient.RemoveReactionAsync(state.Channel, state.Ts, "x");
 				}
 
 				if (issue.ResolvedAt != null && fixFailedSpan == null)
 				{
-					await AddReactionAsync(state.Channel, state.Ts, "tick");
+					await _slackClient.AddReactionAsync(state.Channel, state.Ts, "tick");
 				}
 				else
 				{
-					await RemoveReactionAsync(state.Channel, state.Ts, "tick");
+					await _slackClient.RemoveReactionAsync(state.Channel, state.Ts, "tick");
 				}
 
 				if (issue.ExternalIssueKey != null)
@@ -1000,7 +961,7 @@ namespace Horde.Build.Notifications.Sinks
 			(MessageStateDocument state, bool isNew) = await AddOrUpdateMessageStateAsync(channel, eventId, null, "");
 			if (isNew)
 			{
-				string? ts = await SendMessageToThread(channel, threadTs, message);
+				string? ts = await _slackClient.PostMessageAsync(channel, threadTs, message);
 				if (ts != null)
 				{
 					await SetMessageTimestampAsync(state.Id, channel, ts);
@@ -1008,17 +969,6 @@ namespace Horde.Build.Notifications.Sinks
 			}
 		}
 
-		async Task<string?> SendMessageToThread(string channel, string threadTs, string text)
-		{
-			BlockKitMessage message = new BlockKitMessage();
-			message.Channel = channel;
-			message.ThreadTs = threadTs;
-			message.Text = text;
-
-			PostMessageResponse response = await SendRequestAsync<PostMessageResponse>(PostMessageUrl, message);
-			return response.Ts;
-		}
-
 		async Task NotifyIssueUpdatedAsync(IUser user, IIssue issue, IIssueDetails details)
 		{
 			string? slackUserId = await GetSlackUserId(user);
@@ -1049,15 +999,15 @@ namespace Horde.Build.Notifications.Sinks
 		{
 			using IDisposable scope = _logger.BeginScope("SendIssueMessageAsync (User: {SlackUser}, Issue: {IssueId})", recipient, issue.Id);
 
-			BlockKitAttachment attachment = new BlockKitAttachment();
-			attachment.Color = Color.Red;
+			SlackAttachment attachment = new SlackAttachment();
+			attachment.Color = ErrorColor;
 
 			Uri issueUrl = _settings.DashboardUrl;
 			if (details.Steps.Count > 0)
 			{
 				issueUrl = GetIssueUrl(issue, details.Steps[0]);
 			}
-			attachment.Blocks.Add(new SectionBlock($"*<{issueUrl}|Issue #{issue.Id}: {issue.Summary}>*"));
+			attachment.AddSection($"*<{issueUrl}|Issue #{issue.Id}: {issue.Summary}>*");
 
 			string streamList = StringUtils.FormatList(details.Spans.Select(x => $"*{x.StreamName}*").Distinct(StringComparer.OrdinalIgnoreCase).OrderBy(x => x, StringComparer.OrdinalIgnoreCase));
 			StringBuilder summaryBuilder = new StringBuilder($"Occurring in {streamList}.");
@@ -1076,7 +1026,7 @@ namespace Horde.Build.Notifications.Sinks
 					}
 				}
 			}
-			attachment.Blocks.Add(new SectionBlock(summaryBuilder.ToString()));
+			attachment.AddSection(summaryBuilder.ToString());
 
 			IIssueSpan? lastSpan = details.Spans.OrderByDescending(x => x.LastFailure.StepTime).FirstOrDefault();
 			if (lastSpan != null && lastSpan.LastFailure.LogId != null)
@@ -1094,11 +1044,11 @@ namespace Horde.Build.Notifications.Sinks
 					for (int idx = 0; idx < Math.Min(events.Count, 3); idx++)
 					{
 						ILogEventData data = await _logFileService.GetEventDataAsync(logFile, events[idx].LineIndex, events[idx].LineCount);
-						attachment.Blocks.Add(QuoteBlock(data.Message));
+						attachment.AddSection(QuoteText(data.Message));
 					}
 					if (events.Count > 3)
 					{
-						attachment.Blocks.Add(new SectionBlock("```...```"));
+						attachment.AddSection("```...```");
 					}
 				}
 			}
@@ -1121,13 +1071,13 @@ namespace Horde.Build.Notifications.Sinks
 				{
 					text = $":tick: Marked fixed in *CL {issue.FixChange.Value}*.";
 				}
-				attachment.Blocks.Add(new SectionBlock(text));
+				attachment.AddSection(text);
 			}
 			else if (userId != null && issue.OwnerId == userId)
 			{
 				if (issue.AcknowledgedAt.HasValue)
 				{
-					attachment.Blocks.Add(new SectionBlock($":+1: Acknowledged at {FormatSlackTime(issue.AcknowledgedAt.Value)}"));
+					attachment.AddSection($":+1: Acknowledged at {FormatSlackTime(issue.AcknowledgedAt.Value)}");
 				}
 				else
 				{
@@ -1135,7 +1085,7 @@ namespace Horde.Build.Notifications.Sinks
 					{
 						string mention = await FormatMentionAsync(issue.NominatedById.Value, allowMentions);
 						string text = $"You were nominated to fix this issue by {mention} at {FormatSlackTime(issue.NominatedAt ?? DateTime.UtcNow)}";
-						attachment.Blocks.Add(new SectionBlock(text));
+						attachment.AddSection(text);
 					}
 					else
 					{
@@ -1143,13 +1093,13 @@ namespace Horde.Build.Notifications.Sinks
 						if (changes.Count > 0)
 						{
 							string text = $"Horde has determined that {StringUtils.FormatList(changes.Select(x => $"CL {x}"), "or")} is the most likely cause for this issue.";
-							attachment.Blocks.Add(new SectionBlock(text));
+							attachment.AddSection(text);
 						}
 					}
 
 					ActionsBlock actions = new ActionsBlock();
-					actions.AddButton("Acknowledge", value: $"issue_{issue.Id}_ack_{userId}", style: ActionButton.ButtonStyle.Primary);
-					actions.AddButton("Not Me", value: $"issue_{issue.Id}_decline_{userId}", style: ActionButton.ButtonStyle.Danger);
+					actions.AddButton("Acknowledge", value: $"issue_{issue.Id}_ack_{userId}", style: ButtonStyle.Primary);
+					actions.AddButton("Not Me", value: $"issue_{issue.Id}_decline_{userId}", style: ButtonStyle.Danger);
 					attachment.Blocks.Add(actions);
 				}
 			}
@@ -1158,19 +1108,19 @@ namespace Horde.Build.Notifications.Sinks
 				string ownerMention = await FormatMentionAsync(issue.OwnerId.Value, allowMentions);
 				if (issue.AcknowledgedAt.HasValue)
 				{
-					attachment.Blocks.Add(new SectionBlock($":+1: Acknowledged by {ownerMention} at {FormatSlackTime(issue.AcknowledgedAt.Value)}"));
+					attachment.AddSection($":+1: Acknowledged by {ownerMention} at {FormatSlackTime(issue.AcknowledgedAt.Value)}");
 				}
 				else if (issue.NominatedById == null)
 				{
-					attachment.Blocks.Add(new SectionBlock($"Assigned to {ownerMention}"));
+					attachment.AddSection($"Assigned to {ownerMention}");
 				}
 				else if (issue.NominatedById == userId)
 				{
-					attachment.Blocks.Add(new SectionBlock($"You nominated {ownerMention} to fix this issue."));
+					attachment.AddSection($"You nominated {ownerMention} to fix this issue.");
 				}
 				else
 				{
-					attachment.Blocks.Add(new SectionBlock($"{ownerMention} was nominated to fix this issue by {await FormatMentionAsync(issue.NominatedById.Value, allowMentions)}"));
+					attachment.AddSection($"{ownerMention} was nominated to fix this issue by {await FormatMentionAsync(issue.NominatedById.Value, allowMentions)}");
 				}
 			}
 			else if (userId != null)
@@ -1180,16 +1130,15 @@ namespace Horde.Build.Notifications.Sinks
 				{
 					if (suspect.DeclinedAt != null)
 					{
-						attachment.Blocks.Add(new SectionBlock($":downvote: Declined at {FormatSlackTime(suspect.DeclinedAt.Value)}"));
+						attachment.AddSection($":downvote: Declined at {FormatSlackTime(suspect.DeclinedAt.Value)}");
 					}
 					else
 					{
-						attachment.Blocks.Add(new SectionBlock("Please check if any of your recently submitted changes have caused this issue."));
+						attachment.AddSection("Please check if any of your recently submitted changes have caused this issue.");
 
-						ActionsBlock actions = new ActionsBlock();
-						actions.AddButton("Will Fix", value: $"issue_{issue.Id}_accept_{userId}", style: ActionButton.ButtonStyle.Primary);
-						actions.AddButton("Not Me", value: $"issue_{issue.Id}_decline_{userId}", style: ActionButton.ButtonStyle.Danger);
-						attachment.Blocks.Add(actions);
+						ActionsBlock actions = attachment.AddActions();
+						actions.AddButton("Will Fix", value: $"issue_{issue.Id}_accept_{userId}", style: ButtonStyle.Primary);
+						actions.AddButton("Not Me", value: $"issue_{issue.Id}_decline_{userId}", style: ButtonStyle.Danger);
 					}
 				}
 			}
@@ -1211,12 +1160,12 @@ namespace Horde.Build.Notifications.Sinks
 						declinedLines.Add($":downvote: Declined by {await FormatNameAsync(suspect.AuthorId)} at {FormatSlackTime(suspect.DeclinedAt.Value)} (CL {suspect.Change})");
 					}
 				}
-				attachment.Blocks.Add(new SectionBlock(String.Join("\n", declinedLines)));
+				attachment.AddSection(String.Join("\n", declinedLines));
 			}
 
 			if (IsRecipientAllowed(recipient, "issue update"))
 			{
-				await SendOrUpdateMessageAsync(recipient, GetIssueEventId(issue), userId, attachments: new[] { attachment });
+				await SendOrUpdateMessageAsync(recipient, GetIssueEventId(issue), userId, attachment);
 			}
 		}
 
@@ -1274,61 +1223,6 @@ namespace Horde.Build.Notifications.Sinks
 			return $"<@{slackUserId}>";
 		}
 
-		async Task HandleIssueResponseAsync(EventPayload payload, ActionInfo action)
-		{
-			string? userName = payload.User?.UserName;
-			if(userName == null)
-			{
-				_logger.LogWarning("No user for message payload: {Payload}", payload);
-				return;
-			}
-			if (payload.ResponseUrl == null)
-			{
-				_logger.LogWarning("No response url for payload: {Payload}", payload);
-				return;
-			}
-
-			Match match = Regex.Match(action.Value ?? String.Empty, @"^issue_(\d+)_([a-zA-Z]+)_([a-fA-F0-9]{24})$");
-			if (!match.Success)
-			{
-				_logger.LogWarning("Could not match format of button action: {Action}", action.Value);
-				return;
-			}
-
-			int issueId = Int32.Parse(match.Groups[1].Value, NumberStyles.Integer, CultureInfo.InvariantCulture);
-			string verb = match.Groups[2].Value;
-			UserId userId = match.Groups[3].Value.ToObjectId<IUser>();
-			_logger.LogInformation("Issue {IssueId}: {Action} from {SlackUser} ({UserId})", issueId, verb, userName, userId);
-
-			if (String.Equals(verb, "ack", StringComparison.Ordinal))
-			{
-				await _issueService.UpdateIssueAsync(issueId, acknowledged: true);
-			}
-			else if (String.Equals(verb, "accept", StringComparison.Ordinal))
-			{
-				await _issueService.UpdateIssueAsync(issueId, ownerId: userId, acknowledged: true);
-			}
-			else if (String.Equals(verb, "decline", StringComparison.Ordinal))
-			{
-				await _issueService.UpdateIssueAsync(issueId, declinedById: userId);
-			}
-
-			IIssue? newIssue = await _issueService.Collection.GetIssueAsync(issueId);
-			if (newIssue != null)
-			{
-				IUser? user = await _userCollection.GetUserAsync(userId);
-				if (user != null)
-				{
-					string? recipient = await GetSlackUserId(user);
-					if (recipient != null)
-					{
-						IIssueDetails details = await _issueService.GetIssueDetailsAsync(newIssue);
-						await SendIssueMessageAsync(recipient, newIssue, details, userId, defaultAllowMentions);
-					}
-				}
-			}
-		}
-
 		static string FormatSlackTime(DateTimeOffset time)
 		{
 			return $"<!date^{time.ToUnixTimeSeconds()}^{{time}}|{time}>";
@@ -1370,7 +1264,10 @@ namespace Horde.Build.Notifications.Sinks
 				await SendIssueReportForStreamAsync(group.Channel, group.Time, report);
 			}
 
-			await SendMessageAsync(group.Channel, blocks: new[] { new DividerBlock() }, withEnvironment: false);
+			SlackMessage message = new SlackMessage();
+			message.AddDivider();
+
+			await SendMessageAsync(group.Channel, message);
 		}
 
 		async Task SendIssueReportForStreamAsync(string channel, DateTime time, IssueReport report)
@@ -1441,25 +1338,25 @@ namespace Horde.Build.Notifications.Sinks
 				}
 			}
 
-			List<BlockBase> blocks = new List<BlockBase>();
-			blocks.Add(new HeaderBlock($"Summary for {report.Stream.Name}"));
+			SlackMessage headerMessage = new SlackMessage();
+			headerMessage.AddHeader($"Summary for {report.Stream.Name}");
 
-			PostMessageResponse? response = await SendMessageAsync(channel, blocks: blocks.ToArray(), withEnvironment: false);
-			if (response != null && response.Ts != null)
+			string? ts = await SendMessageAsync(channel, headerMessage);
+			if (ts != null)
 			{
 				string reportEventId = GetReportEventId(report.Stream.Id, report.WorkflowId);
 				string json = JsonSerializer.Serialize(state, _jsonSerializerOptions);
-				await AddOrUpdateMessageStateAsync(channel, reportEventId, null, json, response.Ts);
+				await AddOrUpdateMessageStateAsync(channel, reportEventId, null, json, ts);
 
 				if (state.Blocks.Count == 0)
 				{
 					string header = ":tick: No issues open.";
-					await SendMessageAsync(channel, text: header, withEnvironment: false);
+					await SendMessageAsync(channel, header);
 				}
 
 				for (int idx = 0; idx < state.Blocks.Count; idx++)
 				{
-					string blockEventId = GetReportBlockEventId(response.Ts, idx);
+					string blockEventId = GetReportBlockEventId(ts, idx);
 					await UpdateReportBlockAsync(channel, blockEventId, time, report.Stream, state.Blocks[idx].TemplateId, issuesByBlock[idx], state.Blocks[idx].TemplateHeader);
 				}
 
@@ -1467,7 +1364,7 @@ namespace Horde.Build.Notifications.Sinks
 				{
 					double totalPct = (report.WorkflowStats.NumPassingSteps * 100.0) / report.WorkflowStats.NumSteps;
 					string header = $"*{totalPct:0.0}%* of build steps ({report.WorkflowStats.NumPassingSteps:n0}/{report.WorkflowStats.NumSteps:n0}) succeeded since last status update.";
-					await SendMessageAsync(channel, text: header, withEnvironment: false);
+					await SendMessageAsync(channel, header);
 				}
 			}
 		}
@@ -1606,7 +1503,7 @@ namespace Horde.Build.Notifications.Sinks
 				body.Append(text);
 			}
 
-			await SendOrUpdateMessageAsync(channel, eventId, null, body.ToString(), withEnvironment: false);
+			await SendOrUpdateMessageAsync(channel, eventId, null, body.ToString());
 		}
 
 		string GetSeverityPrefix(IssueSeverity severity)
@@ -1740,34 +1637,34 @@ namespace Horde.Build.Notifications.Sinks
 			}
 		}
 
-		private Task SendConfigUpdateFailureMessageAsync(string recipient, string errorMessage, string fileName, int? change = null, string? author = null, string? description = null)
+		private async Task SendConfigUpdateFailureMessageAsync(string recipient, string errorMessage, string fileName, int? change = null, string? author = null, string? description = null)
 		{
-			Color outcomeColor = BlockKitAttachmentColors.Error;
-			BlockKitAttachment attachment = new BlockKitAttachment();
+			string outcomeColor = ErrorColor;
+			SlackAttachment attachment = new SlackAttachment();
 			attachment.Color = outcomeColor;
 			attachment.FallbackText = $"Update Failure: {fileName}";
 
-			attachment.Blocks.Add(new HeaderBlock($"Config Update Failure :rip:", true));
+			attachment.AddHeader($"Config Update Failure :rip:", true);
 
-			attachment.Blocks.Add(new SectionBlock($"Horde was unable to update {fileName}"));
-			attachment.Blocks.Add(QuoteBlock(errorMessage));
+			attachment.AddSection($"Horde was unable to update {fileName}");
+			attachment.AddSection(QuoteText(errorMessage));
 			if (change != null)
 			{
 				if (author != null)
 				{
-					attachment.Blocks.Add(new SectionBlock($"Possibly due to CL: {change.Value} by <@{author}>"));
+					attachment.AddSection($"Possibly due to CL: {change.Value} by <@{author}>");
 				}
 				else
 				{
-					attachment.Blocks.Add(new SectionBlock($"Possibly due to CL: {change.Value} - (Could not determine author from P4 user)"));
+					attachment.AddSection($"Possibly due to CL: {change.Value} - (Could not determine author from P4 user)");
 				}
 				if (description != null)
 				{
-					attachment.Blocks.Add(QuoteBlock(description));
+					attachment.AddSection(QuoteText(description));
 				}
 			}
 
-			return SendMessageAsync(recipient, attachments: new[] { attachment });
+			await SendMessageAsync(recipient, attachment);
 		}
 
 		#endregion
@@ -1790,22 +1687,22 @@ namespace Horde.Build.Notifications.Sinks
 		/// <param name="recipient"></param>
 		/// <param name="file">The file</param>
 		/// <returns></returns>
-		Task SendStreamUpdateFailureMessage(string recipient, FileSummary file)
+		async Task SendStreamUpdateFailureMessage(string recipient, FileSummary file)
 		{
-			Color outcomeColor = BlockKitAttachmentColors.Error;
-			BlockKitAttachment attachment = new BlockKitAttachment();
+			string outcomeColor = ErrorColor;
+			SlackAttachment attachment = new SlackAttachment();
 			attachment.Color = outcomeColor;
 			attachment.FallbackText = $"{file.DepotPath} - Update Failure";
 
-			attachment.Blocks.Add(new HeaderBlock($"Stream Update Failure :rip:", true));
+			attachment.AddHeader($"Stream Update Failure :rip:", true);
 
-			attachment.Blocks.Add(new SectionBlock($"<!here> Horde was unable to update {file.DepotPath}"));
+			attachment.AddSection($"<!here> Horde was unable to update {file.DepotPath}");
 			if (file.Error != null)
 			{
-				attachment.Blocks.Add(QuoteBlock(file.Error));
+				attachment.AddSection(QuoteText(file.Error));
 			}
 
-			return SendMessageAsync(recipient, attachments: new[] { attachment });
+			await SendMessageAsync(recipient, attachment);
 		}
 
 		#endregion
@@ -1864,7 +1761,7 @@ namespace Horde.Build.Notifications.Sinks
 				message = message.Substring(0, 146) + "...";
 			}
 
-			BlockKitAttachment attachment = new BlockKitAttachment();
+			SlackAttachment attachment = new SlackAttachment();
 							
 			attachment.FallbackText = $"{message}";
 
@@ -1873,7 +1770,7 @@ namespace Horde.Build.Notifications.Sinks
 				attachment.FallbackText += $" - Device: {device.Name} Pool: {pool.Name}";
 			}
 				
-			attachment.Blocks.Add(new HeaderBlock(AddEnvironmentAnnotation(message), false));
+			attachment.AddHeader(message, false);
 
 			if (stream != null && job != null && step != null && node != null)
 			{
@@ -1881,27 +1778,21 @@ namespace Horde.Build.Notifications.Sinks
 				Uri jobStepLogLink = new Uri($"{_settings.DashboardUrl}log/{step.LogId}");
 
 				attachment.FallbackText += $" - {stream.Name} - {GetJobChangeText(job)} - {job.Name} - {node.Name}";
-				attachment.Blocks.Add(new SectionBlock($"*<{jobStepLink}|{stream.Name} - {GetJobChangeText(job)} - {job.Name} - {node.Name}>*"));
-				attachment.Blocks.Add(new SectionBlock($"<{jobStepLogLink}|View Job Step Log>"));
+				attachment.AddSection($"*<{jobStepLink}|{stream.Name} - {GetJobChangeText(job)} - {job.Name} - {node.Name}>*");
+				attachment.AddSection($"<{jobStepLogLink}|View Job Step Log>");
 			}
 			else
 			{
 				attachment.FallbackText += " - No job information (Gauntlet might need to be updated in stream)";
-				attachment.Blocks.Add(new SectionBlock("*No job information (Gauntlet might need to be updated in stream)*"));
+				attachment.AddSection("*No job information (Gauntlet might need to be updated in stream)*");
 			}
 
-			return SendMessageAsync(recipient, attachments: new[] { attachment });
-
+			return SendMessageAsync(recipient, attachment);
 		}
 
 		#endregion
 
-		static SectionBlock QuoteBlock(string text)
-		{
-			return new SectionBlock(QuoteText(text, TextObject.MaxLength));
-		}
-
-		static string QuoteText(string text, int maxLength)
+		static string QuoteText(string text, int maxLength = MaxLineLength)
 		{
 			maxLength -= 6;
 			if (text.Length > maxLength)
@@ -1933,9 +1824,9 @@ namespace Horde.Build.Notifications.Sinks
 			}
 		}
 
-		static bool ShouldUpdateUser(SlackUser? document)
+		static bool ShouldUpdateUser(SlackUserDocument? document)
 		{
-			if(document == null || document.Version < SlackUser.CurrentVersion)
+			if(document == null || document.Version < SlackUserDocument.CurrentVersion)
 			{
 				return true;
 			}
@@ -1957,7 +1848,7 @@ namespace Horde.Build.Notifications.Sinks
 			return (await GetSlackUser(user))?.SlackUserId;
 		}
 
-		private async Task<SlackUser?> GetSlackUser(IUser user)
+		private async Task<SlackUserDocument?> GetSlackUser(IUser user)
 		{
 			string? email = user.Email;
 			if (email == null)
@@ -1966,16 +1857,16 @@ namespace Horde.Build.Notifications.Sinks
 				return null;
 			}
 
-			SlackUser? userDocument;
+			SlackUserDocument? userDocument;
 			if (!_userCache.TryGetValue(email, out userDocument))
 			{
 				userDocument = await _slackUsers.Find(x => x.Id == user.Id).FirstOrDefaultAsync();
 				if (userDocument == null || ShouldUpdateUser(userDocument))
 				{
-					UserInfo? userInfo = await GetSlackUserInfoByEmail(email);
+					SlackUser? userInfo = await _slackClient.FindUserByEmailAsync(email);
 					if (userDocument == null || userInfo != null)
 					{
-						userDocument = new SlackUser(user.Id, userInfo);
+						userDocument = new SlackUserDocument(user.Id, userInfo);
 						await _slackUsers.ReplaceOneAsync(x => x.Id == user.Id, userDocument, new ReplaceOptions { IsUpsert = true });
 					}
 				}
@@ -1989,106 +1880,16 @@ namespace Horde.Build.Notifications.Sinks
 			return userDocument;
 		}
 
-		private async Task<UserInfo?> GetSlackUserInfoByEmail(string email)
+		private async Task<string?> SendMessageAsync(string recipient, SlackMessage message)
 		{
-			using HttpClient client = new HttpClient();
-
-			using HttpRequestMessage getUserIdRequest = new HttpRequestMessage(HttpMethod.Post, $"https://slack.com/api/users.lookupByEmail?email={email}");
-			getUserIdRequest.Headers.Add("Authorization", $"Bearer {_settings.SlackToken ?? ""}");
-
-			HttpResponseMessage responseMessage = await client.SendAsync(getUserIdRequest);
-			byte[] responseData = await responseMessage.Content.ReadAsByteArrayAsync();
-
-			UserResponse userResponse = JsonSerializer.Deserialize<UserResponse>(responseData)!;
-			if(!userResponse.Ok || userResponse.User == null)
+			if (!IsRecipientAllowed(recipient, message.Text))
 			{
-				_logger.LogWarning("Unable to find Slack user id for {Email}: {Response}", email, Encoding.UTF8.GetString(responseData));
 				return null;
 			}
-
-			return userResponse.User;
-		}
-
-		class SlackResponse
-		{
-			[JsonPropertyName("ok")]
-			public bool Ok { get; set; }
-
-			[JsonPropertyName("error")]
-			public string? Error { get; set; }
-		}
-
-		class PostMessageResponse : SlackResponse
-		{
-			[JsonPropertyName("channel")]
-			public string? Channel { get; set; }
-
-			[JsonPropertyName("ts")]
-			public string? Ts { get; set; }
-		}
-
-		class GetPermalinkRequest
-		{
-			[JsonPropertyName("channel")]
-			public string? Channel { get; set; }
-
-			[JsonPropertyName("message_ts")]
-			public string? MessageTs { get; set; }
-		}
-
-		class GetPermalinkResponse : SlackResponse
-		{
-			[JsonPropertyName("channel")]
-			public string? Channel { get; set; }
-
-			[JsonPropertyName("permalink")]
-			public string? Permalink { get; set; }
-		}
-
-		[return: NotNullIfNotNull("message")]
-		private string? AddEnvironmentAnnotation(string? message)
-		{
-			if (_environment.IsProduction())
-			{
-				return message;
-			}
-			else if (_environment.IsDevelopment())
-			{
-				return $"[DEV] {message}";
-			}
-			else if (_environment.IsStaging())
-			{
-				return $"[STAGING] {message}";
-			}
 			else
 			{
-				return $"[{_environment.EnvironmentName.ToUpperInvariant()}] {message}";
-			}
-		}
-
-		private async Task<PostMessageResponse?> SendMessageAsync(string recipient, string? text = null, BlockBase[]? blocks = null, BlockKitAttachment[]? attachments = null, bool withEnvironment = true)
-		{
-			if (!IsRecipientAllowed(recipient, text))
-			{
-				return null;
-			}
-
-			BlockKitMessage message = new BlockKitMessage();
-			message.Channel = recipient;
-			if (text != null)
-			{
-				message.Text = withEnvironment ? AddEnvironmentAnnotation(text) : text;
-			}
-			if (blocks != null)
-			{
-				message.Blocks.AddRange(blocks);
+				return await _slackClient.PostMessageAsync(recipient, message);
 			}
-			if (attachments != null)
-			{
-				message.Attachments.AddRange(attachments);
-			}
-
-			return await SendRequestAsync<PostMessageResponse>(PostMessageUrl, message);
 		}
 
 		bool IsRecipientAllowed(string recipient, string? description)
@@ -2101,22 +1902,8 @@ namespace Horde.Build.Notifications.Sinks
 			return true;
 		}
 
-		private async Task<(MessageStateDocument, bool)> SendOrUpdateMessageAsync(string recipient, string eventId, UserId? userId, string? text = null, BlockBase[]? blocks = null, BlockKitAttachment[]? attachments = null, bool withEnvironment = true)
+		private async Task<(MessageStateDocument, bool)> SendOrUpdateMessageAsync(string recipient, string eventId, UserId? userId, SlackMessage message)
 		{
-			BlockKitMessage message = new BlockKitMessage();
-			if (text != null)
-			{
-				message.Text = withEnvironment ? AddEnvironmentAnnotation(text) : text;
-			}
-			if (blocks != null)
-			{
-				message.Blocks.AddRange(blocks);
-			}
-			if (attachments != null)
-			{
-				message.Attachments.AddRange(attachments);
-			}
-
 			string requestDigest = ContentHash.MD5(JsonSerializer.Serialize(message, new JsonSerializerOptions { DefaultIgnoreCondition = JsonIgnoreCondition.WhenWritingNull })).ToString();
 
 			MessageStateDocument? prevState = await GetMessageStateAsync(recipient, eventId);
@@ -2129,78 +1916,20 @@ namespace Horde.Build.Notifications.Sinks
 			if (isNew)
 			{
 				_logger.LogInformation("Sending new slack message to {SlackUser} (msg: {MessageId})", recipient, state.Id);
-				message.Channel = recipient;
-				message.Ts = null;
 
-				PostMessageResponse response = await SendRequestAsync<PostMessageResponse>(PostMessageUrl, message);
-				if (String.IsNullOrEmpty(response.Ts) || String.IsNullOrEmpty(response.Channel))
-				{
-					_logger.LogWarning("Missing 'ts' or 'channel' field on slack response");
-				}
-
-				state.Channel = response.Channel ?? String.Empty;
-				state.Ts = response.Ts ?? String.Empty;
+				state.Channel = recipient;
+				state.Ts = await _slackClient.PostMessageAsync(recipient, message);
 
 				await SetMessageTimestampAsync(state.Id, state.Channel, state.Ts);
 			}
 			else if (!String.IsNullOrEmpty(state.Ts))
 			{
 				_logger.LogInformation("Updating existing slack message {MessageId} for user {SlackUser} ({Channel}, {MessageTs})", state.Id, recipient, state.Channel, state.Ts);
-				message.Channel = state.Channel;
-				message.Ts = state.Ts;
-				await SendRequestAsync<PostMessageResponse>(UpdateMessageUrl, message);
+				await _slackClient.UpdateMessageAsync(state.Channel, state.Ts, message);
 			}
 			return (state, isNew);
 		}
 
-		private async Task<string?> GetPermalinkAsync(string channel, string messageTs)
-		{
-			string requestUrl = $"{GetPermalinkUrl}?channel={channel}&message_ts={messageTs}";
-			using (HttpRequestMessage sendMessageRequest = new HttpRequestMessage(HttpMethod.Get, requestUrl))
-			{
-				GetPermalinkResponse response = await SendRequestAsync<GetPermalinkResponse>(sendMessageRequest, null);
-				return response?.Permalink;
-			}
-		}
-
-		private async Task<TResponse> SendRequestAsync<TResponse>(string requestUrl, object request) where TResponse : SlackResponse
-		{
-			using (HttpRequestMessage sendMessageRequest = new HttpRequestMessage(HttpMethod.Post, requestUrl))
-			{
-				string requestJson = JsonSerializer.Serialize(request, new JsonSerializerOptions { DefaultIgnoreCondition = JsonIgnoreCondition.WhenWritingNull });
-				using (StringContent messageContent = new StringContent(requestJson, Encoding.UTF8, "application/json"))
-				{
-					sendMessageRequest.Content = messageContent;
-					return await SendRequestAsync<TResponse>(sendMessageRequest, requestJson);
-				}
-			}
-		}
-
-		static readonly string[] s_ignoreErrorCodes =
-		{
-			"no_reaction",
-			"already_reacted"
-		};
-
-		private async Task<TResponse> SendRequestAsync<TResponse>(HttpRequestMessage request, string? requestJson) where TResponse : SlackResponse
-		{
-			using (HttpClient client = new HttpClient())
-			{
-				request.Headers.Add("Authorization", $"Bearer {_settings.SlackToken ?? ""}");
-
-				HttpResponseMessage response = await client.SendAsync(request);
-				byte[] responseBytes = await response.Content.ReadAsByteArrayAsync();
-
-				TResponse responseObject = JsonSerializer.Deserialize<TResponse>(responseBytes)!;
-				if (!responseObject.Ok && !s_ignoreErrorCodes.Contains(responseObject.Error, StringComparer.Ordinal))
-				{
-					_logger.LogError("Failed to send Slack message ({Error}). Request: {Request}. Response: {Response}", responseObject.Error, requestJson, Encoding.UTF8.GetString(responseBytes));
-				}
-
-				return responseObject;
-			}
-		}
-
 		/// <inheritdoc/>
 		protected override async Task ExecuteAsync(CancellationToken stoppingToken)
 		{
@@ -2268,11 +1997,15 @@ namespace Horde.Build.Notifications.Sinks
 			await socket.ConnectAsync(socketUrl, stoppingToken);
 
 			byte[] buffer = new byte[2048];
-			while (!stoppingToken.IsCancellationRequested)
+
+			bool disconnect = false;
+			while(!disconnect)
 			{
+				stoppingToken.ThrowIfCancellationRequested();
+
 				// Read the next message
 				int length = 0;
-				for (; ; )
+				while(!disconnect)
 				{
 					if (length == buffer.Length)
 					{
@@ -2296,52 +2029,192 @@ namespace Horde.Build.Notifications.Sinks
 				_logger.LogInformation("Slack event: {Message}", Encoding.UTF8.GetString(buffer, 0, length));
 				EventMessage eventMessage = JsonSerializer.Deserialize<EventMessage>(buffer.AsSpan(0, length))!;
 
-				// Acknowledge the message
-				if (eventMessage.EnvelopeId != null)
-				{
-					object response = new { eventMessage.EnvelopeId };
-					await socket.SendAsync(JsonSerializer.SerializeToUtf8Bytes(response), WebSocketMessageType.Text, true, stoppingToken);
-				}
+				// Prepare the response
+				EventResponse response = new EventResponse();
+				response.EnvelopeId = eventMessage.EnvelopeId;
 
 				// Handle the message type
 				if (eventMessage.Type != null)
 				{
 					string type = eventMessage.Type;
-					if (type.Equals("disconnect", StringComparison.Ordinal))
+					if (type.Equals("hello", StringComparison.Ordinal))
 					{
-						break;
+						// nop
+					}
+					else if (type.Equals("disconnect", StringComparison.Ordinal))
+					{
+						disconnect = true;
 					}
-					else if (type.Equals("interactive", StringComparison.Ordinal))
+					else if (type.Equals("interactive", StringComparison.Ordinal) && eventMessage.Payload != null)
 					{
-						await HandleInteractionMessage(eventMessage);
+						try
+						{
+							response.Payload = await HandleInteractionMessage(eventMessage.Payload, stoppingToken);
+						}
+						catch (Exception ex)
+						{
+							_logger.LogError(ex, "Exception handling Slack interaction message: {Message}", ex.Message);
+						}
 					}
 					else
 					{
 						_logger.LogDebug("Unhandled event type ({Type})", type);
 					}
 				}
+
+				// Send the response
+				await socket.SendAsync(JsonSerializer.SerializeToUtf8Bytes(response), WebSocketMessageType.Text, true, stoppingToken);
 			}
 		}
-
 		/// <summary>
 		/// Handle a button being clicked
 		/// </summary>
-		/// <param name="message">The event message</param>
-		/// <returns></returns>
-		private async Task HandleInteractionMessage(EventMessage message)
+		private async Task<object?> HandleInteractionMessage(EventPayload payload, CancellationToken cancellationToken)
 		{
-			if (message.Payload != null && message.Payload.User != null && message.Payload.User.UserName != null && String.Equals(message.Payload.Type, "block_actions", StringComparison.Ordinal))
+			if (payload.User != null && payload.User.Id != null)
 			{
-				foreach (ActionInfo action in message.Payload.Actions)
+				if (String.Equals(payload.Type, "block_actions", StringComparison.Ordinal))
 				{
-					if (action.Value != null)
+					foreach (ActionInfo action in payload.Actions)
 					{
-						if (action.Value.StartsWith("issue_", StringComparison.Ordinal))
+						if (action.Value != null)
 						{
-							await HandleIssueResponseAsync(message.Payload, action);
+							Match? match;
+							if (TryMatch(action.Value, @"^issue_(\d+)_([a-zA-Z]+)_([a-fA-F0-9]{24})$", out match))
+							{
+								int issueId = Int32.Parse(match.Groups[1].Value, NumberStyles.Integer, CultureInfo.InvariantCulture);
+								string verb = match.Groups[2].Value;
+								UserId userId = match.Groups[3].Value.ToObjectId<IUser>();
+								await HandleIssueDmResponseAsync(issueId, verb, userId, payload.User.Id);
+							}
+							else if (TryMatch(action.Value, @"^issue_(\d+)_([a-zA-Z]+)$", out match) && payload.TriggerId != null)
+							{
+								int issueId = Int32.Parse(match.Groups[1].Value, NumberStyles.Integer, CultureInfo.InvariantCulture);
+								string verb = match.Groups[2].Value;
+								await HandleIssueChannelResponseAsync(issueId, verb, payload.User.Id, payload.TriggerId);
+							}
+							else
+							{
+								_logger.LogWarning("Unrecognized action value: {Value}", action.Value);
+							}
 						}
 					}
 				}
+				else if (String.Equals(payload.Type, "view_submission", StringComparison.Ordinal))
+				{
+					if (payload.View != null && payload.View.State != null && payload.View.CallbackId != null)
+					{
+						Match? match;
+						if(TryMatch(payload.View.CallbackId, @"^issue_(\d+)_markfixed_([a-fA-F0-9]{24})$", out match))
+						{
+							int issueId = Int32.Parse(match.Groups[1].Value, NumberStyles.Integer, CultureInfo.InvariantCulture);
+							UserId userId = match.Groups[2].Value.ToObjectId<IUser>();
+
+							string? fixChangeStr;
+							if (payload.View.State.TryGetValue("fix_cl", "fix_cl_action", out fixChangeStr))
+							{
+								int fixChange;
+								if (!Int32.TryParse(fixChangeStr, out fixChange) || fixChange < 0)
+								{
+									Dictionary<string, string> errors = new Dictionary<string, string>();
+									errors.Add("fix_cl", $"'{fixChangeStr}' is not a valid fix changelist.");
+									return new { response_action = "errors", errors = errors };
+								}
+
+								await _issueService.UpdateIssueAsync(issueId, fixChange: fixChange, resolvedById: userId);
+								_logger.LogInformation("Marked issue {IssueId} fixed by user {UserId} in {Change}", issueId, userId, fixChange);
+							}
+						}
+					}
+				}
+			}
+			return null;
+		}
+
+		static bool TryMatch(string input, string pattern, [NotNullWhen(true)] out Match? match)
+		{
+			Match result = Regex.Match(input, pattern);
+			if(result.Success)
+			{
+				match = result;
+				return true;
+			}
+			else
+			{
+				match = null;
+				return false;
+			}
+		}
+
+		async Task HandleIssueDmResponseAsync(int issueId, string verb, UserId userId, string userName)
+		{
+			_logger.LogInformation("Issue {IssueId}: {Action} from {SlackUser} ({UserId})", issueId, verb, userName, userId);
+
+			if (String.Equals(verb, "ack", StringComparison.Ordinal))
+			{
+				await _issueService.UpdateIssueAsync(issueId, acknowledged: true);
+			}
+			else if (String.Equals(verb, "accept", StringComparison.Ordinal))
+			{
+				await _issueService.UpdateIssueAsync(issueId, ownerId: userId, acknowledged: true);
+			}
+			else if (String.Equals(verb, "decline", StringComparison.Ordinal))
+			{
+				await _issueService.UpdateIssueAsync(issueId, declinedById: userId);
+			}
+
+			IIssue? newIssue = await _issueService.Collection.GetIssueAsync(issueId);
+			if (newIssue != null)
+			{
+				IUser? user = await _userCollection.GetUserAsync(userId);
+				if (user != null)
+				{
+					string? recipient = await GetSlackUserId(user);
+					if (recipient != null)
+					{
+						IIssueDetails details = await _issueService.GetIssueDetailsAsync(newIssue);
+						await SendIssueMessageAsync(recipient, newIssue, details, userId, defaultAllowMentions);
+					}
+				}
+			}
+		}
+
+		async Task HandleIssueChannelResponseAsync(int issueId, string verb, string slackUserId, string triggerId)
+		{
+			_logger.LogInformation("Issue {IssueId}: {Action} from {SlackUser}", issueId, verb, slackUserId);
+
+			SlackUser? slackUser = await _slackClient.GetUserAsync(slackUserId);
+			if (slackUser == null || slackUser.Profile == null || slackUser.Profile.Email == null)
+			{
+				_logger.LogWarning("Unable to find Slack user profile for {UserId}", slackUserId);
+				return;
+			}
+
+			IUser? user = await _userCollection.FindUserByEmailAsync(slackUser.Profile.Email);
+			if (user == null)
+			{
+				_logger.LogWarning("Unable to find Horde user profile for {Email}", slackUser.Profile.Email);
+				return;
+			}
+
+			if (String.Equals(verb, "ack", StringComparison.Ordinal))
+			{
+				await _issueService.UpdateIssueAsync(issueId, acknowledged: true, ownerId: user.Id);
+			}
+			else if (String.Equals(verb, "decline", StringComparison.Ordinal))
+			{
+				await _issueService.UpdateIssueAsync(issueId, declinedById: user.Id);
+			}
+			else if (String.Equals(verb, "markfixed", StringComparison.Ordinal))
+			{
+				SlackView view = new SlackView($"Issue {issueId}");
+				view.CallbackId = $"issue_{issueId}_markfixed_{user.Id}";
+				view.AddInput("Fix Changelist", new PlainTextInputElement("fix_cl_action", placeholder: "Number")).BlockId = "fix_cl";
+				view.Close = "Cancel";
+				view.Submit = "Mark Fixed";
+
+				await _slackClient.OpenViewAsync(triggerId, view);
+				return;
 			}
 		}
 	}
diff --git a/Engine/Source/Programs/Horde/Horde.Build/Users/IUserCollection.cs b/Engine/Source/Programs/Horde/Horde.Build/Users/IUserCollection.cs
--- a/Engine/Source/Programs/Horde/Horde.Build/Users/IUserCollection.cs
+++ b/Engine/Source/Programs/Horde/Horde.Build/Users/IUserCollection.cs
@@ -48,6 +48,13 @@ namespace Horde.Build.Users
 		/// <returns>The user information</returns>
 		Task<IUser?> FindUserByLoginAsync(string login);
 
+		/// <summary>
+		/// Gets a user by email address
+		/// </summary>
+		/// <param name="email">Email for the user</param>
+		/// <returns>The user information</returns>
+		Task<IUser?> FindUserByEmailAsync(string email);
+
 		/// <summary>
 		/// Find or add a user with the given claims. Claims will be updated if the user exists.
 		/// </summary>
diff --git a/Engine/Source/Programs/Horde/Horde.Build/Users/UserCollection.cs b/Engine/Source/Programs/Horde/Horde.Build/Users/UserCollection.cs
--- a/Engine/Source/Programs/Horde/Horde.Build/Users/UserCollection.cs
+++ b/Engine/Source/Programs/Horde/Horde.Build/Users/UserCollection.cs
@@ -124,6 +124,12 @@ namespace Horde.Build.Users
 			return await _users.Find(x => x.PrimaryClaim == primaryClaim).FirstOrDefaultAsync();
 		}
 
+		/// <inheritdoc/>
+		public Task<IUser?> FindUserByEmailAsync(string email)
+		{
+			return Task.FromResult<IUser?>(null);
+		}
+
 		/// <inheritdoc/>
 		public async Task<IUser> FindOrAddUserByLoginAsync(string login, string? name, string? email)
 		{
diff --git a/Engine/Source/Programs/Horde/Horde.Build/Users/UserCollectionV2.cs b/Engine/Source/Programs/Horde/Horde.Build/Users/UserCollectionV2.cs
--- a/Engine/Source/Programs/Horde/Horde.Build/Users/UserCollectionV2.cs
+++ b/Engine/Source/Programs/Horde/Horde.Build/Users/UserCollectionV2.cs
@@ -244,6 +244,13 @@ namespace Horde.Build.Users
 			return await _users.Find(x => x.LoginUpper == loginUpper).FirstOrDefaultAsync();
 		}
 
+		/// <inheritdoc/>
+		public async Task<IUser?> FindUserByEmailAsync(string email)
+		{
+			string emailUpper = email.ToUpperInvariant();
+			return await _users.Find(x => x.EmailUpper == emailUpper).FirstOrDefaultAsync();
+		}
+
 		/// <inheritdoc/>
 		public async Task<IUser> FindOrAddUserByLoginAsync(string login, string? name, string? email)
 		{
diff --git a/Engine/Source/Programs/Horde/Horde.Build/Utilities/BlockKit/ActionsBlock.cs b/Engine/Source/Programs/Horde/Horde.Build/Utilities/BlockKit/ActionsBlock.cs
deleted file mode 100644
diff --git a/Engine/Source/Programs/Horde/Horde.Build/Utilities/BlockKit/BlockBase.cs b/Engine/Source/Programs/Horde/Horde.Build/Utilities/BlockKit/BlockBase.cs
deleted file mode 100644
diff --git a/Engine/Source/Programs/Horde/Horde.Build/Utilities/BlockKit/BlockKitAttachment.cs b/Engine/Source/Programs/Horde/Horde.Build/Utilities/BlockKit/BlockKitAttachment.cs
deleted file mode 100644
diff --git a/Engine/Source/Programs/Horde/Horde.Build/Utilities/BlockKit/BlockKitMessage.cs b/Engine/Source/Programs/Horde/Horde.Build/Utilities/BlockKit/BlockKitMessage.cs
deleted file mode 100644
diff --git a/Engine/Source/Programs/Horde/Horde.Build/Utilities/BlockKit/DividerBlock.cs b/Engine/Source/Programs/Horde/Horde.Build/Utilities/BlockKit/DividerBlock.cs
deleted file mode 100644
diff --git a/Engine/Source/Programs/Horde/Horde.Build/Utilities/BlockKit/HeaderBlock.cs b/Engine/Source/Programs/Horde/Horde.Build/Utilities/BlockKit/HeaderBlock.cs
deleted file mode 100644
diff --git a/Engine/Source/Programs/Horde/Horde.Build/Utilities/BlockKit/SectionBlock.cs b/Engine/Source/Programs/Horde/Horde.Build/Utilities/BlockKit/SectionBlock.cs
deleted file mode 100644
diff --git a/Engine/Source/Programs/Horde/Horde.Build/Utilities/BlockKit/SlackResponses.cs b/Engine/Source/Programs/Horde/Horde.Build/Utilities/BlockKit/SlackResponses.cs
deleted file mode 100644
diff --git a/Engine/Source/Programs/Horde/Horde.Build/Utilities/BlockKit/TextObject.cs b/Engine/Source/Programs/Horde/Horde.Build/Utilities/BlockKit/TextObject.cs
deleted file mode 100644
diff --git a/Engine/Source/Programs/Shared/EpicGames.Slack/Block.cs b/Engine/Source/Programs/Shared/EpicGames.Slack/Block.cs
new file mode 100644
--- /dev/null
+++ b/Engine/Source/Programs/Shared/EpicGames.Slack/Block.cs
@@ -0,0 +1,67 @@
+// Copyright Epic Games, Inc. All Rights Reserved.
+
+using System;
+using System.Collections.Generic;
+using System.Text.Json;
+using System.Text.Json.Serialization;
+using EpicGames.Slack.Blocks;
+using EpicGames.Slack.Elements;
+
+namespace EpicGames.Slack
+{
+	/// <summary>
+	/// Abstract base class for all BlockKit blocks to derive from.
+	/// </summary>
+	[JsonConverter(typeof(BlockConverter))]
+	public abstract class Block
+	{
+		/// <summary>
+		/// The block type
+		/// </summary>
+		[JsonPropertyName("type"), JsonPropertyOrder(0)]
+		public string Type { get; }
+
+		/// <summary>
+		/// A string acting as a unique identifier for a block. If not specified, one will be generated. Maximum length for this field is 255 characters. 
+		/// block_id should be unique for each message and each iteration of a message. If a message is updated, use a new block_id.
+		/// </summary>
+		[JsonPropertyName("block_id"), JsonPropertyOrder(30)]
+		public string? BlockId { get; set; }
+
+		/// <summary>
+		/// Constructor
+		/// </summary>
+		/// <param name="type">Type of the block</param>
+		protected Block(string type)
+		{
+			Type = type;
+		}
+	}
+
+	/// <summary>
+	/// Interface for a container of <see cref="Block"/> items
+	/// </summary>
+	public interface ISlackBlockContainer
+	{
+		/// <summary>
+		/// List of blocks
+		/// </summary>
+		List<Block> Blocks { get; }
+	}
+
+	/// <summary>
+	/// Polymorphic serializer for Block objects.
+	/// </summary>
+	class BlockConverter : JsonConverter<Block>
+	{
+		public override Block? Read(ref Utf8JsonReader reader, Type typeToConvert, JsonSerializerOptions options)
+		{
+			throw new NotImplementedException();
+		}
+
+		public override void Write(Utf8JsonWriter writer, Block value, JsonSerializerOptions options)
+		{
+			JsonSerializer.Serialize(writer, (object)value, value.GetType(), options);
+		}
+	}
+}
diff --git a/Engine/Source/Programs/Shared/EpicGames.Slack/Blocks/ActionsBlock.cs b/Engine/Source/Programs/Shared/EpicGames.Slack/Blocks/ActionsBlock.cs
new file mode 100644
--- /dev/null
+++ b/Engine/Source/Programs/Shared/EpicGames.Slack/Blocks/ActionsBlock.cs
@@ -0,0 +1,56 @@
+// Copyright Epic Games, Inc. All Rights Reserved.
+
+using System;
+using System.Collections.Generic;
+using System.Text.Json.Serialization;
+
+namespace EpicGames.Slack.Blocks
+{
+	/// <summary>
+	/// A block that is used to hold interactive <see cref="Element"/> objects.
+	/// </summary>
+	public class ActionsBlock : Block, ISlackElementContainer
+	{
+		/// <summary>
+		/// A collection of interactive elements.
+		/// </summary>
+		[JsonPropertyName("elements"), JsonPropertyOrder(1)]
+		public List<Element> Elements { get; } = new List<Element>();
+
+		/// <summary>
+		/// Constructor
+		/// </summary>
+		public ActionsBlock() : base("actions")
+		{
+		}
+	}
+
+	/// <summary>
+	/// Extension methods for <see cref="Block"/>
+	/// </summary>
+	public static class ActionsBlockExtensions
+	{
+		/// <summary>
+		/// Add an <see cref="ActionsBlock"/> to the list of blocks
+		/// </summary>
+		/// <param name="container">Block container</param>
+		public static ActionsBlock AddActions(this ISlackBlockContainer container)
+		{
+			ActionsBlock block = new ActionsBlock();
+			container.Blocks.Add(block);
+			return block;
+		}
+
+		/// <summary>
+		/// Add an <see cref="ActionsBlock"/> to the list of blocks
+		/// </summary>
+		/// <param name="container">Block container</param>
+		/// <param name="configure">Configuration function for the block</param>
+		public static void AddActions(this ISlackBlockContainer container, Action<ActionsBlock> configure)
+		{
+			ActionsBlock block = new ActionsBlock();
+			configure(block);
+			container.Blocks.Add(block);
+		}
+	}
+}
diff --git a/Engine/Source/Programs/Shared/EpicGames.Slack/Blocks/ContextBlock.cs b/Engine/Source/Programs/Shared/EpicGames.Slack/Blocks/ContextBlock.cs
new file mode 100644
--- /dev/null
+++ b/Engine/Source/Programs/Shared/EpicGames.Slack/Blocks/ContextBlock.cs
@@ -0,0 +1,46 @@
+// Copyright Epic Games, Inc. All Rights Reserved.
+
+using System;
+using System.Collections.Generic;
+using System.Text.Json;
+using System.Text.Json.Serialization;
+
+namespace EpicGames.Slack.Blocks
+{
+	/// <summary>
+	/// Displays message context, which can include both images and text.
+	/// </summary>
+	public class ContextBlock : Block
+	{
+		/// <summary>
+		/// An array of image elements and text objects. Maximum number of items is 10.
+		/// </summary>
+		[JsonPropertyName("elements")]
+		public List<Element> Elements { get; } = new List<Element>();
+
+		/// <summary>
+		/// Constructor
+		/// </summary>
+		public ContextBlock() : base("context")
+		{
+		}
+	}
+
+	/// <summary>
+	/// Extension methods for <see cref="Block"/>
+	/// </summary>
+	public static class ContextBlockExtensions
+	{
+		/// <summary>
+		/// Add an <see cref="ContextBlock"/> to the list of blocks
+		/// </summary>
+		/// <param name="container">Block container</param>
+		/// <param name="configure">Configuration callback for the block</param>
+		public static void AddContext(this ISlackBlockContainer container, Action<ContextBlock> configure)
+		{
+			ContextBlock block = new ContextBlock();
+			configure(block);
+			container.Blocks.Add(block);
+		}
+	}
+}
diff --git a/Engine/Source/Programs/Shared/EpicGames.Slack/Blocks/DividerBlock.cs b/Engine/Source/Programs/Shared/EpicGames.Slack/Blocks/DividerBlock.cs
new file mode 100644
--- /dev/null
+++ b/Engine/Source/Programs/Shared/EpicGames.Slack/Blocks/DividerBlock.cs
@@ -0,0 +1,35 @@
+// Copyright Epic Games, Inc. All Rights Reserved.
+
+using System.Collections.Generic;
+
+namespace EpicGames.Slack.Blocks
+{
+	/// <summary>
+	/// A content divider, like an <hr>, to split up different blocks inside of a message. The divider block is nice and neat, requiring only a type.
+	/// </summary>
+	public class DividerBlock : Block
+	{
+		/// <summary>
+		/// Constructor
+		/// </summary>
+		public DividerBlock() : base("divider")
+		{
+		}
+	}
+
+	/// <summary>
+	/// Extension methods for <see cref="Block"/>
+	/// </summary>
+	public static class DividerBlockExtensions
+	{
+		/// <summary>
+		/// Add an <see cref="DividerBlock"/> to the list of blocks
+		/// </summary>
+		/// <param name="container">Block container</param>
+		public static void AddDivider(this ISlackBlockContainer container)
+		{
+			DividerBlock block = new DividerBlock();
+			container.Blocks.Add(block);
+		}
+	}
+}
diff --git a/Engine/Source/Programs/Shared/EpicGames.Slack/Blocks/HeaderBlock.cs b/Engine/Source/Programs/Shared/EpicGames.Slack/Blocks/HeaderBlock.cs
new file mode 100644
--- /dev/null
+++ b/Engine/Source/Programs/Shared/EpicGames.Slack/Blocks/HeaderBlock.cs
@@ -0,0 +1,60 @@
+// Copyright Epic Games, Inc. All Rights Reserved.
+
+using System.Collections.Generic;
+using System.Text.Json;
+using EpicGames.Slack.Elements;
+
+namespace EpicGames.Slack.Blocks
+{
+	/// <summary>
+	/// A header is a plain-text block that displays in a larger, bold font. Use it to delineate between different groups of content in your app's surfaces.
+	/// </summary>
+	public class HeaderBlock : Block
+	{
+		/// <summary>
+		/// The text for the block, in the form of a plain_text text object. Maximum length for the text in this field is 150 characters.
+		/// </summary>
+		public PlainTextObject Text { get; }
+
+		/// <summary>
+		/// Constructor
+		/// </summary>
+		/// <param name="text">Any text to initiale the Section with</param>
+		public HeaderBlock(PlainTextObject text) : base("header")
+		{
+			Text = text;
+		}
+
+		/// <summary>
+		/// Constructor
+		/// </summary>
+		/// <param name="text">Text to display in the header</param>
+		public HeaderBlock(string text) : this(new PlainTextObject(text))
+		{
+		}
+	}
+
+	/// <summary>
+	/// Extension methods for <see cref="Block"/>
+	/// </summary>
+	public static class HeaderBlockExtensions
+	{
+		/// <summary>
+		/// Add an <see cref="HeaderBlock"/> to the list of blocks
+		/// </summary>
+		/// <param name="container">List of blocks</param>
+		public static void AddHeader(this ISlackBlockContainer container, PlainTextObject text)
+		{
+			HeaderBlock block = new HeaderBlock(text);
+			container.Blocks.Add(block);
+		}
+
+		/// <summary>
+		/// Add an <see cref="HeaderBlock"/> to the list of blocks
+		/// </summary>
+		/// <param name="container">List of blocks</param>
+		/// <param name="text">Text for the header</param>
+		/// <param name="emoji">Whether to enable emoji</param>
+		public static void AddHeader(this ISlackBlockContainer container, string text, bool? emoji) => AddHeader(container, new PlainTextObject(text, emoji));
+	}
+}
diff --git a/Engine/Source/Programs/Shared/EpicGames.Slack/Blocks/ImageBlock.cs b/Engine/Source/Programs/Shared/EpicGames.Slack/Blocks/ImageBlock.cs
new file mode 100644
--- /dev/null
+++ b/Engine/Source/Programs/Shared/EpicGames.Slack/Blocks/ImageBlock.cs
@@ -0,0 +1,60 @@
+// Copyright Epic Games, Inc. All Rights Reserved.
+
+using System;
+using System.Collections.Generic;
+using System.Text.Json;
+using System.Text.Json.Serialization;
+using EpicGames.Slack.Elements;
+
+namespace EpicGames.Slack.Blocks
+{
+	/// <summary>
+	/// A simple image block, designed to make those cat photos really pop.
+	/// </summary>
+	public class ImageBlock : Block
+	{
+		/// <summary>
+		/// The URL of the image to be displayed.Maximum length for this field is 3000 characters.
+		/// </summary>
+		[JsonPropertyName("image_url")]
+		public Uri ImageUrl { get; set; }
+
+		/// <summary>
+		/// A plain-text summary of the image.This should not contain any markup.Maximum length for this field is 2000 characters.
+		/// </summary>
+		[JsonPropertyName("alt_text")]
+		public string AltText { get; set; }
+
+		/// <summary>
+		/// An optional title for the image in the form of a text object that can only be of type: plain_text. Maximum length for the text in this field is 2000 characters.
+		/// </summary>
+		[JsonPropertyName("title")]
+		public PlainTextObject? Title { get; set; }
+
+		/// <summary>
+		/// Constructor
+		/// </summary>
+		public ImageBlock(Uri imageUrl, string altText, PlainTextObject? title = null) : base("image")
+		{
+			ImageUrl = imageUrl;
+			AltText = altText;
+			Title = title;
+		}
+	}
+
+	/// <summary>
+	/// Extension methods for <see cref="ImageBlock"/>
+	/// </summary>
+	public static class ImageBlockExtensions
+	{
+		/// <summary>
+		/// Add an <see cref="ImageBlock"/> to the list of blocks
+		/// </summary>
+		/// <param name="container">Block container</param>
+		public static void AddImage(this ISlackBlockContainer container, Uri imageUrl, string altText, PlainTextObject? title = null)
+		{
+			ImageBlock block = new ImageBlock(imageUrl, altText, title);
+			container.Blocks.Add(block);
+		}
+	}
+}
diff --git a/Engine/Source/Programs/Shared/EpicGames.Slack/Blocks/InputBlock.cs b/Engine/Source/Programs/Shared/EpicGames.Slack/Blocks/InputBlock.cs
new file mode 100644
--- /dev/null
+++ b/Engine/Source/Programs/Shared/EpicGames.Slack/Blocks/InputBlock.cs
@@ -0,0 +1,72 @@
+// Copyright Epic Games, Inc. All Rights Reserved.
+
+using System;
+using System.Collections.Generic;
+using System.Text.Json;
+using System.Text.Json.Serialization;
+using EpicGames.Slack.Elements;
+
+namespace EpicGames.Slack.Blocks
+{
+	/// <summary>
+	/// A block that collects information from users - it can hold a plain-text input element, a checkbox element, a radio button element, a select menu element, a multi-select menu element, or a datepicker.
+	/// </summary>
+	public class InputBlock : Block
+	{
+		/// <summary>
+		/// A label that appears above an input element in the form of a text object that must have type of plain_text. Maximum length for the text in this field is 2000 characters.
+		/// </summary>
+		[JsonPropertyName("label"), JsonPropertyOrder(1)]
+		public PlainTextObject Label { get; set; }
+
+		/// <summary>
+		/// A plain-text input element, a checkbox element, a radio button element, a select menu element, a multi-select menu element, or a datepicker.
+		/// </summary>
+		[JsonPropertyName("element"), JsonPropertyOrder(2)]
+		public Element Element { get; set; }
+
+		/// <summary>
+		/// A boolean that indicates whether or not the use of elements in this block should dispatch a block_actions payload. Defaults to false.
+		/// </summary>
+		[JsonPropertyName("dispatch_action"), JsonPropertyOrder(3)]
+		public bool? DispatchAction { get; set; }
+
+		/// <summary>
+		/// An optional hint that appears below an input element in a lighter grey. It must be a text object with a type of plain_text. Maximum length for the text in this field is 2000 characters.
+		/// </summary>
+		[JsonPropertyName("hint"), JsonPropertyOrder(4)]
+		public PlainTextObject? Hint { get; set; }
+
+		/// <summary>
+		/// A boolean that indicates whether the input element may be empty when a user submits the modal. Defaults to false.
+		/// </summary>
+		[JsonPropertyName("optional"), JsonPropertyOrder(5)]
+		public bool? Optional { get; set; }
+
+		/// <summary>
+		/// Constructor
+		/// </summary>
+		public InputBlock(PlainTextObject label, Element element) : base("input")
+		{
+			Label = label;
+			Element = element;
+		}
+	}
+
+	/// <summary>
+	/// Extension methods for <see cref="InputBlock"/>
+	/// </summary>
+	public static class InputBlockExtensions
+	{
+		/// <summary>
+		/// Add an <see cref="ImageBlock"/> to the list of blocks
+		/// </summary>
+		/// <param name="container">Block container</param>
+		public static InputBlock AddInput(this ISlackBlockContainer container, PlainTextObject label, Element element)
+		{
+			InputBlock block = new InputBlock(label, element);
+			container.Blocks.Add(block);
+			return block;
+		}
+	}
+}
diff --git a/Engine/Source/Programs/Shared/EpicGames.Slack/Blocks/SectionBlock.cs b/Engine/Source/Programs/Shared/EpicGames.Slack/Blocks/SectionBlock.cs
new file mode 100644
--- /dev/null
+++ b/Engine/Source/Programs/Shared/EpicGames.Slack/Blocks/SectionBlock.cs
@@ -0,0 +1,73 @@
+// Copyright Epic Games, Inc. All Rights Reserved.
+
+using System;
+using System.Collections.Generic;
+using System.Text.Json;
+using System.Text.Json.Serialization;
+using EpicGames.Slack.Elements;
+
+namespace EpicGames.Slack.Blocks
+{
+	/// <summary>
+	/// Represents a BlockKit Section block
+	/// </summary>
+	public class SectionBlock : Block
+	{
+		/// <summary>
+		/// The text for the block, in the form of a text object. Maximum length for the text in this field is 3000 characters.
+		/// </summary>
+		[JsonPropertyName("text")]
+		public TextObject? Text { get; set; }
+
+		/// <summary>
+		/// Required if no text is provided. An array of text objects. Any text objects included with fields will be rendered in a compact format that allows for 2 columns of side-by-side text. 
+		/// Maximum number of items is 10. Maximum length for the text in each item is 2000 characters.
+		/// </summary>
+		[JsonPropertyName("fields")]
+		public List<TextObject>? Fields { get; set; }
+
+		/// <summary>
+		/// One of the available element objects.
+		/// </summary>
+		[JsonPropertyName("accessory")]
+		public Element? Accessory { get; set; }
+
+		/// <summary>
+		/// Constructor
+		/// </summary>
+		/// <param name="text">Any text to initiale the Section with</param>
+		/// <param name="accessory">Optional accessory element</param>
+		public SectionBlock(TextObject text, Element? accessory = null) : base("section")
+		{
+			Text = text;
+			Accessory = accessory;
+		}
+
+		/// <summary>
+		/// Constructor
+		/// </summary>
+		/// <param name="fields">Any text to initiale the Section with</param>
+		/// <param name="accessory">Optional accessory element</param>
+		public SectionBlock(List<TextObject> fields, Element? accessory = null) : base("section")
+		{
+			Fields = fields;
+			Accessory = accessory;
+		}
+	}
+
+	/// <summary>
+	/// Extension methods for <see cref="SectionBlock"/>
+	/// </summary>
+	public static class SectionBlockExtensions
+	{
+		/// <summary>
+		/// Add an <see cref="SectionBlock"/> to the list of blocks
+		/// </summary>
+		/// <param name="container">Container for the blocks</param>
+		public static void AddSection(this ISlackBlockContainer container, TextObject text)
+		{
+			SectionBlock block = new SectionBlock(text);
+			container.Blocks.Add(block);
+		}
+	}
+}
diff --git a/Engine/Source/Programs/Shared/EpicGames.Slack/Element.cs b/Engine/Source/Programs/Shared/EpicGames.Slack/Element.cs
new file mode 100644
--- /dev/null
+++ b/Engine/Source/Programs/Shared/EpicGames.Slack/Element.cs
@@ -0,0 +1,58 @@
+// Copyright Epic Games, Inc. All Rights Reserved.
+
+using System;
+using System.Collections.Generic;
+using System.Text.Json;
+using System.Text.Json.Serialization;
+
+namespace EpicGames.Slack
+{
+	/// <summary>
+	/// Abstract base class for all BlockKit elements to derive from.
+	/// </summary>
+	[JsonConverter(typeof(ElementConverter))]
+	public abstract class Element
+	{
+		/// <summary>
+		/// Type of the element
+		/// </summary>
+		[JsonPropertyName("type"), JsonPropertyOrder(0)]
+		public string Type { get; set; }
+
+		/// <summary>
+		/// Constructor
+		/// </summary>
+		/// <param name="type"></param>
+		protected Element(string type)
+		{
+			Type = type;
+		}
+	}
+
+	/// <summary>
+	/// Something that contains a list of elements. Used for extension methods.
+	/// </summary>
+	public interface ISlackElementContainer
+	{
+		/// <summary>
+		/// List of elements in this container
+		/// </summary>
+		public List<Element> Elements { get; }
+	}
+
+	/// <summary>
+	/// Polymorphic serializer for Block objects.
+	/// </summary>
+	class ElementConverter : JsonConverter<Element>
+	{
+		public override Element? Read(ref Utf8JsonReader reader, Type typeToConvert, JsonSerializerOptions options)
+		{
+			throw new NotImplementedException();
+		}
+
+		public override void Write(Utf8JsonWriter writer, Element value, JsonSerializerOptions options)
+		{
+			JsonSerializer.Serialize(writer, (object)value, value.GetType(), options);
+		}
+	}
+}
diff --git a/Engine/Source/Programs/Shared/EpicGames.Slack/Elements/ButtonElement.cs b/Engine/Source/Programs/Shared/EpicGames.Slack/Elements/ButtonElement.cs
new file mode 100644
--- /dev/null
+++ b/Engine/Source/Programs/Shared/EpicGames.Slack/Elements/ButtonElement.cs
@@ -0,0 +1,126 @@
+// Copyright Epic Games, Inc. All Rights Reserved.
+
+using System;
+using System.Collections.Generic;
+using System.Globalization;
+using System.Runtime.Serialization;
+using System.Text.Json;
+using System.Text.Json.Serialization;
+
+namespace EpicGames.Slack.Elements
+{
+	/// <summary>
+	/// Determines the visual color scheme to use for the button.
+	/// </summary>
+	[JsonConverter(typeof(ButtonStyleJsonConverter))]
+	public enum ButtonStyle
+	{
+		/// <summary>
+		/// Used for affirmation/confirmation actions.
+		/// </summary>
+		Primary,
+
+		/// <summary>
+		/// Used when an action is destructive and cannot be undone.
+		/// </summary>
+		Danger,
+	}
+
+	/// <summary>
+	/// Json serializer for <see cref="ButtonStyle"/>
+	/// </summary>
+	class ButtonStyleJsonConverter : JsonConverter<ButtonStyle>
+	{
+		public override ButtonStyle Read(ref Utf8JsonReader reader, Type typeToConvert, JsonSerializerOptions options)
+		{
+			throw new NotImplementedException();
+		}
+
+		public override void Write(Utf8JsonWriter writer, ButtonStyle value, JsonSerializerOptions options)
+		{
+			switch (value)
+			{
+				case ButtonStyle.Primary:
+					writer.WriteStringValue("primary");
+					break;
+				case ButtonStyle.Danger:
+					writer.WriteStringValue("danger");
+					break;
+				default:
+					throw new InvalidDataContractException();
+			}
+		}
+	}
+
+	/// <summary>
+	/// An button element that can be added to an <see cref="ActionsBlock"/>.
+	/// </summary>
+	public class ButtonElement : Element
+	{
+		/// <summary>
+		/// A text object that defines the button's text. Can only be of type: plain_text. text may truncate with ~30 characters. Maximum length for the text in this field is 75 characters.
+		/// </summary>
+		[JsonPropertyName("text"), JsonPropertyOrder(1)]
+		public PlainTextObject Text { get; }
+
+		/// <summary>
+		/// An identifier for this action. You can use this when you receive an interaction payload to identify the source of the action. 
+		/// Should be unique among all other action_ids in the containing block. Maximum length for this field is 255 characters.
+		/// </summary>
+		[JsonPropertyName("action_id"), JsonPropertyOrder(2)]
+		public string? ActionId { get; set; }
+
+		/// <summary>
+		/// A URL to load in the user's browser when the button is clicked. Maximum length for this field is 3000 characters. If you're using url, you'll still receive an interaction payload and will need to send an acknowledgement response.
+		/// </summary>
+		[JsonPropertyName("url"), JsonPropertyOrder(3)]
+		public Uri? Url { get; set; }
+
+		/// <summary>
+		/// The value to send along with the interaction payload. Maximum length for this field is 2000 characters.
+		/// </summary>
+		[JsonPropertyName("value"), JsonPropertyOrder(4)]
+		public string? Value { get; set; }
+
+		/// <summary>
+		/// Decorates buttons with alternative visual color schemes. Use this option with restraint.
+		/// </summary>
+		[JsonPropertyName("style"), JsonPropertyOrder(5)]
+		public ButtonStyle? Style { get; set; }
+
+		/// <summary>
+		/// A label for longer descriptive text about a button element. This label will be read out by screen readers instead of the button text object. Maximum length for this field is 75 characters.
+		/// </summary>
+		[JsonPropertyName("accessibility_label"), JsonPropertyOrder(6)]
+		public string? AccessibilityLabel { get; set; }
+
+		/// <summary>
+		/// Construct a new Button action element.
+		/// </summary>
+		public ButtonElement(PlainTextObject text, Uri? url = null, string? actionId = null, string? value = null, ButtonStyle? style = null)
+			: base("button")
+		{
+			Text = text;
+			Url = url;
+			Value = value;
+			ActionId = actionId;
+			Style = style;
+		}
+	}
+
+	/// <summary>
+	/// Extension methods for buttons
+	/// </summary>
+	public static class ButtonExtensions
+	{
+		/// <summary>
+		/// Add a new button to the list of elements
+		/// </summary>
+		public static ISlackElementContainer AddButton(this ISlackElementContainer list, string text, Uri? url = null, string? actionId = null, string? value = null, ButtonStyle? style = null)
+		{
+			ButtonElement button = new ButtonElement(text, url, actionId, value, style);
+			list.Elements.Add(button);
+			return list;
+		}
+	}
+}
diff --git a/Engine/Source/Programs/Shared/EpicGames.Slack/Elements/CheckboxGroupElement.cs b/Engine/Source/Programs/Shared/EpicGames.Slack/Elements/CheckboxGroupElement.cs
new file mode 100644
--- /dev/null
+++ b/Engine/Source/Programs/Shared/EpicGames.Slack/Elements/CheckboxGroupElement.cs
@@ -0,0 +1,58 @@
+// Copyright Epic Games, Inc. All Rights Reserved.
+
+using System;
+using System.Collections.Generic;
+using System.Globalization;
+using System.Runtime.Serialization;
+using System.Text.Json;
+using System.Text.Json.Serialization;
+
+namespace EpicGames.Slack.Elements
+{
+	/// <summary>
+	/// A checkbox group that allows a user to choose multiple items from a list of possible options.
+	/// </summary>
+	public class CheckboxGroupElement : Element
+	{
+		/// <summary>
+		/// An identifier for this action. You can use this when you receive an interaction payload to identify the source of the action. 
+		/// Should be unique among all other action_ids in the containing block. Maximum length for this field is 255 characters.
+		/// </summary>
+		[JsonPropertyName("action_id")]
+		public string ActionId { get; set; }
+
+		/// <summary>
+		/// An array of option objects. A maximum of 10 options are allowed.
+		/// </summary>
+		[JsonPropertyName("options")]
+		public List<Option> Options { get; } = new List<Option>();
+
+		/// <summary>
+		/// An array of option objects that exactly matches one or more of the options within options. These options will be selected when the checkbox group initially loads.
+		/// </summary>
+		[JsonPropertyName("initial_options")]
+		public List<Option> InitialOptions { get; } = new List<Option>();
+
+		/// <summary>
+		/// A confirm object that defines an optional confirmation dialog that appears after clicking one of the checkboxes in this element.
+		/// </summary>
+		[JsonPropertyName("confirm")]
+		public SlackConfirm? Confirm { get; set; }
+
+		/// <summary>
+		/// Indicates whether the element will be set to auto focus within the view object. Only one element can be set to true. Defaults to false.
+		/// </summary>
+		[JsonPropertyName("focus_on_load")]
+		public bool FocusOnLoad { get; set; }
+
+		/// <summary>
+		/// Construct a new Button action element.
+		/// </summary>
+		public CheckboxGroupElement(string actionId, List<Option> options)
+			: base("checkboxes")
+		{
+			ActionId = actionId;
+			Options.AddRange(options);
+		}
+	}
+}
diff --git a/Engine/Source/Programs/Shared/EpicGames.Slack/Elements/DatePickerElement.cs b/Engine/Source/Programs/Shared/EpicGames.Slack/Elements/DatePickerElement.cs
new file mode 100644
--- /dev/null
+++ b/Engine/Source/Programs/Shared/EpicGames.Slack/Elements/DatePickerElement.cs
@@ -0,0 +1,57 @@
+// Copyright Epic Games, Inc. All Rights Reserved.
+
+using System;
+using System.Collections.Generic;
+using System.Globalization;
+using System.Runtime.Serialization;
+using System.Text.Json;
+using System.Text.Json.Serialization;
+
+namespace EpicGames.Slack.Elements
+{
+	/// <summary>
+	/// An element which lets users easily select a date from a calendar style UI.
+	/// </summary>
+	public class DatePickerElement : Element
+	{
+		/// <summary>
+		/// An identifier for this action. You can use this when you receive an interaction payload to identify the source of the action. 
+		/// Should be unique among all other action_ids in the containing block. Maximum length for this field is 255 characters.
+		/// </summary>
+		[JsonPropertyName("action_id")]
+		public string ActionId { get; set; }
+
+		/// <summary>
+		/// A plain_text only text object that defines the placeholder text shown on the datepicker. Maximum length for the text in this field is 150 characters.
+		/// </summary>
+		[JsonPropertyName("placeholder")]
+		public PlainTextObject? Placeholder { get; set; }
+
+		/// <summary>
+		/// The initial date that is selected when the element is loaded. This should be in the format YYYY-MM-DD.
+		/// </summary>
+		[JsonPropertyName("initial_date")]
+		public string? InitialDate { get; set; }
+
+		/// <summary>
+		/// A confirm object that defines an optional confirmation dialog that appears after clicking one of the checkboxes in this element.
+		/// </summary>
+		[JsonPropertyName("confirm")]
+		public SlackConfirm? Confirm { get; set; }
+
+		/// <summary>
+		/// Indicates whether the element will be set to auto focus within the view object. Only one element can be set to true. Defaults to false.
+		/// </summary>
+		[JsonPropertyName("focus_on_load")]
+		public bool FocusOnLoad { get; set; }
+
+		/// <summary>
+		/// Construct a new Button action element.
+		/// </summary>
+		public DatePickerElement(string actionId)
+			: base("datepicker")
+		{
+			ActionId = actionId;
+		}
+	}
+}
diff --git a/Engine/Source/Programs/Shared/EpicGames.Slack/Elements/OverflowMenuElement.cs b/Engine/Source/Programs/Shared/EpicGames.Slack/Elements/OverflowMenuElement.cs
new file mode 100644
--- /dev/null
+++ b/Engine/Source/Programs/Shared/EpicGames.Slack/Elements/OverflowMenuElement.cs
@@ -0,0 +1,51 @@
+// Copyright Epic Games, Inc. All Rights Reserved.
+
+using System;
+using System.Collections.Generic;
+using System.Globalization;
+using System.Runtime.Serialization;
+using System.Text.Json;
+using System.Text.Json.Serialization;
+
+namespace EpicGames.Slack.Elements
+{
+	/// <summary>
+	/// This is like a cross between a button and a select menu - when a user clicks on this overflow button, they will be presented with a list of options to choose from. 
+	/// Unlike the select menu, there is no typeahead field, and the button always appears with an ellipsis ("…") rather than customisable text.
+	///
+	/// As such, it is usually used if you want a more compact layout than a select menu, or to supply a list of less visually important actions after a row of buttons.
+	/// You can also specify simple URL links as overflow menu options, instead of actions.
+	/// </summary>
+	public class OverflowMenuElement : Element
+	{
+		/// <summary>
+		/// An identifier for this action. You can use this when you receive an interaction payload to identify the source of the action. 
+		/// Should be unique among all other action_ids in the containing block. Maximum length for this field is 255 characters.
+		/// </summary>
+		[JsonPropertyName("action_id")]
+		public string ActionId { get; set; }
+
+		/// <summary>
+		/// An array of up to five option objects to display in the menu.
+		/// </summary>
+		[JsonPropertyName("options")]
+		public List<Option> Options { get; } = new List<Option>();
+
+		/// <summary>
+		/// A confirm object that defines an optional confirmation dialog that appears after clicking one of the checkboxes in this element.
+		/// </summary>
+		[JsonPropertyName("confirm")]
+		public SlackConfirm? Confirm { get; set; }
+
+		/// <summary>
+		/// Construct a new Button action element.
+		/// </summary>
+		public OverflowMenuElement(string actionId, List<Option> options, SlackConfirm? confirm = null)
+			: base("overflow")
+		{
+			ActionId = actionId;
+			Options.AddRange(options);
+			Confirm = confirm;
+		}
+	}
+}
diff --git a/Engine/Source/Programs/Shared/EpicGames.Slack/Elements/PlainTextInputElement.cs b/Engine/Source/Programs/Shared/EpicGames.Slack/Elements/PlainTextInputElement.cs
new file mode 100644
--- /dev/null
+++ b/Engine/Source/Programs/Shared/EpicGames.Slack/Elements/PlainTextInputElement.cs
@@ -0,0 +1,71 @@
+// Copyright Epic Games, Inc. All Rights Reserved.
+
+using System;
+using System.Collections.Generic;
+using System.Globalization;
+using System.Text.Json;
+using System.Text.Json.Serialization;
+
+namespace EpicGames.Slack.Elements
+{
+	/// <summary>
+	/// A plain-text input, similar to the HTML <input> tag, creates a field where a user can enter freeform data. It can appear as a single-line field or a larger text area using the multiline flag.
+	/// </summary>
+	public class PlainTextInputElement : Element
+	{
+		/// <summary>
+		/// An identifier for this action. You can use this when you receive an interaction payload to identify the source of the action. 
+		/// Should be unique among all other action_ids in the containing block. Maximum length for this field is 255 characters.
+		/// </summary>
+		[JsonPropertyName("action_id")]
+		public string ActionId { get; set; }
+
+		/// <summary>
+		/// A plain_text only text object that defines the placeholder text shown in the plain-text input. Maximum length for the text in this field is 150 characters.
+		/// </summary>
+		[JsonPropertyName("placeholder")]
+		public PlainTextObject? Placeholder { get; set; }
+
+		/// <summary>
+		/// The initial value in the plain-text input when it is loaded.
+		/// </summary>
+		[JsonPropertyName("initial_value")]
+		public string? InitialValue { get; set; }
+
+		/// <summary>
+		/// Indicates whether the input will be a single line (false) or a larger textarea (true). Defaults to false.
+		/// </summary>
+		[JsonPropertyName("multiline")]
+		public bool MultiLine { get; set; }
+
+		/// <summary>
+		/// The minimum length of input that the user must provide. If the user provides less, they will receive an error. Maximum value is 3000.
+		/// </summary>
+		[JsonPropertyName("min_length")]
+		public int? MinLength { get; set; }
+
+		/// <summary>
+		/// The maximum length of input that the user can provide. If the user provides more, they will receive an error.
+		/// </summary>
+		[JsonPropertyName("max_length")]
+		public int? MaxLength { get; set; }
+
+		/// <summary>
+		/// Indicates whether the element will be set to auto focus within the view object. Only one element can be set to true. Defaults to false.
+		/// </summary>
+		[JsonPropertyName("focus_on_load")]
+		public bool FocusOnLoad { get; set; }
+
+		/// <summary>
+		/// Construct a new Button action element.
+		/// </summary>
+		public PlainTextInputElement(string actionId, PlainTextObject? placeholder = null, string? initialValue = null, bool multiLine = false)
+			: base("plain_text_input")
+		{
+			ActionId = actionId;
+			Placeholder = placeholder;
+			InitialValue = initialValue;
+			MultiLine = multiLine;
+		}
+	}
+}
diff --git a/Engine/Source/Programs/Shared/EpicGames.Slack/Elements/RadioButtonGroupElement.cs b/Engine/Source/Programs/Shared/EpicGames.Slack/Elements/RadioButtonGroupElement.cs
new file mode 100644
--- /dev/null
+++ b/Engine/Source/Programs/Shared/EpicGames.Slack/Elements/RadioButtonGroupElement.cs
@@ -0,0 +1,58 @@
+// Copyright Epic Games, Inc. All Rights Reserved.
+
+using System;
+using System.Collections.Generic;
+using System.Globalization;
+using System.Runtime.Serialization;
+using System.Text.Json;
+using System.Text.Json.Serialization;
+
+namespace EpicGames.Slack.Elements
+{
+	/// <summary>
+	/// A radio button group that allows a user to choose one item from a list of possible options.
+	/// </summary>
+	public class RadioButtonGroupElement : Element
+	{
+		/// <summary>
+		/// An identifier for this action. You can use this when you receive an interaction payload to identify the source of the action. 
+		/// Should be unique among all other action_ids in the containing block. Maximum length for this field is 255 characters.
+		/// </summary>
+		[JsonPropertyName("action_id")]
+		public string ActionId { get; set; }
+
+		/// <summary>
+		/// An array of option objects. A maximum of 10 options are allowed.
+		/// </summary>
+		[JsonPropertyName("options")]
+		public List<Option> Options { get; } = new List<Option>();
+
+		/// <summary>
+		/// An array of option objects that exactly matches one or more of the options within options. These options will be selected when the checkbox group initially loads.
+		/// </summary>
+		[JsonPropertyName("initial_options")]
+		public List<Option> InitialOptions { get; } = new List<Option>();
+
+		/// <summary>
+		/// A confirm object that defines an optional confirmation dialog that appears after clicking one of the checkboxes in this element.
+		/// </summary>
+		[JsonPropertyName("confirm")]
+		public SlackConfirm? Confirm { get; set; }
+
+		/// <summary>
+		/// Indicates whether the element will be set to auto focus within the view object. Only one element can be set to true. Defaults to false.
+		/// </summary>
+		[JsonPropertyName("focus_on_load")]
+		public bool FocusOnLoad { get; set; }
+
+		/// <summary>
+		/// Construct a new Button action element.
+		/// </summary>
+		public RadioButtonGroupElement(string actionId, List<Option> options)
+			: base("radio_buttons")
+		{
+			ActionId = actionId;
+			Options.AddRange(options);
+		}
+	}
+}
diff --git a/Engine/Source/Programs/Shared/EpicGames.Slack/Elements/SelectMenuElement.cs b/Engine/Source/Programs/Shared/EpicGames.Slack/Elements/SelectMenuElement.cs
new file mode 100644
--- /dev/null
+++ b/Engine/Source/Programs/Shared/EpicGames.Slack/Elements/SelectMenuElement.cs
@@ -0,0 +1,94 @@
+// Copyright Epic Games, Inc. All Rights Reserved.
+
+using System;
+using System.Collections.Generic;
+using System.Globalization;
+using System.Runtime.Serialization;
+using System.Text.Json;
+using System.Text.Json.Serialization;
+
+namespace EpicGames.Slack.Elements
+{
+	/// <summary>
+	/// An element which lets users easily select a date from a calendar style UI.
+	/// </summary>
+	public abstract class SelectMenuElement : Element
+	{
+		/// <summary>
+		/// An identifier for this action. You can use this when you receive an interaction payload to identify the source of the action. 
+		/// Should be unique among all other action_ids in the containing block. Maximum length for this field is 255 characters.
+		/// </summary>
+		[JsonPropertyName("action_id")]
+		public string ActionId { get; set; }
+
+		/// <summary>
+		/// Constructor
+		/// </summary>
+		protected SelectMenuElement(string type, string actionId) : base(type)
+		{
+			ActionId = actionId;
+		}
+	}
+
+	/// <summary>
+	/// This is the simplest form of select menu, with a static list of options passed in when defining the element.
+	/// </summary>
+	public class StaticSelectMenuElement : SelectMenuElement
+	{
+		/// <summary>
+		/// A plain_text only text object that defines the placeholder text shown on the menu. Maximum length for the text in this field is 150 characters.
+		/// </summary>
+		[JsonPropertyName("placeholder")]
+		public PlainTextObject Placeholder { get; set; }
+
+		/// <summary>
+		/// An array of option objects. A maximum of 10 options are allowed.
+		/// </summary>
+		[JsonPropertyName("options")]
+		public List<Option>? Options { get; set; }
+
+		/// <summary>
+		/// An array of option objects. A maximum of 10 options are allowed.
+		/// </summary>
+		[JsonPropertyName("option_groups")]
+		public List<OptionGroup>? OptionGroups { get; set; }
+
+		/// <summary>
+		/// An array of option objects that exactly matches one or more of the options within options. These options will be selected when the checkbox group initially loads.
+		/// </summary>
+		[JsonPropertyName("initial_option")]
+		public Option? InitialOption { get; }
+
+		/// <summary>
+		/// A confirm object that defines an optional confirmation dialog that appears after clicking one of the checkboxes in this element.
+		/// </summary>
+		[JsonPropertyName("confirm")]
+		public SlackConfirm? Confirm { get; set; }
+
+		/// <summary>
+		/// Indicates whether the element will be set to auto focus within the view object. Only one element can be set to true. Defaults to false.
+		/// </summary>
+		[JsonPropertyName("focus_on_load")]
+		public bool FocusOnLoad { get; set; }
+
+		/// <summary>
+		/// Construct a new Button action element.
+		/// </summary>
+		public StaticSelectMenuElement(string actionId, PlainTextObject placeholder, List<Option> options)
+			: base("static_select", actionId)
+		{
+			Placeholder = placeholder;
+			Options = new List<Option>(options);
+		}
+
+		/// <summary>
+		/// Construct a new Button action element.
+		/// </summary>
+		public StaticSelectMenuElement(string actionId, PlainTextObject placeholder, List<OptionGroup> optionGroups)
+			: base("static_select", actionId)
+		{
+			Placeholder = placeholder;
+			OptionGroups = new List<OptionGroup>(optionGroups);
+		}
+	}
+}
diff --git a/Engine/Source/Programs/Shared/EpicGames.Slack/Elements/TextObject.cs b/Engine/Source/Programs/Shared/EpicGames.Slack/Elements/TextObject.cs
new file mode 100644
--- /dev/null
+++ b/Engine/Source/Programs/Shared/EpicGames.Slack/Elements/TextObject.cs
@@ -0,0 +1,68 @@
+// Copyright Epic Games, Inc. All Rights Reserved.
+
+using System;
+using System.Collections.Generic;
+using System.Globalization;
+using System.Text.Json;
+using System.Text.Json.Serialization;
+
+namespace EpicGames.Slack.Elements
+{
+	/// <summary>
+	/// A text object
+	/// </summary>
+	public class TextObject : Element
+	{
+		/// <summary>
+		/// Text to render
+		/// </summary>
+		[JsonPropertyName("text"), JsonPropertyOrder(1)]
+		public string Text { get; set; }
+
+		/// <summary>
+		/// Whether to show emojis in the text
+		/// </summary>
+		[JsonPropertyName("emoji"), JsonPropertyOrder(2)]
+		public bool? Emoji { get; set; }
+
+		/// <summary>
+		/// Constructor
+		/// </summary>
+		/// <param name="text"></param>
+		public TextObject(string text) : this("mrkdwn", text, null)
+		{
+		}
+
+		/// <summary>
+		/// Constructor
+		/// </summary>
+		protected TextObject(string type, string text, bool? emoji) : base(type)
+		{
+			Text = text;
+			Emoji = emoji;
+		}
+
+		/// <summary>
+		/// Implicit conversion from a regular string
+		/// </summary>
+		public static implicit operator TextObject(string text) => new TextObject(text);
+	}
+
+	/// <summary>
+	/// Plain text object
+	/// </summary>
+	public class PlainTextObject : TextObject
+	{
+		/// <summary>
+		/// Constructor
+		/// </summary>
+		public PlainTextObject(string text, bool? emoji = null) : base("plain_text", text, emoji)
+		{
+		}
+
+		/// <summary>
+		/// Implicit conversion from a regular string
+		/// </summary>
+		public static implicit operator PlainTextObject(string text) => new PlainTextObject(text);
+	}
+}
diff --git a/Engine/Source/Programs/Shared/EpicGames.Slack/Elements/TimePickerElement.cs b/Engine/Source/Programs/Shared/EpicGames.Slack/Elements/TimePickerElement.cs
new file mode 100644
--- /dev/null
+++ b/Engine/Source/Programs/Shared/EpicGames.Slack/Elements/TimePickerElement.cs
@@ -0,0 +1,57 @@
+// Copyright Epic Games, Inc. All Rights Reserved.
+
+using System;
+using System.Collections.Generic;
+using System.Globalization;
+using System.Runtime.Serialization;
+using System.Text.Json;
+using System.Text.Json.Serialization;
+
+namespace EpicGames.Slack.Elements
+{
+	/// <summary>
+	/// An element which lets users easily select a date from a calendar style UI.
+	/// </summary>
+	public class TimePickerElement : Element
+	{
+		/// <summary>
+		/// An identifier for this action. You can use this when you receive an interaction payload to identify the source of the action. 
+		/// Should be unique among all other action_ids in the containing block. Maximum length for this field is 255 characters.
+		/// </summary>
+		[JsonPropertyName("action_id")]
+		public string ActionId { get; set; }
+
+		/// <summary>
+		/// A plain_text only text object that defines the placeholder text shown on the datepicker. Maximum length for the text in this field is 150 characters.
+		/// </summary>
+		[JsonPropertyName("placeholder")]
+		public PlainTextObject? Placeholder { get; set; }
+
+		/// <summary>
+		/// The initial time that is selected when the element is loaded. This should be in the format HH:mm, where HH is the 24-hour format of an hour (00 to 23) and mm is minutes with leading zeros (00 to 59), for example 22:25 for 10:25pm.
+		/// </summary>
+		[JsonPropertyName("initial_time")]
+		public string? InitialTime { get; set; }
+
+		/// <summary>
+		/// A confirm object that defines an optional confirmation dialog that appears after clicking one of the checkboxes in this element.
+		/// </summary>
+		[JsonPropertyName("confirm")]
+		public SlackConfirm? Confirm { get; set; }
+
+		/// <summary>
+		/// Indicates whether the element will be set to auto focus within the view object. Only one element can be set to true. Defaults to false.
+		/// </summary>
+		[JsonPropertyName("focus_on_load")]
+		public bool FocusOnLoad { get; set; }
+
+		/// <summary>
+		/// Construct a new Button action element.
+		/// </summary>
+		public TimePickerElement(string actionId)
+			: base("timepicker")
+		{
+			ActionId = actionId;
+		}
+	}
+}
diff --git a/Engine/Source/Programs/Shared/EpicGames.Slack/EpicGames.Slack.csproj b/Engine/Source/Programs/Shared/EpicGames.Slack/EpicGames.Slack.csproj
new file mode 100644
--- /dev/null
+++ b/Engine/Source/Programs/Shared/EpicGames.Slack/EpicGames.Slack.csproj
@@ -0,0 +1,42 @@
+<Project Sdk="Microsoft.NET.Sdk">
+  <PropertyGroup>
+    <TargetFramework>net6.0</TargetFramework>
+    <Configuration Condition=" '$(Configuration)' == '' ">Development</Configuration>
+    <OutputType>Library</OutputType>
+    <Deterministic>true</Deterministic>
+    <RootNamespace>EpicGames.Slack</RootNamespace>
+    <GenerateAssemblyInfo>false</GenerateAssemblyInfo>
+    <GenerateTargetFrameworkAttribute>false</GenerateTargetFrameworkAttribute>
+    <Configurations>Debug;Release;Development;Analyze</Configurations>
+    <Nullable>enable</Nullable>
+    <DebugType>pdbonly</DebugType>
+    <RunAnalyzersDuringBuild>false</RunAnalyzersDuringBuild>
+  </PropertyGroup>
+  <PropertyGroup Condition=" '$(Configuration)|$(Platform)' == 'Debug|AnyCPU' ">
+    <DocumentationFile></DocumentationFile>
+  </PropertyGroup>
+  <PropertyGroup Condition="'$(Configuration)|$(Platform)' == 'Development|AnyCPU'">
+    <DocumentationFile></DocumentationFile>
+    <Optimize>true</Optimize>
+  </PropertyGroup>
+  <PropertyGroup Condition="'$(Configuration)|$(Platform)' == 'Release|AnyCPU'">
+    <DocumentationFile></DocumentationFile>
+    <Optimize>true</Optimize>
+  </PropertyGroup>
+  <PropertyGroup Condition=" '$(Configuration)|$(Platform)' == 'Analyze|AnyCPU' ">
+    <DocumentationFile></DocumentationFile>
+    <RunAnalyzersDuringBuild>True</RunAnalyzersDuringBuild>
+    <GenerateDocumentationFile>True</GenerateDocumentationFile>
+  </PropertyGroup>
+  <ItemGroup>
+    <Compile Include="..\..\Shared\MetaData.cs">
+      <Link>Properties\MetaData.cs</Link>
+    </Compile>
+  </ItemGroup>
+  <ItemGroup>
+    <Folder Include="Blocks\" />
+  </ItemGroup>
+  <ItemGroup>
+    <PackageReference Include="Microsoft.Extensions.Logging.Abstractions" Version="6.0.0" />
+  </ItemGroup>
+</Project>
\ No newline at end of file
diff --git a/Engine/Source/Programs/Shared/EpicGames.Slack/Option.cs b/Engine/Source/Programs/Shared/EpicGames.Slack/Option.cs
new file mode 100644
--- /dev/null
+++ b/Engine/Source/Programs/Shared/EpicGames.Slack/Option.cs
@@ -0,0 +1,83 @@
+// Copyright Epic Games, Inc. All Rights Reserved.
+
+using System;
+using System.Collections.Generic;
+using System.Text.Json;
+using System.Text.Json.Serialization;
+using EpicGames.Slack.Elements;
+
+namespace EpicGames.Slack
+{
+	/// <summary>
+	/// An object that represents a single selectable item in a select menu, multi-select menu, checkbox group, radio button group, or overflow menu.
+	/// </summary>
+	public class Option
+	{
+		/// <summary>
+		/// A text object that defines the text shown in the option on the menu. Overflow, select, and multi-select menus can only use plain_text objects, while radio buttons and checkboxes can use mrkdwn text objects.
+		/// Maximum length for the text in this field is 75 characters.
+		/// </summary>
+		[JsonPropertyName("text")]
+		public TextObject Text { get; set; }
+
+		/// <summary>
+		/// A unique string value that will be passed to your app when this option is chosen. Maximum length for this field is 75 characters.
+		/// </summary>
+		[JsonPropertyName("value")]
+		public string Value { get; set; }
+
+		/// <summary>
+		/// A plain_text only text object that defines a line of descriptive text shown below the text field beside the radio button. Maximum length for the text object within this field is 75 characters.
+		/// </summary>
+		[JsonPropertyName("description")]
+		public PlainTextObject? Description { get; set; }
+
+		/// <summary>
+		/// A URL to load in the user's browser when the option is clicked. The url attribute is only available in overflow menus. Maximum length for this field is 3000 characters. If you're using url, you'll still receive an interaction payload and will need to send an acknowledgement response.
+		/// </summary>
+		[JsonPropertyName("url")]
+		public Uri? Url { get; set; }
+
+		/// <summary>
+		/// Constructor
+		/// </summary>
+		public Option(TextObject text, string value)
+		{
+			Text = text;
+			Value = value;
+		}
+
+		/// <summary>
+		/// Constructor
+		/// </summary>
+		public Option(string text, string value)
+			: this(new PlainTextObject(text), value)
+		{
+		}
+	}
+
+	/// <summary>
+	/// Provides a way to group options in a select menu or multi-select menu.
+	/// </summary>
+	public class OptionGroup
+	{
+		/// <summary>
+		/// A plain_text only text object that defines the label shown above this group of options. Maximum length for the text in this field is 75 characters.
+		/// </summary>
+		public PlainTextObject Label { get; }
+
+		/// <summary>
+		/// An array of option objects that belong to this specific group. Maximum of 100 items.
+		/// </summary>
+		public List<Option> Options { get; } = new List<Option>();
+
+		/// <summary>
+		/// Constructor
+		/// </summary>
+		public OptionGroup(PlainTextObject label, List<Option> options)
+		{
+			Label = label;
+			Options.AddRange(options);
+		}
+	}
+}
diff --git a/Engine/Source/Programs/Shared/EpicGames.Slack/SlackAttachment.cs b/Engine/Source/Programs/Shared/EpicGames.Slack/SlackAttachment.cs
new file mode 100644
--- /dev/null
+++ b/Engine/Source/Programs/Shared/EpicGames.Slack/SlackAttachment.cs
@@ -0,0 +1,45 @@
+// Copyright Epic Games, Inc. All Rights Reserved.
+
+using System;
+using System.Collections.Generic;
+using System.Drawing;
+using System.Text.Json;
+using System.Text.Json.Serialization;
+using EpicGames.Slack.Blocks;
+using EpicGames.Slack.Elements;
+
+namespace EpicGames.Slack
+{
+	/// <summary>
+	/// Abstract base class for all BlockKit blocks to derive from.
+	/// </summary>
+	public class SlackAttachment : ISlackBlockContainer
+	{
+		/// <summary>
+		/// The color of the line down the left side of the attachment.
+		/// </summary>
+		public string Color { get; set; } = "#ff0000";
+
+		/// <summary>
+		/// Used for the "toast" notifications sent by Slack. If not set, and the Text property of the
+		/// base message isn't set, the notification will have a message saying it has no preview.
+		/// </summary>
+		public string? FallbackText { get; set; }
+
+		/// <summary>
+		/// A collection of BlockKit blocks the attachment is composed from.
+		/// </summary>
+		public List<Block> Blocks { get; } = new List<Block>();
+
+		/// <summary>
+		/// Implicit conversion to a message of its own
+		/// </summary>
+		/// <param name="attachment"></param>
+		public static implicit operator SlackMessage(SlackAttachment attachment)
+		{
+			SlackMessage message = new SlackMessage();
+			message.Attachments.Add(attachment);
+			return message;
+		}
+	}
+}
diff --git a/Engine/Source/Programs/Shared/EpicGames.Slack/SlackClient.cs b/Engine/Source/Programs/Shared/EpicGames.Slack/SlackClient.cs
new file mode 100644
--- /dev/null
+++ b/Engine/Source/Programs/Shared/EpicGames.Slack/SlackClient.cs
@@ -0,0 +1,459 @@
+// Copyright Epic Games, Inc. All Rights Reserved.
+
+using System;
+using System.Collections.Generic;
+using System.Net.Http;
+using System.Net.WebSockets;
+using System.Text;
+using System.Text.Json;
+using System.Text.Json.Serialization;
+using System.Threading;
+using System.Threading.Tasks;
+using Microsoft.Extensions.Logging;
+
+namespace EpicGames.Slack
+{
+	public class SlackException : Exception
+	{
+		public string Code { get; }
+
+		public SlackException(string code) : base(code)
+		{
+			Code = code;
+		}
+	}
+
+	public class SlackClient
+	{
+		class SlackResponse
+		{
+			[JsonPropertyName("ok")]
+			public bool Ok { get; set; }
+
+			[JsonPropertyName("error")]
+			public string? Error { get; set; }
+		}
+
+		readonly HttpClient _httpClient;
+		readonly JsonSerializerOptions _serializerOptions;
+		readonly ILogger _logger;
+
+		public SlackClient(HttpClient httpClient, ILogger logger)
+		{
+			_httpClient = httpClient;
+
+			_serializerOptions = new JsonSerializerOptions();
+			_serializerOptions.DefaultIgnoreCondition = JsonIgnoreCondition.WhenWritingNull;
+
+			_logger = logger;
+		}
+
+		static bool ShouldLogError<TResponse>(TResponse response) where TResponse : SlackResponse => !response.Ok;
+
+		private Task<TResponse> SendRequestAsync<TResponse>(string requestUrl, object request) where TResponse : SlackResponse
+		{
+			return SendRequestAsync<TResponse>(requestUrl, request, ShouldLogError);
+		}
+
+		private async Task<TResponse> SendRequestAsync<TResponse>(string requestUrl, object request, Func<TResponse, bool> shouldLogError) where TResponse : SlackResponse
+		{
+			using (HttpRequestMessage sendMessageRequest = new HttpRequestMessage(HttpMethod.Post, requestUrl))
+			{
+				string requestJson = JsonSerializer.Serialize(request, _serializerOptions);
+				using (StringContent messageContent = new StringContent(requestJson, Encoding.UTF8, "application/json"))
+				{
+					sendMessageRequest.Content = messageContent;
+					return await SendRequestAsync<TResponse>(sendMessageRequest, requestJson, shouldLogError);
+				}
+			}
+		}
+
+		private Task<TResponse> SendRequestAsync<TResponse>(HttpRequestMessage request, string requestJson) where TResponse : SlackResponse
+		{
+			return SendRequestAsync<TResponse>(request, requestJson, ShouldLogError);
+		}
+
+		private async Task<TResponse> SendRequestAsync<TResponse>(HttpRequestMessage request, string requestJson, Func<TResponse, bool> shouldLogError) where TResponse : SlackResponse
+		{
+			HttpResponseMessage response = await _httpClient.SendAsync(request);
+			byte[] responseBytes = await response.Content.ReadAsByteArrayAsync();
+
+			TResponse responseObject = JsonSerializer.Deserialize<TResponse>(responseBytes)!;
+			if (shouldLogError(responseObject))
+			{
+				_logger.LogError("Failed to send Slack message ({Error}). Request: {Request}. Response: {Response}", responseObject.Error, requestJson, Encoding.UTF8.GetString(responseBytes));
+				throw new SlackException(responseObject.Error ?? "unknown");
+			}
+			return responseObject;
+		}
+
+		#region Chat
+
+		const string PostMessageUrl = "https://slack.com/api/chat.postMessage";
+		const string UpdateMessageUrl = "https://slack.com/api/chat.update";
+		const string GetPermalinkUrl = "https://slack.com/api/chat.getPermalink";
+
+		class PostMessageRequest
+		{
+			[JsonPropertyName("channel")]
+			public string? Channel { get; set; }
+
+			[JsonPropertyName("ts")]
+			public string? Ts { get; set; }
+
+			[JsonPropertyName("thread_ts")]
+			public string? ThreadTs { get; set; }
+
+			[JsonPropertyName("text")]
+			public string? Text { get; set; }
+
+			[JsonPropertyName("mrkdwn")]
+			public bool? Markdown { get; set; }
+
+			[JsonPropertyName("blocks")]
+			public List<Block>? Blocks { get; set; }
+			
+			[JsonPropertyName("reply_broadcast")]
+			public bool? ReplyBroadcast { get; set; }
+
+			[JsonPropertyName("unfurl_links")]
+			public bool? UnfurlLinks { get; set; }
+
+			[JsonPropertyName("unfurl_media")]
+			public bool? UnfurlMedia { get; set; }
+		}
+
+		class PostMessageResponse : SlackResponse
+		{
+			[JsonPropertyName("channel")]
+			public string? Channel { get; set; }
+
+			[JsonPropertyName("ts")]
+			public string? Ts { get; set; }
+		}
+
+		public async Task<string> PostMessageAsync(string recipient, SlackMessage message)
+		{
+			return await PostOrUpdateMessageAsync(recipient, null, null, message, false);
+		}
+
+		public async Task<string> PostMessageAsync(string recipient, string threadTs, SlackMessage message, bool replyBroadcast = false)
+		{
+			return await PostOrUpdateMessageAsync(recipient, null, threadTs, message, replyBroadcast);
+		}
+
+		public async Task UpdateMessageAsync(string recipient, string ts, SlackMessage message)
+		{
+			await PostOrUpdateMessageAsync(recipient, ts, null, message, false);
+		}
+
+		public async Task UpdateMessageAsync(string recipient, string ts, string threadTs, SlackMessage message, bool replyBroadcast = false)
+		{
+			await PostOrUpdateMessageAsync(recipient, ts, threadTs, message, replyBroadcast);
+		}
+
+		async Task<string> PostOrUpdateMessageAsync(string recipient, string? ts, string? threadTs, SlackMessage message, bool replyBroadcast)
+		{
+			PostMessageRequest request = new PostMessageRequest();
+			request.Channel = recipient;
+			request.Ts = ts;
+			request.ThreadTs = threadTs;
+			request.Text = message.Text;
+			if (message.Blocks.Count > 0)
+			{
+				request.Blocks = message.Blocks;
+			}
+			if (replyBroadcast)
+			{
+				request.ReplyBroadcast = replyBroadcast;
+			}
+			if (!message.UnfurlLinks)
+			{
+				request.UnfurlLinks = false;
+			}
+			if (!message.UnfurlMedia)
+			{
+				request.UnfurlMedia = false;
+			}
+
+			PostMessageResponse response = await SendRequestAsync<PostMessageResponse>(ts == null? PostMessageUrl : UpdateMessageUrl, request);
+			if (!response.Ok || response.Ts == null)
+			{
+				throw new SlackException(response.Error ?? "unknown");
+			}
+
+			return response.Ts;
+		}
+
+		class GetPermalinkRequest
+		{
+			[JsonPropertyName("channel")]
+			public string? Channel { get; set; }
+
+			[JsonPropertyName("message_ts")]
+			public string? MessageTs { get; set; }
+		}
+
+		class GetPermalinkResponse : SlackResponse
+		{
+			[JsonPropertyName("channel")]
+			public string? Channel { get; set; }
+
+			[JsonPropertyName("permalink")]
+			public string? Permalink { get; set; }
+		}
+
+		public async Task<string> GetPermalinkAsync(string channel, string ts)
+		{
+			string requestUrl = $"{GetPermalinkUrl}?channel={channel}&message_ts={ts}";
+			using (HttpRequestMessage sendMessageRequest = new HttpRequestMessage(HttpMethod.Get, requestUrl))
+			{
+				GetPermalinkResponse response = await SendRequestAsync<GetPermalinkResponse>(sendMessageRequest, "");
+				if (!response.Ok || response.Permalink == null)
+				{
+					throw new SlackException(response.Error ?? "unknown");
+				}
+				return response.Permalink;
+			}
+		}
+
+		#endregion
+
+		#region Reactions
+
+		const string AddReactionUrl = "https://slack.com/api/reactions.add";
+		const string RemoveReactionUrl = "https://slack.com/api/reactions.remove";
+
+		class ReactionMessage
+		{
+			[JsonPropertyName("channel")]
+			public string? Channel { get; set; }
+
+			[JsonPropertyName("timestamp")]
+			public string? Ts { get; set; }
+
+			[JsonPropertyName("name")]
+			public string? Name { get; set; }
+		}
+
+		public async Task AddReactionAsync(string channel, string ts, string name)
+		{
+			ReactionMessage message = new ReactionMessage();
+			message.Channel = channel;
+			message.Ts = ts;
+			message.Name = name;
+
+			static bool ShouldLogError(SlackResponse response) => !response.Ok && !String.Equals(response.Error, "already_reacted", StringComparison.Ordinal);
+
+			await SendRequestAsync<SlackResponse>(AddReactionUrl, message, ShouldLogError);
+		}
+
+		public async Task RemoveReactionAsync(string channel, string ts, string name)
+		{
+			ReactionMessage message = new ReactionMessage();
+			message.Channel = channel;
+			message.Ts = ts;
+			message.Name = name;
+
+			static bool ShouldLogError(SlackResponse response) => !response.Ok && !String.Equals(response.Error, "no_reaction", StringComparison.Ordinal);
+
+			await SendRequestAsync<SlackResponse>(RemoveReactionUrl, message, ShouldLogError);
+		}
+
+		#endregion
+
+		#region Conversations
+
+		const string ConversationsInviteUrl = "https://slack.com/api/conversations.invite";
+
+		class InviteMessage
+		{
+			[JsonPropertyName("channel")]
+			public string? Channel { get; set; }
+
+			[JsonPropertyName("users")]
+			public string? Users { get; set; } // Comma separated list of ids
+		}
+
+		public Task InviteUserAsync(string channel, string userId) => InviteUsersAsync(channel, new[] { userId });
+
+		public async Task InviteUsersAsync(string channel, IEnumerable<string> userIds)
+		{
+			InviteMessage message = new InviteMessage();
+			message.Channel = channel;
+			message.Users = String.Join(",", userIds);
+			await SendRequestAsync<SlackResponse>(ConversationsInviteUrl, message);
+		}
+
+		#endregion
+
+		#region Users
+
+		const string UsersInfoUrl = "https://slack.com/api/users.info";
+		const string UsersLookupByEmailUrl = "https://slack.com/api/users.lookupByEmail";
+
+		class UserResponse : SlackResponse
+		{
+			[JsonPropertyName("user")]
+			public SlackUser? User { get; set; }
+		}
+
+		public async Task<SlackUser> GetUserAsync(string userId)
+		{
+			using (HttpRequestMessage request = new HttpRequestMessage(HttpMethod.Post, $"{UsersInfoUrl}?user={userId}"))
+			{
+				UserResponse response = await SendRequestAsync<UserResponse>(request, "", ShouldLogError);
+				return response.User!;
+			}
+		}
+
+		public async Task<SlackUser?> FindUserByEmailAsync(string email)
+		{
+			using (HttpRequestMessage request = new HttpRequestMessage(HttpMethod.Post, $"{UsersLookupByEmailUrl}?email={email}"))
+			{
+				static bool ShouldLogError(UserResponse response) => !response.Ok && !String.Equals(response.Error, "users_not_found", StringComparison.Ordinal);
+
+				UserResponse response = await SendRequestAsync<UserResponse>(request, "", ShouldLogError);
+				return response.User;
+			}
+		}
+
+		#endregion
+
+		#region Views
+
+		const string ViewsOpenUrl = "https://slack.com/api/views.open";
+
+		class ViewsOpenRequest
+		{
+			[JsonPropertyName("trigger_id")]
+			public string? TriggerId { get; set; }
+
+			[JsonPropertyName("view")]
+			public SlackView? View { get; set; }
+		}
+
+		public async Task OpenViewAsync(string triggerId, SlackView view)
+		{
+			ViewsOpenRequest request = new ViewsOpenRequest();
+			request.TriggerId = triggerId;
+			request.View = view;
+			await SendRequestAsync<PostMessageResponse>(ViewsOpenUrl, request);
+		}
+
+		#endregion
+/*
+		#region Interactions
+
+		const string AppConnectionOpenUrl = "https://slack.com/api/apps.connections.open";
+
+		class OpenSocketResponse : SlackResponse
+		{
+			[JsonPropertyName("url")]
+			public Uri? Url { get; set; }
+		}
+
+		class EventMessage
+		{
+			[JsonPropertyName("type")]
+			public string? Type { get; set; }
+
+			[JsonPropertyName("envelope_id")]
+			public string? EnvelopeId { get; set; }
+
+			[JsonPropertyName("payload")]
+			public EventPayload? Payload { get; set; }
+		}
+
+		class EventPayload
+		{
+			[JsonPropertyName("type")]
+			public string? Type { get; set; }
+
+			[JsonPropertyName("trigger_id")]
+			public string? TriggerId { get; set; }
+
+			[JsonPropertyName("user")]
+			public UserInfo? User { get; set; }
+
+			[JsonPropertyName("response_url")]
+			public string? ResponseUrl { get; set; }
+
+			[JsonPropertyName("actions")]
+			public List<ActionInfo> Actions { get; set; } = new List<ActionInfo>();
+		}
+
+		/// <inheritdoc/>
+		public async Task HandleInteractionsAsync(CancellationToken stoppingToken)
+		{
+			using HttpRequestMessage request = new HttpRequestMessage(HttpMethod.Post, AppConnectionOpenUrl);
+			request.Content = new FormUrlEncodedContent(Array.Empty<KeyValuePair<string?, string?>>());
+
+			OpenSocketResponse response = await SendRequestAsync<OpenSocketResponse>(request, "");
+		}
+
+		private async Task HandleInteractionsInternalAsync(Uri socketUrl, CancellationToken stoppingToken)
+		{
+			using ClientWebSocket socket = new ClientWebSocket();
+			await socket.ConnectAsync(socketUrl, stoppingToken);
+
+			byte[] buffer = new byte[2048];
+			while (!stoppingToken.IsCancellationRequested)
+			{
+				// Read the next message
+				int length = 0;
+				for (; ; )
+				{
+					if (length == buffer.Length)
+					{
+						Array.Resize(ref buffer, buffer.Length + 2048);
+					}
+
+					WebSocketReceiveResult result = await socket.ReceiveAsync(new ArraySegment<byte>(buffer, length, buffer.Length - length), stoppingToken);
+					if (result.MessageType == WebSocketMessageType.Close)
+					{
+						return;
+					}
+					length += result.Count;
+
+					if (result.EndOfMessage)
+					{
+						break;
+					}
+				}
+
+				// Get the message data
+				_logger.LogInformation("Slack event: {Message}", Encoding.UTF8.GetString(buffer, 0, length));
+				EventMessage eventMessage = JsonSerializer.Deserialize<EventMessage>(buffer.AsSpan(0, length))!;
+
+				// Acknowledge the message
+				if (eventMessage.EnvelopeId != null)
+				{
+					object response = new { eventMessage.EnvelopeId };
+					await socket.SendAsync(JsonSerializer.SerializeToUtf8Bytes(response), WebSocketMessageType.Text, true, stoppingToken);
+				}
+
+				// Handle the message type
+				if (eventMessage.Type != null)
+				{
+					string type = eventMessage.Type;
+					if (type.Equals("disconnect", StringComparison.Ordinal))
+					{
+						break;
+					}
+					else if (type.Equals("interactive", StringComparison.Ordinal))
+					{
+						await HandleInteractionMessage(eventMessage);
+					}
+					else
+					{
+						_logger.LogDebug("Unhandled event type ({Type})", type);
+					}
+				}
+			}
+		}
+
+		#endregion
+*/
+	}
+}
diff --git a/Engine/Source/Programs/Shared/EpicGames.Slack/SlackConfirm.cs b/Engine/Source/Programs/Shared/EpicGames.Slack/SlackConfirm.cs
new file mode 100644
--- /dev/null
+++ b/Engine/Source/Programs/Shared/EpicGames.Slack/SlackConfirm.cs
@@ -0,0 +1,71 @@
+// Copyright Epic Games, Inc. All Rights Reserved.
+
+using System.Text.Json.Serialization;
+using EpicGames.Slack.Elements;
+
+namespace EpicGames.Slack
+{
+	/// <summary>
+	/// An object that defines a dialog that provides a confirmation step to any interactive element. This dialog will ask the user to confirm their action by offering a confirm and deny buttons.
+	/// </summary>
+	public class SlackConfirm
+	{
+		/// <summary>
+		/// A plain_text-only text object that defines the dialog's title. Maximum length for this field is 100 characters.
+		/// </summary>
+		[JsonPropertyName("title")]
+		public PlainTextObject Title { get; set; }
+
+		/// <summary>
+		/// A text object that defines the explanatory text that appears in the confirm dialog. Maximum length for the text in this field is 300 characters.
+		/// </summary>
+		[JsonPropertyName("text")]
+		public TextObject Text { get; set; }
+
+		/// <summary>
+		/// A plain_text-only text object to define the text of the button that confirms the action. Maximum length for the text in this field is 30 characters.
+		/// </summary>
+		[JsonPropertyName("confirm")]
+		public PlainTextObject Confirm { get; set; }
+
+		/// <summary>
+		/// A plain_text-only text object to define the text of the button that cancels the action. Maximum length for the text in this field is 30 characters.
+		/// </summary>
+		[JsonPropertyName("deny")]
+		public PlainTextObject Deny { get; set; }
+
+		/// <summary>
+		/// Defines the color scheme applied to the confirm button. A value of danger will display the button with a red background on desktop, or red text on mobile. A value of primary will display the button with a green background on desktop, or blue text on mobile. If this field is not provided, the default value will be primary.
+		/// </summary>
+		[JsonPropertyName("style")]
+		public ButtonStyle Style { get; set; }
+
+		/// <summary>
+		/// Constructor
+		/// </summary>
+		public SlackConfirm(string title, string text, string confirm, string deny, ButtonStyle style = ButtonStyle.Primary)
+			: this(title, new PlainTextObject(text), confirm, deny, style)
+		{
+		}
+
+		/// <summary>
+		/// Constructor
+		/// </summary>
+		public SlackConfirm(string title, TextObject text, string confirm, string deny, ButtonStyle style = ButtonStyle.Primary)
+			: this(new PlainTextObject(title), text, new PlainTextObject(confirm), new PlainTextObject(deny), style)
+		{
+		}
+
+		/// <summary>
+		/// Constructor
+		/// </summary>
+		public SlackConfirm(PlainTextObject title, TextObject text, PlainTextObject confirm, PlainTextObject deny, ButtonStyle style = ButtonStyle.Primary)
+		{
+			Title = title;
+			Text = text;
+			Confirm = confirm;
+			Deny = deny;
+			Style = style;
+		}
+	}
+}
diff --git a/Engine/Source/Programs/Shared/EpicGames.Slack/SlackMessage.cs b/Engine/Source/Programs/Shared/EpicGames.Slack/SlackMessage.cs
new file mode 100644
--- /dev/null
+++ b/Engine/Source/Programs/Shared/EpicGames.Slack/SlackMessage.cs
@@ -0,0 +1,57 @@
+// Copyright Epic Games, Inc. All Rights Reserved.
+
+using System;
+using System.Collections.Generic;
+using System.Linq;
+using System.Text;
+using System.Threading.Tasks;
+
+namespace EpicGames.Slack
+{
+	/// <summary>
+	/// A slack message object
+	/// </summary>
+	public class SlackMessage : ISlackBlockContainer
+	{
+		/// <summary>
+		/// The formatted text of the message to be published. If blocks are included, this will become the fallback text used in notifications.
+		/// </summary>
+		public string? Text { get; set; }
+
+		/// <summary>
+		/// Disable Slack markup parsing by setting to false. Enabled by default.
+		/// </summary>
+		public bool? Markdown { get; set; }
+
+		/// <summary>
+		/// A JSON-based array of structured blocks, presented as a URL-encoded string.
+		/// </summary>
+		public List<Block> Blocks { get; } = new List<Block>();
+
+		/// <summary>
+		/// Attachments for the message
+		/// </summary>
+		public List<SlackAttachment> Attachments { get; } = new List<SlackAttachment>();
+
+		/// <summary>
+		/// Pass true to enable unfurling of primarily text-based content.
+		/// </summary>
+		public bool UnfurlLinks { get; set; } = true;
+
+		/// <summary>
+		/// Pass false to disable unfurling of media content.
+		/// </summary>
+		public bool UnfurlMedia { get; set; } = true;
+
+		/// <summary>
+		/// Implciit conversion operator from a Markdown string
+		/// </summary>
+		/// <param name="text"></param>
+		public static implicit operator SlackMessage(string text)
+		{
+			SlackMessage message = new SlackMessage();
+			message.Text = text;
+			return message;
+		}
+	}
+}
diff --git a/Engine/Source/Programs/Shared/EpicGames.Slack/SlackUser.cs b/Engine/Source/Programs/Shared/EpicGames.Slack/SlackUser.cs
new file mode 100644
--- /dev/null
+++ b/Engine/Source/Programs/Shared/EpicGames.Slack/SlackUser.cs
@@ -0,0 +1,107 @@
+// Copyright Epic Games, Inc. All Rights Reserved.
+
+using System;
+using System.Collections.Generic;
+using System.Linq;
+using System.Text;
+using System.Text.Json.Serialization;
+using System.Threading.Tasks;
+
+namespace EpicGames.Slack
+{
+	/// <summary>
+	/// Information about a Slack user
+	/// </summary>
+	public class SlackUser
+	{
+		/// <summary>
+		/// Identifier for this workspace user. It is unique to the workspace containing the user. Use this field together with team_id as a unique key when storing related data or when specifying the user in API requests. We recommend considering the format of the string to be an opaque value, and not to rely on a particular structure.
+		/// </summary>
+		[JsonPropertyName("id")]
+		public string? Id { get; set; }
+
+		/// <summary>
+		/// The profile object contains the default fields of a user's workspace profile. A user's custom profile fields may be discovered using users.profile.get.
+		/// </summary>
+		[JsonPropertyName("profile")]
+		public SlackUserProfile? Profile { get; set; }
+
+		/// <summary>
+		/// A human-readable string for the geographic timezone-related region this user has specified in their account.
+		/// </summary>
+		[JsonPropertyName("tz")]
+		public string? Timezone { get; set; }
+
+		/// <summary>
+		/// Describes the commonly used name of the tz timezone.
+		/// </summary>
+		[JsonPropertyName("tz_label")]
+		public string? TimezoneLabel { get; set; }
+
+		/// <summary>
+		/// Indicates the number of seconds to offset UTC time by for this user's tz.
+		/// </summary>
+		[JsonPropertyName("tz_offset")]
+		public int TimezoneOffset { get; set; }
+	}
+
+	/// <summary>
+	/// Profile information for a Slack user
+	/// </summary>
+	public class SlackUserProfile
+	{
+		/// <summary>
+		/// The display name the user has chosen to identify themselves by in their workspace profile. Do not use this field as a unique identifier for a user, as it may change at any time. Instead, use id and team_id in concert.
+		/// </summary>
+		[JsonPropertyName("display_name")]
+		public string DisplayName { get; set; } = "";
+
+		/// <summary>
+		/// A valid email address. It cannot have spaces, and it must have an @ and a domain. It cannot be in use by another member of the same team. Changing a user's email address will send an email to both the old and new addresses, and also post a slackbot to the user informing them of the change. This field can only be changed by admins for users on paid teams.
+		/// </summary>
+		[JsonPropertyName("email")]
+		public string Email { get; set; } = "";
+
+		/// <summary>
+		/// The user's first name. The name slackbot cannot be used. Updating first_name will update the first name within real_name.
+		/// </summary>
+		[JsonPropertyName("first_name")]
+		public string FirstName { get; set; } = "";
+
+		/// <summary>
+		/// The user's last name. The name slackbot cannot be used. Updating last_name will update the second name within real_name.
+		/// </summary>
+		[JsonPropertyName("last_name")]
+		public string LastName { get; set; } = "";
+
+		/// <summary>
+		/// The user's first and last name. Updating this field will update first_name and last_name. If only one name is provided, the value of last_name will be cleared.
+		/// </summary>
+		[JsonPropertyName("real_name")]
+		public string RealName { get; set; } = "";
+
+		/// <summary>
+		/// Whether the user has a custom profile image
+		/// </summary>
+		[JsonPropertyName("is_custom_image")]
+		public bool IsCustomImage { get; set; }
+
+		/// <summary>
+		/// These various fields will contain https URLs that point to square ratio, web-viewable images (GIFs, JPEGs, or PNGs) that represent different sizes of a user's profile picture.
+		/// </summary>
+		[JsonPropertyName("image_24")]
+		public string? Image24 { get; set; }
+
+		/// <inheritdoc cref="Image24"/>
+		[JsonPropertyName("image_32")]
+		public string? Image32 { get; set; }
+
+		/// <inheritdoc cref="Image24"/>
+		[JsonPropertyName("image_48")]
+		public string? Image48 { get; set; }
+
+		/// <inheritdoc cref="Image24"/>
+		[JsonPropertyName("image_72")]
+		public string? Image72 { get; set; }
+	}
+}
diff --git a/Engine/Source/Programs/Shared/EpicGames.Slack/SlackView.cs b/Engine/Source/Programs/Shared/EpicGames.Slack/SlackView.cs
new file mode 100644
--- /dev/null
+++ b/Engine/Source/Programs/Shared/EpicGames.Slack/SlackView.cs
@@ -0,0 +1,154 @@
+// Copyright Epic Games, Inc. All Rights Reserved.
+
+using System;
+using System.Collections.Generic;
+using System.Diagnostics.CodeAnalysis;
+using System.Linq;
+using System.Text;
+using System.Text.Json.Serialization;
+using System.Threading.Tasks;
+using EpicGames.Slack.Elements;
+
+namespace EpicGames.Slack
+{
+	/// <summary>
+	/// Describes a Slack view
+	/// </summary>
+	public class SlackView : ISlackBlockContainer
+	{
+		/// <summary>
+		/// The type of view. Set to modal for modals and home for Home tabs.
+		/// </summary>
+		[JsonPropertyName("type")]
+		public string Type { get; } = "modal";
+
+		/// <summary>
+		/// The title that appears in the top-left of the modal. Must be a plain_text text element with a max length of 24 characters.
+		/// </summary>
+		[JsonPropertyName("title")]
+		public PlainTextObject Title { get; set; }
+
+		/// <summary>
+		/// An array of blocks that defines the content of the view. Max of 100 blocks.
+		/// </summary>
+		[JsonPropertyName("blocks")]
+		public List<Block> Blocks { get; } = new List<Block>();
+
+		/// <summary>
+		/// An optional plain_text element that defines the text displayed in the close button at the bottom-right of the view. Max length of 24 characters.
+		/// </summary>
+		[JsonPropertyName("close")]
+		public PlainTextObject? Close { get; set; }
+
+		/// <summary>
+		/// An optional plain_text element that defines the text displayed in the submit button at the bottom-right of the view. submit is required when an input block is within the blocks array. Max length of 24 characters.
+		/// </summary>
+		[JsonPropertyName("submit")]
+		public PlainTextObject? Submit { get; set; }
+
+		/// <summary>
+		/// An optional string that will be sent to your app in view_submission and block_actions events. Max length of 3000 characters.
+		/// </summary>
+		[JsonPropertyName("private_metadata")]
+		public string? PrivateMetadata { get; set; }
+
+		/// <summary>
+		/// An identifier to recognize interactions and submissions of this particular view.Don't use this to store sensitive information (use private_metadata instead). Max length of 255 characters.
+		/// </summary>
+		[JsonPropertyName("callback_id")]
+		public string? CallbackId { get; set; }
+
+		/// <summary>
+		/// When set to true, clicking on the close button will clear all views in a modal and close it. Defaults to false.
+		/// </summary>
+		[JsonPropertyName("clear_on_close")]
+		public bool? ClearOnClose { get; set; }
+
+		/// <summary>
+		/// Indicates whether Slack will send your request URL a view_closed event when a user clicks the close button. Defaults to false.
+		/// </summary>
+		[JsonPropertyName("notify_on_close")]
+		public bool? NotifyOnClose { get; set; }
+
+		/// <summary>
+		/// A custom identifier that must be unique for all views on a per-team basis.
+		/// </summary>
+		[JsonPropertyName("external_id")]
+		public string? ExternalId { get; set; }
+
+		/// <summary>
+		/// When set to true, disables the submit button until the user has completed one or more inputs. This property is for configuration modals.
+		/// </summary>
+		[JsonPropertyName("submit_disabled")]
+		public bool? SubmitDisabled { get; set; }
+
+		/// <summary>
+		/// State value when returned via a form notification
+		/// </summary>
+		[JsonPropertyName("state")]
+		public SlackViewState? State { get; set; }
+
+		/// <summary>
+		/// Constructor
+		/// </summary>
+		/// <param name="title"></param>
+		public SlackView(PlainTextObject title)
+		{
+			Title = title;
+		}
+	}
+
+	/// <summary>
+	/// State object returned from a submission
+	/// </summary>
+	public class SlackViewState
+	{
+		/// <summary>
+		/// Values for the form, by blockId, then by actionId
+		/// </summary>
+		[JsonPropertyName("values")]
+#pragma warning disable CA2227 // Collection properties should be read only
+		public Dictionary<string, Dictionary<string, SlackViewValue>> Values { get; set; } = new Dictionary<string, Dictionary<string, SlackViewValue>>();
+#pragma warning restore CA2227 // Collection properties should be read only
+
+		/// <summary>
+		/// Gets a value from the state
+		/// </summary>
+		/// <param name="blockId"></param>
+		/// <param name="actionId"></param>
+		/// <param name="value"></param>
+		/// <returns></returns>
+		public bool TryGetValue(string blockId, string actionId, [NotNullWhen(true)] out string? value)
+		{
+			Dictionary<string, SlackViewValue>? actions;
+			if (Values.TryGetValue(blockId, out actions) && actions.TryGetValue(actionId, out SlackViewValue? viewValue))
+			{
+				value = viewValue.Value;
+				return true;
+			}
+			else
+			{
+				value = null;
+				return false;
+			}
+		}
+	}
+
+	/// <summary>
+	/// Value for a form item
+	/// </summary>
+	public class SlackViewValue
+	{
+		/// <summary>
+		/// The item type
+		/// </summary>
+		[JsonPropertyName("type")]
+		public string Type { get; set; } = String.Empty;
+
+		/// <summary>
+		/// Value for the widget
+		/// </summary>
+		[JsonPropertyName("value")]
+		public string Value { get; set; } = String.Empty;
+	}
+}
