Horde Storage - Fixed issue were orphan blob cleanup would extend life time of refs by updating their last access time

diff --git a/Engine/Source/Programs/Horde/HordeStorage/Horde.Storage/Implementation/GC/OrphanBlobCleanupRefs.cs b/Engine/Source/Programs/Horde/HordeStorage/Horde.Storage/Implementation/GC/OrphanBlobCleanupRefs.cs
--- a/Engine/Source/Programs/Horde/HordeStorage/Horde.Storage/Implementation/GC/OrphanBlobCleanupRefs.cs
+++ b/Engine/Source/Programs/Horde/HordeStorage/Horde.Storage/Implementation/GC/OrphanBlobCleanupRefs.cs
@@ -147,7 +147,7 @@ namespace Horde.Storage.Implementation
                     try
                     {
                         (ObjectRecord, BlobContents?) _ = await _objectService.Get(blobNamespace,
-                            bucket, key, new string[] { "name" });
+                            bucket, key, new string[] { "name" }, doLastAccessTracking:false);
                         found = true;
                     }
                     catch (ObjectNotFoundException)
diff --git a/Engine/Source/Programs/Horde/HordeStorage/Horde.Storage/Implementation/ObjectService.cs b/Engine/Source/Programs/Horde/HordeStorage/Horde.Storage/Implementation/ObjectService.cs
--- a/Engine/Source/Programs/Horde/HordeStorage/Horde.Storage/Implementation/ObjectService.cs
+++ b/Engine/Source/Programs/Horde/HordeStorage/Horde.Storage/Implementation/ObjectService.cs
@@ -20,7 +20,7 @@ namespace Horde.Storage.Implementation
 {
     public interface IObjectService
     {
-        Task<(ObjectRecord, BlobContents?)> Get(NamespaceId ns, BucketId bucket, IoHashKey key, string[] fields);
+        Task<(ObjectRecord, BlobContents?)> Get(NamespaceId ns, BucketId bucket, IoHashKey key, string[] fields, bool doLastAccessTracking = true);
         Task<(ContentId[], BlobIdentifier[])> Put(NamespaceId ns, BucketId bucket, IoHashKey key, BlobIdentifier blobHash, CbObject payload);
         Task<(ContentId[], BlobIdentifier[])> Finalize(NamespaceId ns, BucketId bucket, IoHashKey key, BlobIdentifier blobHash);
 
@@ -53,7 +53,7 @@ namespace Horde.Storage.Implementation
             _lastAccessTracker = lastAccessTracker;
         }
 
-        public async Task<(ObjectRecord, BlobContents?)> Get(NamespaceId ns, BucketId bucket, IoHashKey key, string[]? fields = null)
+        public async Task<(ObjectRecord, BlobContents?)> Get(NamespaceId ns, BucketId bucket, IoHashKey key, string[]? fields = null, bool doLastAccessTracking = true)
         {
             // if no field filtering is being used we assume everything is needed
             IReferencesStore.FieldFlags flags = IReferencesStore.FieldFlags.All;
@@ -81,14 +81,17 @@ namespace Horde.Storage.Implementation
                 o = await _referencesStore.Get(ns, bucket, key, flags);
             }
 
-            // we do not wait for the last access tracking as it does not matter when it completes
-            Task lastAccessTask = _lastAccessTracker.TrackUsed(new LastAccessRecord(ns, bucket, key)).ContinueWith((task, _) =>
+            if (doLastAccessTracking)
             {
-                if (task.Exception != null)
+                // we do not wait for the last access tracking as it does not matter when it completes
+                Task lastAccessTask = _lastAccessTracker.TrackUsed(new LastAccessRecord(ns, bucket, key)).ContinueWith((task, _) =>
                 {
-                    _logger.Error(task.Exception, "Exception when tracking last access record");
-                }
-            }, null, TaskScheduler.Current);
+                    if (task.Exception != null)
+                    {
+                        _logger.Error(task.Exception, "Exception when tracking last access record");
+                    }
+                }, null, TaskScheduler.Current);
+            }
 
             BlobContents? blobContents = null;
             if ((flags & IReferencesStore.FieldFlags.IncludePayload) != 0)
