Horde Storage - Finalize errors now return 400 rather then 500s. Its still a error caused by bad cache keys from the cooker - but as its so tricky to fix properly in the cooker we will have these arounds for a long time and such its better to use 400 to indicate its a known type of error

diff --git a/Engine/Source/Programs/Horde/HordeStorage/Horde.Storage/Controllers/ReferencesController.cs b/Engine/Source/Programs/Horde/HordeStorage/Horde.Storage/Controllers/ReferencesController.cs
--- a/Engine/Source/Programs/Horde/HordeStorage/Horde.Storage/Controllers/ReferencesController.cs
+++ b/Engine/Source/Programs/Horde/HordeStorage/Horde.Storage/Controllers/ReferencesController.cs
@@ -749,11 +749,18 @@ namespace Horde.Storage.Controllers
                 return accessResult;
             }
 
-            (ContentId[] missingReferences, BlobIdentifier[] missingBlobs) = await _objectService.Finalize(ns, bucket, key, hash);
-            List<ContentHash> missingHashes = new List<ContentHash>(missingReferences);
-            missingHashes.AddRange(missingBlobs);
+            try
+            {
+                (ContentId[] missingReferences, BlobIdentifier[] missingBlobs) = await _objectService.Finalize(ns, bucket, key, hash);
+                List<ContentHash> missingHashes = new List<ContentHash>(missingReferences);
+                missingHashes.AddRange(missingBlobs);
 
-            return Ok(new PutObjectResponse(missingHashes.ToArray()));
+                return Ok(new PutObjectResponse(missingHashes.ToArray()));
+            }
+            catch (ObjectHashMismatchException e)
+            {
+                return BadRequest(e.Message);
+            }
         }
 
         [HttpPost("{ns}")]
