Horde Storage - Remove blob from blob index when the last region it existed in is being removed.

diff --git a/Engine/Source/Programs/Horde/HordeStorage/Horde.Storage/Implementation/Blob/ScyllaBlobIndex.cs b/Engine/Source/Programs/Horde/HordeStorage/Horde.Storage/Implementation/Blob/ScyllaBlobIndex.cs
--- a/Engine/Source/Programs/Horde/HordeStorage/Horde.Storage/Implementation/Blob/ScyllaBlobIndex.cs
+++ b/Engine/Source/Programs/Horde/HordeStorage/Horde.Storage/Implementation/Blob/ScyllaBlobIndex.cs
@@ -82,6 +82,12 @@ public class ScyllaBlobIndex : IBlobIndex
         scope.Span.ResourceName = $"{ns}.{id}";
 
         await _mapper.UpdateAsync<ScyllaBlobIndexTable>("SET regions = regions - ? WHERE namespace = ? AND blob_id = ?", new string[] { region }, ns.ToString(), new ScyllaBlobIdentifier(id));
+
+        BlobInfo? blobInfo = await GetBlobInfo(ns, id);
+        if (blobInfo != null && blobInfo.Regions.Count == 0)
+        {
+            await RemoveBlobFromIndex(ns, id);
+        }
     }
 
     public async Task<bool> BlobExistsInRegion(NamespaceId ns, BlobIdentifier blobIdentifier)
