Horde: Run certain pool sizing strategies more frequently

Some strategies, such as the one for auto-scaling compute tasks need to be much more reactive.

diff --git a/Engine/Source/Programs/Horde/Horde.Build/Agents/Fleet/AutoscaleServiceV2.cs b/Engine/Source/Programs/Horde/Horde.Build/Agents/Fleet/AutoscaleServiceV2.cs
--- a/Engine/Source/Programs/Horde/Horde.Build/Agents/Fleet/AutoscaleServiceV2.cs
+++ b/Engine/Source/Programs/Horde/Horde.Build/Agents/Fleet/AutoscaleServiceV2.cs
@@ -35,6 +35,7 @@ namespace Horde.Build.Agents.Fleet
 		private readonly IDogStatsd _dogStatsd;
 		private readonly IClock _clock;
 		private readonly ITicker _ticker;
+		private readonly ITicker _tickerHighFrequency;
 		private readonly ServerSettings _settings;
 		private readonly TimeSpan _defaultScaleOutCooldown;
 		private readonly TimeSpan _defaultScaleInCooldown;
@@ -66,6 +67,7 @@ namespace Horde.Build.Agents.Fleet
 			_dogStatsd = dogStatsd;
 			_clock = clock;
 			_ticker = clock.AddSharedTicker<AutoscaleServiceV2>(TimeSpan.FromMinutes(5.0), TickLeaderAsync, logger);
+			_tickerHighFrequency = clock.AddSharedTicker<AutoscaleServiceV2>(TimeSpan.FromSeconds(30), TickHighFrequencyAsync, logger);
 			_settings = settings.Value;
 			_logger = logger;
 			_defaultScaleOutCooldown = TimeSpan.FromSeconds(settings.Value.AgentPoolScaleOutCooldownSeconds);
@@ -85,18 +87,43 @@ namespace Horde.Build.Agents.Fleet
 		{
 			_logger.LogInformation("Autoscaling pools...");
 			Stopwatch stopwatch = Stopwatch.StartNew();
-			using IScope _ = GlobalTracer.Instance.BuildSpan("AutoscaleService").StartActive();
+			using IScope _ = GlobalTracer.Instance.BuildSpan("AutoscaleService.TickAsync").StartActive();
 			
+			await BatchResizePools(false);
+
+			stopwatch.Stop();
+			_logger.LogInformation("Autoscaling pools took {ElapsedTime} ms", stopwatch.ElapsedMilliseconds);
+		}
+
+		internal async ValueTask TickHighFrequencyAsync(CancellationToken stoppingToken)
+		{
+			_logger.LogInformation("Autoscaling pools (high frequency)...");
+			Stopwatch stopwatch = Stopwatch.StartNew();
+			using IScope _ = GlobalTracer.Instance.BuildSpan("AutoscaleService.TickHighFrequency").StartActive();
+			
+			await BatchResizePools(true);
+			
+			stopwatch.Stop();
+			_logger.LogInformation("Autoscaling pools (high frequency) took {ElapsedTime} ms", stopwatch.ElapsedMilliseconds);
+		}
+		
+		private async Task BatchResizePools(bool onlyHighFrequency)
+		{
 			// Group pools by strategy type to ensure they can be called once (more optimal)
 			Dictionary<PoolSizeStrategy, List<PoolSizeData>> poolSizeDataByStrategy = await GetPoolSizeDataByStrategyType();
+
+			if (onlyHighFrequency)
+			{
+				poolSizeDataByStrategy = poolSizeDataByStrategy
+					.Where(x => x.Key.SupportsHighFrequency())
+					.ToDictionary(i => i.Key, i => i.Value);
+			}
+			
 			foreach ((PoolSizeStrategy strategyType, List<PoolSizeData> currentData) in poolSizeDataByStrategy)
 			{
 				List<PoolSizeData> newData = await GetPoolSizeStrategy(strategyType).CalcDesiredPoolSizesAsync(currentData);
 				await ResizePools(newData);
 			}
-			
-			stopwatch.Stop();
-			_logger.LogInformation("Autoscaling pools took {ElapsedTime} ms", stopwatch.ElapsedMilliseconds);
 		}
 
 		internal async Task<Dictionary<PoolSizeStrategy, List<PoolSizeData>>> GetPoolSizeDataByStrategyType()
diff --git a/Engine/Source/Programs/Horde/Horde.Build/Agents/Fleet/IPoolSizeStrategy.cs b/Engine/Source/Programs/Horde/Horde.Build/Agents/Fleet/IPoolSizeStrategy.cs
--- a/Engine/Source/Programs/Horde/Horde.Build/Agents/Fleet/IPoolSizeStrategy.cs
+++ b/Engine/Source/Programs/Horde/Horde.Build/Agents/Fleet/IPoolSizeStrategy.cs
@@ -36,6 +36,26 @@ namespace Horde.Build.Agents.Fleet
 		/// </summary>
 		ComputeQueueAwsMetric
 	}
+	
+	/// <summary>
+	/// Extensions for pool size strategy
+	/// </summary>
+	public static class PoolSizeStrategyExtensions
+	{
+		/// <summary>
+		/// Checks whether the strategy can be called more often than normal (i.e it's cheap to run)
+		/// </summary>
+		/// <param name="strategy"></param>
+		/// <returns>True if it supports high-frequency</returns>
+		public static bool SupportsHighFrequency(this PoolSizeStrategy strategy)
+		{
+			return strategy switch
+			{
+				PoolSizeStrategy.ComputeQueueAwsMetric => true,
+				_ => false
+			};
+		}
+	}
 
 	/// <summary>
 	/// Class for specifying and grouping data together required for calculating pool size
