Horde: Add missing start/stop trigger for new high-frequency ticking of auto-scaling

diff --git a/Engine/Source/Programs/Horde/Horde.Build/Agents/Fleet/AutoscaleServiceV2.cs b/Engine/Source/Programs/Horde/Horde.Build/Agents/Fleet/AutoscaleServiceV2.cs
--- a/Engine/Source/Programs/Horde/Horde.Build/Agents/Fleet/AutoscaleServiceV2.cs
+++ b/Engine/Source/Programs/Horde/Horde.Build/Agents/Fleet/AutoscaleServiceV2.cs
@@ -75,10 +75,16 @@ namespace Horde.Build.Agents.Fleet
 		}
 
 		/// <inheritdoc/>
-		public Task StartAsync(CancellationToken cancellationToken) => _ticker.StartAsync();
+		public Task StartAsync(CancellationToken cancellationToken)
+		{
+			return Task.WhenAll(_ticker.StartAsync(), _tickerHighFrequency.StartAsync());
+		}
 
 		/// <inheritdoc/>
-		public Task StopAsync(CancellationToken cancellationToken) => _ticker.StopAsync();
+		public Task StopAsync(CancellationToken cancellationToken)
+		{
+			return Task.WhenAll(_ticker.StopAsync(), _tickerHighFrequency.StopAsync());
+		}
 
 		/// <inheritdoc/>
 		public void Dispose() => _ticker.Dispose();
