Horde Storage - Error endpoint will no longer produce errors for aborted health checks (as this is handled case within the health checks themselves).
This can happen if you've configured multiple things to probe the same health checks (or the system doing health checks simply overlaps two of them as can happen in kubernetes)

diff --git a/Engine/Source/Programs/Horde/HordeStorage/Jupiter.Common/Controllers/ErrorController.cs b/Engine/Source/Programs/Horde/HordeStorage/Jupiter.Common/Controllers/ErrorController.cs
--- a/Engine/Source/Programs/Horde/HordeStorage/Jupiter.Common/Controllers/ErrorController.cs
+++ b/Engine/Source/Programs/Horde/HordeStorage/Jupiter.Common/Controllers/ErrorController.cs
@@ -1,5 +1,6 @@
 ﻿// Copyright Epic Games, Inc. All Rights Reserved.
 
+using System;
 using Microsoft.AspNetCore.Diagnostics;
 using Microsoft.AspNetCore.Mvc;
 
@@ -19,6 +20,16 @@ namespace Jupiter.Controllers
                 return NoContent();
             }
 
+            // ignore cancelled exceptions from the health checks, that will happen if a health check is started while another is running
+            // the later will contain a valid result
+            if (context.Error is OperationCanceledException)
+            {
+                if (context.Path.StartsWith("/health", StringComparison.Ordinal))
+                {
+                    return Ok();
+                }
+            }
+
             return Problem(
                 detail: context.Error.StackTrace,
                 title: context.Error.Message);
