Horde Storage - Handle clients with limited bandwidth a bit better - essentially we detect cases caused by limited bandwidth when receiving data and throwing a custom exception (resulting in 408 request timed out) as to avoid a 500 issues. Users should just retry this when they have enough bandwidth available.

diff --git a/Engine/Source/Programs/Horde/HordeStorage/Horde.Storage/Controllers/CompressedBlobController.cs b/Engine/Source/Programs/Horde/HordeStorage/Horde.Storage/Controllers/CompressedBlobController.cs
--- a/Engine/Source/Programs/Horde/HordeStorage/Horde.Storage/Controllers/CompressedBlobController.cs
+++ b/Engine/Source/Programs/Horde/HordeStorage/Horde.Storage/Controllers/CompressedBlobController.cs
@@ -5,6 +5,7 @@ using System.Collections.Concurrent;
 using System.Collections.Generic;
 using System.ComponentModel.DataAnnotations;
 using System.Linq;
+using System.Net;
 using System.Net.Mime;
 using System.Threading.Tasks;
 using EpicGames.Horde.Storage;
@@ -13,6 +14,7 @@ using Horde.Storage.Implementation;
 using Jupiter;
 using Jupiter.Common.Implementation;
 using Jupiter.Implementation;
+using Jupiter.Utils;
 using Microsoft.AspNetCore.Authorization;
 using Microsoft.AspNetCore.Mvc;
 using Microsoft.Extensions.Primitives;
@@ -23,7 +25,7 @@ namespace Horde.Storage.Controllers
     using BlobNotFoundException = Horde.Storage.Implementation.BlobNotFoundException;
 
     [ApiController]
-	[Authorize]
+    [Authorize]
     [Route("api/v1/compressed-blobs")]
     public class CompressedBlobController : ControllerBase
     {
@@ -216,24 +218,27 @@ namespace Horde.Storage.Controllers
 
             _diagnosticContext.Set("Content-Length", Request.ContentLength ?? -1);
 
-            using IBufferedPayload payload = await _bufferedPayloadFactory.CreateFromRequest(Request);
-
             try
             {
-                ContentId identifier = await _storage.PutCompressedObject(ns, payload, id, _contentIdStore, _compressedBufferUtils);
+                using IBufferedPayload payload = await _bufferedPayloadFactory.CreateFromRequest(Request);
 
-                return Ok(new
-                {
-                    Identifier = identifier.ToString()
-                });
+                ContentId identifier =
+                    await _storage.PutCompressedObject(ns, payload, id, _contentIdStore, _compressedBufferUtils);
+
+                return Ok(new { Identifier = identifier.ToString() });
             }
             catch (HashMismatchException e)
             {
                 return BadRequest(new ProblemDetails
                 {
-                    Title = $"Incorrect hash, got hash \"{e.SuppliedHash}\" but hash of content was determined to be \"{e.ContentHash}\""
+                    Title =
+                        $"Incorrect hash, got hash \"{e.SuppliedHash}\" but hash of content was determined to be \"{e.ContentHash}\""
                 });
             }
+            catch (ClientSendSlowException e)
+            {
+                return Problem(e.Message, null, (int)HttpStatusCode.RequestTimeout);
+            }
         }
 
         [HttpPost("{ns}")]
@@ -250,10 +255,10 @@ namespace Horde.Storage.Controllers
 
             _diagnosticContext.Set("Content-Length", Request.ContentLength ?? -1);
 
-            using IBufferedPayload payload = await _bufferedPayloadFactory.CreateFromRequest(Request);
-
             try
             {
+                using IBufferedPayload payload = await _bufferedPayloadFactory.CreateFromRequest(Request);
+
                 ContentId identifier = await _storage.PutCompressedObject(ns, payload, null, _contentIdStore, _compressedBufferUtils);
 
                 return Ok(new
@@ -268,6 +273,10 @@ namespace Horde.Storage.Controllers
                     Title = $"Incorrect hash, got hash \"{e.SuppliedHash}\" but hash of content was determined to be \"{e.ContentHash}\""
                 });
             }
+            catch (ClientSendSlowException e)
+            {
+                return Problem(e.Message, null, (int)HttpStatusCode.RequestTimeout);
+            }
         }
 
         /*[HttpDelete("{ns}/{id}")]
@@ -311,7 +320,7 @@ namespace Horde.Storage.Controllers
 
     public class ExistCheckMultipleContentIdResponse
     {
-		[CbField("needs")]
-		public ContentId[] Needs { get; set; } = null!;
+        [CbField("needs")]
+        public ContentId[] Needs { get; set; } = null!;
     }
 }
diff --git a/Engine/Source/Programs/Horde/HordeStorage/Horde.Storage/Controllers/ObjectController.cs b/Engine/Source/Programs/Horde/HordeStorage/Horde.Storage/Controllers/ObjectController.cs
--- a/Engine/Source/Programs/Horde/HordeStorage/Horde.Storage/Controllers/ObjectController.cs
+++ b/Engine/Source/Programs/Horde/HordeStorage/Horde.Storage/Controllers/ObjectController.cs
@@ -5,6 +5,7 @@ using System.Collections.Concurrent;
 using System.Collections.Generic;
 using System.ComponentModel.DataAnnotations;
 using System.Linq;
+using System.Net;
 using System.Net.Mime;
 using System.Threading.Tasks;
 using async_enumerable_dotnet;
@@ -157,10 +158,17 @@ namespace Horde.Storage.Controllers
             }
 
             _diagnosticContext.Set("Content-Length", Request.ContentLength ?? -1);
-            using IBufferedPayload payload = await _bufferedPayloadFactory.CreateFromRequest(Request);
+            try
+            {
+                using IBufferedPayload payload = await _bufferedPayloadFactory.CreateFromRequest(Request);
 
-            BlobIdentifier identifier = await _storage.PutObject(ns, payload, id);
-            return Ok(new PutBlobResponse(identifier));
+                BlobIdentifier identifier = await _storage.PutObject(ns, payload, id);
+                return Ok(new PutBlobResponse(identifier));
+            }
+            catch (ClientSendSlowException e)
+            {
+                return Problem(e.Message, null, (int)HttpStatusCode.RequestTimeout);
+            }
         }
 
         [HttpGet("{ns}/{id}/references")]
diff --git a/Engine/Source/Programs/Horde/HordeStorage/Horde.Storage/Controllers/ReferencesController.cs b/Engine/Source/Programs/Horde/HordeStorage/Horde.Storage/Controllers/ReferencesController.cs
--- a/Engine/Source/Programs/Horde/HordeStorage/Horde.Storage/Controllers/ReferencesController.cs
+++ b/Engine/Source/Programs/Horde/HordeStorage/Horde.Storage/Controllers/ReferencesController.cs
@@ -497,7 +497,7 @@ namespace Horde.Storage.Controllers
             {
                 return NotFound(new ProblemDetails { Title = $"Blob {e.Blob} in namespace {ns} did not exist" });
             }
-			catch (ObjectNotFoundException e)
+            catch (ObjectNotFoundException e)
             {
                 return NotFound(new ProblemDetails {Title = $"Object {e.Bucket} {e.Key} in namespace {e.Namespace} did not exist"});
             }
@@ -595,26 +595,29 @@ namespace Horde.Storage.Controllers
             }
 
             _diagnosticContext.Set("Content-Length", Request.ContentLength ?? -1);
-            
-            using IBufferedPayload payload = await _bufferedPayloadFactory.CreateFromRequest(Request);
-
-            BlobIdentifier headerHash;
-            if (Request.Headers.ContainsKey(CommonHeaders.HashHeaderName))
-            {
-                headerHash = new BlobIdentifier(Request.Headers[CommonHeaders.HashHeaderName]);
-            }
-            else
-            {
-                return BadRequest(new ProblemDetails
-                {
-                    Title = $"Missing expected header {CommonHeaders.HashHeaderName}"
-                });
-            }
 
             CbObject payloadObject;
-            BlobIdentifier blobHeader = headerHash;
+            BlobIdentifier blobHeader;
+
             try
             {
+                using IBufferedPayload payload = await _bufferedPayloadFactory.CreateFromRequest(Request);
+
+                BlobIdentifier headerHash;
+                if (Request.Headers.ContainsKey(CommonHeaders.HashHeaderName))
+                {
+                    headerHash = new BlobIdentifier(Request.Headers[CommonHeaders.HashHeaderName]);
+                }
+                else
+                {
+                    return BadRequest(new ProblemDetails
+                    {
+                        Title = $"Missing expected header {CommonHeaders.HashHeaderName}"
+                    });
+                }
+
+                blobHeader = headerHash;
+
                 switch (Request.ContentType)
                 {
                     case MediaTypeNames.Application.Json:
@@ -668,6 +671,10 @@ namespace Horde.Storage.Controllers
                     Title = $"Incorrect hash, got hash \"{e.SuppliedHash}\" but hash of content was determined to be \"{e.ContentHash}\""
                 });
             }
+            catch (ClientSendSlowException e)
+            {
+                return Problem(e.Message, null, (int)HttpStatusCode.RequestTimeout);
+            }
 
             (ContentId[] missingReferences, BlobIdentifier[] missingBlobs) = await _objectService.Put(ns, bucket, key, blobHeader, payloadObject);
 
diff --git a/Engine/Source/Programs/Horde/HordeStorage/Horde.Storage/Controllers/StorageController.cs b/Engine/Source/Programs/Horde/HordeStorage/Horde.Storage/Controllers/StorageController.cs
--- a/Engine/Source/Programs/Horde/HordeStorage/Horde.Storage/Controllers/StorageController.cs
+++ b/Engine/Source/Programs/Horde/HordeStorage/Horde.Storage/Controllers/StorageController.cs
@@ -6,6 +6,7 @@ using System.Collections.Generic;
 using System.ComponentModel.DataAnnotations;
 using System.IO;
 using System.Linq;
+using System.Net;
 using System.Net.Mime;
 using System.Threading.Tasks;
 using EpicGames.Horde.Storage;
@@ -14,6 +15,7 @@ using Horde.Storage.Implementation;
 using Jupiter;
 using Jupiter.Common.Implementation;
 using Jupiter.Implementation;
+using Jupiter.Utils;
 using Microsoft.AspNetCore.Authorization;
 using Microsoft.AspNetCore.Http;
 using Microsoft.AspNetCore.Mvc;
@@ -170,10 +172,11 @@ namespace Horde.Storage.Controllers
             }
 
             _diagnosticContext.Set("Content-Length", Request.ContentLength ?? -1);
-            using IBufferedPayload payload = await _bufferedPayloadFactory.CreateFromRequest(Request);
 
             try
             {
+                using IBufferedPayload payload = await _bufferedPayloadFactory.CreateFromRequest(Request);
+
                 BlobIdentifier identifier = await _storage.PutObject(ns, payload, id);
                 return Ok(new
                 {
@@ -184,6 +187,10 @@ namespace Horde.Storage.Controllers
             {
                 return StatusCode(StatusCodes.Status429TooManyRequests);
             }
+            catch (ClientSendSlowException e)
+            {
+                return Problem(e.Message, null, (int)HttpStatusCode.RequestTimeout);
+            }
         }
 
         [HttpPost("{ns}")]
@@ -199,17 +206,24 @@ namespace Horde.Storage.Controllers
             }
 
             _diagnosticContext.Set("Content-Length", Request.ContentLength ?? -1);
-            using IBufferedPayload payload = await _bufferedPayloadFactory.CreateFromRequest(Request);
-            
-            await using Stream stream = payload.GetStream();
+            try
+            {
+                using IBufferedPayload payload = await _bufferedPayloadFactory.CreateFromRequest(Request);
 
-            BlobIdentifier id = await BlobIdentifier.FromStream(stream);
-            await _storage.PutObjectKnownHash(ns, payload, id);
+                await using Stream stream = payload.GetStream();
 
-            return Ok(new
+                BlobIdentifier id = await BlobIdentifier.FromStream(stream);
+                await _storage.PutObjectKnownHash(ns, payload, id);
+                
+                return Ok(new
+                {
+                    Identifier = id.ToString()
+                });
+            }
+            catch (ClientSendSlowException e)
             {
-                Identifier = id.ToString()
-            });
+                return Problem(e.Message, null, (int)HttpStatusCode.RequestTimeout);
+            }
         }
 
         [HttpDelete("{ns}/{id}")]
diff --git a/Engine/Source/Programs/Horde/HordeStorage/Jupiter.Common/Utils/StreamUtils.cs b/Engine/Source/Programs/Horde/HordeStorage/Jupiter.Common/Utils/StreamUtils.cs
--- a/Engine/Source/Programs/Horde/HordeStorage/Jupiter.Common/Utils/StreamUtils.cs
+++ b/Engine/Source/Programs/Horde/HordeStorage/Jupiter.Common/Utils/StreamUtils.cs
@@ -1,7 +1,9 @@
 ﻿// Copyright Epic Games, Inc. All Rights Reserved.
 
+using System;
 using System.IO;
 using System.Threading.Tasks;
+using Microsoft.AspNetCore.Http;
 
 namespace Jupiter.Utils
 {
@@ -9,9 +11,38 @@ namespace Jupiter.Utils
     {
         public static async Task<byte[]> ToByteArray(this Stream s)
         {
-            await using MemoryStream ms = new MemoryStream();
-            await s.CopyToAsync(ms);
-            return ms.ToArray();
+            try
+            {
+                await using MemoryStream ms = new MemoryStream();
+                await s.CopyToAsync(ms);
+                return ms.ToArray();
+            }
+            catch (BadHttpRequestException e)
+            {
+                // if a user is sending data to slowly it will be terminated by asp.net and a bad http request exception is thrown
+                // this type of exception can be thrown for multiple reasons so we check for the error message about MinRequestBodyDataRate
+                // that setting can be set to reduce the number of bytes expected to be sent by a user - but that is likely a bad idea
+                // as the most common reason for this happening is a client trying to send more data then what their uplink has bandwidth for.
+                if (e.Message.Contains("MinRequestBodyDataRate", StringComparison.Ordinal))
+                {
+                    throw new ClientSendSlowException(e);
+                }
+
+                // this exception also indicates a send issue, in this case the client sent fewer bytes then what it set as a content length
+                // most likely caused by them shutting down their connection before finishing the entire send
+                if (e.Message.Contains("Unexpected end of request content", StringComparison.Ordinal))
+                {
+                    throw new ClientSendSlowException(e);
+                }
+                throw;
+            }
+        }
+    }
+
+    public class ClientSendSlowException : Exception
+    {
+        public ClientSendSlowException(Exception? innerException) : base("Client was sending data to slow", innerException)
+        {
         }
     }
 }
