Horde: Optimize device pool telemetry storage and responses, added paging index/counts to controller/collection
#rnx

diff --git a/Engine/Source/Programs/Horde/Horde.Build.Tests/DeviceServiceTests.cs b/Engine/Source/Programs/Horde/Horde.Build.Tests/DeviceServiceTests.cs
--- a/Engine/Source/Programs/Horde/Horde.Build.Tests/DeviceServiceTests.cs
+++ b/Engine/Source/Programs/Horde/Horde.Build.Tests/DeviceServiceTests.cs
@@ -219,7 +219,7 @@ namespace Horde.Build.Tests
 			{
 				if (t.PlatformId == "testdeviceplatform1")
 				{
-					Assert.AreEqual(0, t.Available);
+					Assert.IsNull(t.Available);
 					Assert.AreEqual(1, t.Reserved);
 					Assert.AreEqual(1, t.Maintenance);
 
diff --git a/Engine/Source/Programs/Horde/Horde.Build/Devices/DeviceCollection.cs b/Engine/Source/Programs/Horde/Horde.Build/Devices/DeviceCollection.cs
--- a/Engine/Source/Programs/Horde/Horde.Build/Devices/DeviceCollection.cs
+++ b/Engine/Source/Programs/Horde/Horde.Build/Devices/DeviceCollection.cs
@@ -254,61 +254,61 @@ namespace Horde.Build.Devices
 			/// <summary>
 			/// Id of telemetry document
 			/// </summary>
-			[BsonRequired, BsonId]
+			[BsonRequired, BsonId, BsonElement("tid")]
 			public ObjectId TelemetryId { get; set; }
 
 			/// <summary>
 			/// The device id
 			/// </summary>
-			[BsonRequired]
+			[BsonRequired, BsonElement("did")]
 			public DeviceId DeviceId { get; set; }
 
 			/// <summary>
 			/// The time this telemetry data was created
 			/// </summary>
-			[BsonRequired]
+			[BsonRequired, BsonElement("c")]
 			public DateTime CreateTimeUtc { get; set; }
 
 			/// <summary>
 			/// The stream id which utilized device
 			/// </summary>
-			[BsonIgnoreIfNull]
+			[BsonIgnoreIfNull, BsonElement("sid")]
 			public string? StreamId { get; set; }
 
 			/// <summary>
 			/// The job id which utilized device
 			/// </summary>
-			[BsonIgnoreIfNull]
+			[BsonIgnoreIfNull,BsonElement("job")]
 			public string? JobId { get; set; }
 
 			/// <summary>
 			/// The job's step id
 			/// </summary>
-			[BsonIgnoreIfNull]
+			[BsonIgnoreIfNull, BsonElement("step")]
 			public string? StepId { get; set; }
 
 			/// <summary>
 			/// Reservation Id (transient, reservations are deleted upon expiration)
 			/// </summary>
-			[BsonIgnoreIfNull]
+			[BsonIgnoreIfNull, BsonElement("rid")]
 			public ObjectId? ReservationId { get; set; }
 
 			/// <summary>
 			/// The time device was reserved
 			/// </summary>
-			[BsonIgnoreIfNull]
+			[BsonIgnoreIfNull, BsonElement("rs")]
 			public DateTime? ReservationStartUtc { get; set; }
 
 			/// <summary>
 			/// The time device was freed
 			/// </summary>
-			[BsonIgnoreIfNull]
+			[BsonIgnoreIfNull, BsonElement("rf")]
 			public DateTime? ReservationFinishUtc { get; set; }
 
 			/// <summary>
 			/// If the device reported a problem
 			/// </summary>
-			[BsonIgnoreIfNull]
+			[BsonIgnoreIfNull, BsonElement("p")]
 			public DateTime? ProblemTimeUtc { get; set; }
 
 			[BsonConstructor]
@@ -347,49 +347,49 @@ namespace Horde.Build.Devices
 			/// <summary>
 			/// Id of telemetry document
 			/// </summary>
-			[BsonRequired, BsonId]
+			[BsonRequired, BsonId, BsonElement("tid")]
 			public ObjectId TelemetryId { get; set; }
 
 			/// <summary>
 			/// The time this telemetry data was created
 			/// </summary>
-			[BsonRequired]
+			[BsonRequired, BsonElement("c")]
 			public DateTime CreateTimeUtc { get; set; }
 
 			/// <summary>
 			/// The platform id 
 			/// </summary>
-			[BsonRequired]
+			[BsonRequired, BsonElement("pid")]
 			public DevicePlatformId PlatformId { get; set; }
 
 			/// <summary>
 			/// Number of available devices of this platform 
 			/// </summary>
-			[BsonRequired]
+			[BsonIgnoreIfDefault, BsonDefaultValue(0), BsonElement("a")]
 			public int Available { get; set; } = 0;
 
 			/// <summary>
 			/// Number of reserved devices of this platform 
 			/// </summary>
-			[BsonRequired]
+			[BsonIgnoreIfDefault, BsonDefaultValue(0), BsonElement("r")]
 			public int Reserved { get; set; } = 0;
 
 			/// <summary>
 			/// Number of devices in maintenance state
 			/// </summary>
-			[BsonRequired]
+			[BsonIgnoreIfDefault, BsonDefaultValue(0), BsonElement("m")]
 			public int Maintenance { get; set; } = 0;
 
 			/// <summary>
 			/// Number of devices in problem state
-			/// </summary>
-			[BsonRequired]
+			/// </summary>			
+			[BsonIgnoreIfDefault, BsonDefaultValue(0), BsonElement("p")]
 			public int Problem { get; set; } = 0;
 
 			/// <summary>
 			/// Number of devices that are disabled
 			/// </summary>
-			[BsonRequired]
+			[BsonIgnoreIfDefault, BsonDefaultValue(0), BsonElement("d")]
 			public int Disabled { get; set; } = 0;
 
 			/// <summary>
@@ -484,11 +484,11 @@ namespace Horde.Build.Devices
 
 			List<MongoIndex<DeviceTelemetryDocument>> deviceTelemetryIndexes = new List<MongoIndex<DeviceTelemetryDocument>>();
 			deviceTelemetryIndexes.Add((keys => keys.Descending(x => x.CreateTimeUtc)));
-			_deviceTelemetry = mongoService.GetCollection<DeviceTelemetryDocument>("Devices.DeviceTelemetry", deviceTelemetryIndexes);
+			_deviceTelemetry = mongoService.GetCollection<DeviceTelemetryDocument>("Devices.DeviceTelemetryV2", deviceTelemetryIndexes);
 
 			List<MongoIndex<DevicePoolTelemetryDocument>> poolTelemetryIndexes = new List<MongoIndex<DevicePoolTelemetryDocument>>();
 			poolTelemetryIndexes.Add((keys => keys.Descending(x => x.CreateTimeUtc)));
-			_poolTelemetry = mongoService.GetCollection<DevicePoolTelemetryDocument>("Devices.PoolTelemetry", poolTelemetryIndexes);
+			_poolTelemetry = mongoService.GetCollection<DevicePoolTelemetryDocument>("Devices.PoolTelemetryV2", poolTelemetryIndexes);
 		}
 
 		/// <inheritdoc/>
@@ -1051,7 +1051,7 @@ namespace Horde.Build.Devices
 		}
 
 		/// <inheritdoc/>
-		public async Task<List<IDeviceTelemetry>> FindDeviceTelemetryAsync(DeviceId[]? deviceIds = null, DateTimeOffset? minCreateTime = null, DateTimeOffset? maxCreateTime = null)
+		public async Task<List<IDeviceTelemetry>> FindDeviceTelemetryAsync(DeviceId[]? deviceIds = null, DateTimeOffset? minCreateTime = null, DateTimeOffset? maxCreateTime = null, int? index = null, int? count = null)
 		{
 			FilterDefinitionBuilder<DeviceTelemetryDocument> filterBuilder = Builders<DeviceTelemetryDocument>.Filter;
 			FilterDefinition<DeviceTelemetryDocument> filter = filterBuilder.Empty;
@@ -1071,7 +1071,7 @@ namespace Horde.Build.Devices
 				filter &= filterBuilder.Lte(x => x.CreateTimeUtc!, maxCreateTime.Value.UtcDateTime);
 			}
 
-			List<DeviceTelemetryDocument> results = await _deviceTelemetry.Find(filter).ToListAsync();	
+			List<DeviceTelemetryDocument> results = await _deviceTelemetry.Find(filter).Range(index, count).ToListAsync();	
 			return results.ConvertAll<IDeviceTelemetry>(x => x);
 		}
 
@@ -1190,7 +1190,7 @@ namespace Horde.Build.Devices
 		}
 
 		/// <inheritdoc/>
-		public async Task<List<IDevicePoolTelemetry>> FindPoolTelemetryAsync(DateTimeOffset? minCreateTime = null, DateTimeOffset? maxCreateTime = null)
+		public async Task<List<IDevicePoolTelemetry>> FindPoolTelemetryAsync(DateTimeOffset? minCreateTime = null, DateTimeOffset? maxCreateTime = null, int? index = null, int? count = null)
 		{
 			FilterDefinitionBuilder<DevicePoolTelemetryDocument> filterBuilder = Builders<DevicePoolTelemetryDocument>.Filter;
 			FilterDefinition<DevicePoolTelemetryDocument> filter = filterBuilder.Empty;
@@ -1205,7 +1205,7 @@ namespace Horde.Build.Devices
 				filter &= filterBuilder.Lte(x => x.CreateTimeUtc!, maxCreateTime.Value.UtcDateTime);
 			}
 
-			List<DevicePoolTelemetryDocument> results = await _poolTelemetry.Find(filter).ToListAsync();
+			List<DevicePoolTelemetryDocument> results = await _poolTelemetry.Find(filter).Range(index, count).ToListAsync();
 			return results.ConvertAll<IDevicePoolTelemetry>(x => x);
 		}
 
diff --git a/Engine/Source/Programs/Horde/Horde.Build/Devices/DeviceMessages.cs b/Engine/Source/Programs/Horde/Horde.Build/Devices/DeviceMessages.cs
--- a/Engine/Source/Programs/Horde/Horde.Build/Devices/DeviceMessages.cs
+++ b/Engine/Source/Programs/Horde/Horde.Build/Devices/DeviceMessages.cs
@@ -647,27 +647,27 @@ namespace Horde.Build.Devices
 		/// <summary>
 		/// Number of available devices of this platform 
 		/// </summary>
-		public int Available { get; set; }
+		public int? Available { get; set; }
 
 		/// <summary>
 		/// Number of reserved devices of this platform 
 		/// </summary>
-		public int Reserved { get; set; }
+		public int? Reserved { get; set; }
 
 		/// <summary>
 		/// Number of devices in maintenance state
 		/// </summary>
-		public int Maintenance { get; set; }
+		public int? Maintenance { get; set; }
 
 		/// <summary>
 		/// Number of devices in problem state
 		/// </summary>
-		public int Problem { get; set; }
+		public int? Problem { get; set; }
 
 		/// <summary>
 		/// Number of devices in disabled state
 		/// </summary>
-		public int Disabled { get; set; }
+		public int? Disabled { get; set; }
 
 		/// <summary>
 		/// StreamId to reserved device ids
@@ -680,11 +680,11 @@ namespace Horde.Build.Devices
 		public GetDevicePlatformTelemetryResponse(string platformId, int available, int reserved, int maintenance, int problem, int disabled, Dictionary<string, List<string>> streamDevices)
 		{
 			PlatformId = platformId;
-			Available = available;
-			Reserved = reserved;
-			Maintenance = maintenance;
-			Problem = problem;
-			Disabled = disabled;
+			Available = available == 0 ? null : available;
+			Reserved = reserved == 0 ? null : reserved;
+			Maintenance = maintenance == 0 ? null : maintenance;
+			Problem = problem == 0 ? null : problem;
+			Disabled = disabled == 0 ? null : disabled;
 			StreamDevices = streamDevices;
 		}
 	}
diff --git a/Engine/Source/Programs/Horde/Horde.Build/Devices/DeviceService.cs b/Engine/Source/Programs/Horde/Horde.Build/Devices/DeviceService.cs
--- a/Engine/Source/Programs/Horde/Horde.Build/Devices/DeviceService.cs
+++ b/Engine/Source/Programs/Horde/Horde.Build/Devices/DeviceService.cs
@@ -323,17 +323,17 @@ namespace Horde.Build.Devices
 		/// <summary>
 		/// Get device telemetry
 		/// </summary>
-		public Task<List<IDeviceTelemetry>> GetDeviceTelemetryAsync(DeviceId[]? deviceIds = null, DateTimeOffset? minCreateTime=null, DateTimeOffset? maxCreateTime=null)
+		public Task<List<IDeviceTelemetry>> GetDeviceTelemetryAsync(DeviceId[]? deviceIds = null, DateTimeOffset? minCreateTime=null, DateTimeOffset? maxCreateTime=null, int? index = null, int? count = null)
 		{
-			return _devices.FindDeviceTelemetryAsync(deviceIds, minCreateTime, maxCreateTime);
+			return _devices.FindDeviceTelemetryAsync(deviceIds, minCreateTime, maxCreateTime, index, count);
 		}
 
 		/// <summary>
 		/// Get device pool telemetry
 		/// </summary>
-		public Task<List<IDevicePoolTelemetry>> GetDevicePoolTelemetryAsync(DateTimeOffset? minCreateTime = null, DateTimeOffset? maxCreateTime = null)
+		public Task<List<IDevicePoolTelemetry>> GetDevicePoolTelemetryAsync(DateTimeOffset? minCreateTime = null, DateTimeOffset? maxCreateTime = null, int? index = null, int? count = null)
 		{
-			return _devices.FindPoolTelemetryAsync(minCreateTime, maxCreateTime);
+			return _devices.FindPoolTelemetryAsync(minCreateTime, maxCreateTime, index, count);
 		}
 
 
diff --git a/Engine/Source/Programs/Horde/Horde.Build/Devices/DevicesController.cs b/Engine/Source/Programs/Horde/Horde.Build/Devices/DevicesController.cs
--- a/Engine/Source/Programs/Horde/Horde.Build/Devices/DevicesController.cs
+++ b/Engine/Source/Programs/Horde/Horde.Build/Devices/DevicesController.cs
@@ -498,7 +498,9 @@ namespace Horde.Build.Devices
 		[ProducesResponseType(typeof(List<GetDevicePoolTelemetryResponse>), 200)]
 		public async Task<ActionResult<List<GetDevicePoolTelemetryResponse>>> GetDevicePoolTelemetry(
 			[FromQuery] DateTimeOffset? minCreateTime = null,
-			[FromQuery] DateTimeOffset? maxCreateTime = null)
+			[FromQuery] DateTimeOffset? maxCreateTime = null,
+			[FromQuery] int index = 0,
+			[FromQuery] int count = 1024)
 		{
 			if (!await _deviceService.AuthorizeAsync(AclAction.DeviceRead, User))
 			{
@@ -510,7 +512,7 @@ namespace Horde.Build.Devices
 				minCreateTime = new DateTimeOffset(DateTime.UtcNow - TimeSpan.FromHours(24));
 			}
 
-			List<IDevicePoolTelemetry> telemetryData = await _deviceService.GetDevicePoolTelemetryAsync(minCreateTime, maxCreateTime);
+			List<IDevicePoolTelemetry> telemetryData = await _deviceService.GetDevicePoolTelemetryAsync(minCreateTime, maxCreateTime, index, count);
 
 			List<GetDevicePoolTelemetryResponse> response = new List<GetDevicePoolTelemetryResponse>();
 						
@@ -730,7 +732,9 @@ namespace Horde.Build.Devices
 			[FromQuery] string? poolId = null,
 			[FromQuery] string? platformId = null,
 			[FromQuery] DateTimeOffset? minCreateTime = null,
-			[FromQuery] DateTimeOffset? maxCreateTime = null)
+			[FromQuery] DateTimeOffset? maxCreateTime = null,
+			[FromQuery] int index = 0,
+			[FromQuery] int count = 1024)
 		{
 			if (!await _deviceService.AuthorizeAsync(AclAction.DeviceRead, User))
 			{
@@ -764,7 +768,7 @@ namespace Horde.Build.Devices
 				return response;
 			}
 			
-			List<IDeviceTelemetry> telemetryData = await _deviceService.GetDeviceTelemetryAsync(devices.Select(x => x.Id).ToArray(), minCreateTime, maxCreateTime);
+			List<IDeviceTelemetry> telemetryData = await _deviceService.GetDeviceTelemetryAsync(devices.Select(x => x.Id).ToArray(), minCreateTime, maxCreateTime, index, count);
 
 			Dictionary<string, List<GetTelemetryInfoResponse>> results = new Dictionary<string, List<GetTelemetryInfoResponse>>();
 
diff --git a/Engine/Source/Programs/Horde/Horde.Build/Devices/IDeviceCollection.cs b/Engine/Source/Programs/Horde/Horde.Build/Devices/IDeviceCollection.cs
--- a/Engine/Source/Programs/Horde/Horde.Build/Devices/IDeviceCollection.cs
+++ b/Engine/Source/Programs/Horde/Horde.Build/Devices/IDeviceCollection.cs
@@ -238,8 +238,10 @@ namespace Horde.Build.Devices
 		/// <param name="deviceIds"></param>
 		/// <param name="minCreateTime"></param>
 		/// <param name="maxCreateTime"></param>
+		/// <param name="index"></param>
+		/// <param name="count"></param>
 		/// <returns></returns>
-		public Task<List<IDeviceTelemetry>> FindDeviceTelemetryAsync(DeviceId[]? deviceIds = null, DateTimeOffset? minCreateTime = null, DateTimeOffset? maxCreateTime = null);
+		public Task<List<IDeviceTelemetry>> FindDeviceTelemetryAsync(DeviceId[]? deviceIds = null, DateTimeOffset? minCreateTime = null, DateTimeOffset? maxCreateTime = null, int? index = null, int? count = null);
 
 		/// <summary>
 		/// Creates a device pool telemetry snapshot
@@ -252,8 +254,10 @@ namespace Horde.Build.Devices
 		/// </summary>
 		/// <param name="minCreateTime"></param>
 		/// <param name="maxCreateTime"></param>
+		/// <param name="index"></param>
+		/// <param name="count"></param>
 		/// <returns></returns>
-		public Task<List<IDevicePoolTelemetry>> FindPoolTelemetryAsync(DateTimeOffset? minCreateTime = null, DateTimeOffset? maxCreateTime = null);
+		public Task<List<IDevicePoolTelemetry>> FindPoolTelemetryAsync(DateTimeOffset? minCreateTime = null, DateTimeOffset? maxCreateTime = null, int? index = null, int? count = null);
 
 
 	}
