Horde: Also make stream devices nullable for responses
#rnx

diff --git a/Engine/Source/Programs/Horde/Horde.Build/Devices/DeviceMessages.cs b/Engine/Source/Programs/Horde/Horde.Build/Devices/DeviceMessages.cs
--- a/Engine/Source/Programs/Horde/Horde.Build/Devices/DeviceMessages.cs
+++ b/Engine/Source/Programs/Horde/Horde.Build/Devices/DeviceMessages.cs
@@ -672,7 +672,7 @@ namespace Horde.Build.Devices
 		/// <summary>
 		/// StreamId to reserved device ids
 		/// </summary>
-		public Dictionary<string, List<string>> StreamDevices { get; set; } = null!;
+		public Dictionary<string, List<string>>? StreamDevices { get; set; }
 
 		/// <summary>
 		/// Constructor
@@ -685,7 +685,7 @@ namespace Horde.Build.Devices
 			Maintenance = maintenance == 0 ? null : maintenance;
 			Problem = problem == 0 ? null : problem;
 			Disabled = disabled == 0 ? null : disabled;
-			StreamDevices = streamDevices;
+			StreamDevices = streamDevices.Count == 0 ? null : streamDevices;
 		}
 	}
 
