Horde: Fix build warning.

diff --git a/Engine/Source/Programs/Horde/Horde.Build.Tests/DeviceServiceTests.cs b/Engine/Source/Programs/Horde/Horde.Build.Tests/DeviceServiceTests.cs
--- a/Engine/Source/Programs/Horde/Horde.Build.Tests/DeviceServiceTests.cs
+++ b/Engine/Source/Programs/Horde/Horde.Build.Tests/DeviceServiceTests.cs
@@ -223,7 +223,7 @@ namespace Horde.Build.Tests
 					Assert.AreEqual(1, t.Reserved);
 					Assert.AreEqual(1, t.Maintenance);
 
-					Assert.IsTrue(t.StreamDevices.ContainsKey("ue5-main"));
+					Assert.IsTrue(t.StreamDevices!.ContainsKey("ue5-main"));
 					Assert.AreEqual(t.StreamDevices["ue5-main"].Count, 1);
 					Assert.AreEqual(t.StreamDevices["ue5-main"][0], reservedDevice.Id);
 				}
