Horde: Moving device pool telemetry to dedicated ticker
#rnx

diff --git a/Engine/Source/Programs/Horde/Horde.Build/Devices/DeviceService.cs b/Engine/Source/Programs/Horde/Horde.Build/Devices/DeviceService.cs
--- a/Engine/Source/Programs/Horde/Horde.Build/Devices/DeviceService.cs
+++ b/Engine/Source/Programs/Horde/Horde.Build/Devices/DeviceService.cs
@@ -133,9 +133,8 @@ namespace Horde.Build.Devices
 		/// Device collection
 		/// </summary>
 		readonly IDeviceCollection _devices;
-		readonly ITicker _ticker;		
-		const int _poolTelemetryMinutes = 10;
-		int _poolTelemetryTick = _poolTelemetryMinutes;
+		readonly ITicker _ticker;
+		readonly ITicker _telemetryTicker;
 
 		/// <summary>
 		/// Platform map V1 singleton
@@ -149,14 +148,15 @@ namespace Horde.Build.Devices
 		{
 			UserCollection = userCollection;
 			_devices = devices;
-            _jobService = jobService;
+			_jobService = jobService;
 			_projectService = projectService;
-            _streamService = streamService;
-            _aclService = aclService;
-            _notificationService = notificationService;
+			_streamService = streamService;
+			_aclService = aclService;
+			_notificationService = notificationService;
 			_ticker = clock.AddSharedTicker<DeviceService>(TimeSpan.FromMinutes(1.0), TickAsync, logger);
-            _logger = logger;
-			
+			_telemetryTicker = clock.AddSharedTicker("DeviceService.Telemetry", TimeSpan.FromMinutes(10.0), TickTelemetryAsync, logger);
+			_logger = logger;
+
 			_platformMapSingleton = platformMapSingleton;
 
 		}
@@ -168,7 +168,23 @@ namespace Horde.Build.Devices
 		public Task StopAsync(CancellationToken cancellationToken) => _ticker.StopAsync();
 
 		/// <inheritdoc/>
-		public void Dispose() => _ticker.Dispose();
+		public void Dispose() 
+		{ 
+			_ticker.Dispose(); 
+			_telemetryTicker.Dispose(); 
+		}
+
+		/// <summary>
+		/// Ticks service
+		/// </summary>
+		async ValueTask TickTelemetryAsync(CancellationToken stoppingToken)
+		{
+			if (!stoppingToken.IsCancellationRequested)
+			{
+				_logger.LogInformation("Updating pool telemetry");
+				await _devices.CreatePoolTelemetrySnapshot();
+			}
+		}
 
 		/// <summary>
 		/// Ticks service
@@ -218,33 +234,13 @@ namespace Horde.Build.Devices
 				{
 					_logger.LogError(ex, "Exception while expiring device checkouts: {Message}", ex.Message);
 				}
-
-				// update pool telemetry
-				try
-				{
-
-					if (_poolTelemetryTick >= _poolTelemetryMinutes)
-					{
-						_poolTelemetryTick = 0;
-						_logger.LogInformation("Updating pool telemetry");
-						await _devices.CreatePoolTelemetrySnapshot();						
-					}
-					else
-					{
-						_poolTelemetryTick++;
-					}
-
-				}
-				catch (Exception ex)
-				{
-					_logger.LogError(ex, "Exception while generating pool telemetry: {Message}", ex.Message);
-				}
 			}
 		}
 
 		internal async Task TickForTestingAsync()
 		{
 			await TickAsync(CancellationToken.None);
+			await TickTelemetryAsync(CancellationToken.None);
 		}
 
 
