Horde: Add global tracing span to device pool telemetry method
#rnx

diff --git a/Engine/Source/Programs/Horde/Horde.Build/Devices/DeviceService.cs b/Engine/Source/Programs/Horde/Horde.Build/Devices/DeviceService.cs
--- a/Engine/Source/Programs/Horde/Horde.Build/Devices/DeviceService.cs
+++ b/Engine/Source/Programs/Horde/Horde.Build/Devices/DeviceService.cs
@@ -154,7 +154,7 @@ namespace Horde.Build.Devices
 			_aclService = aclService;
 			_notificationService = notificationService;
 			_ticker = clock.AddSharedTicker<DeviceService>(TimeSpan.FromMinutes(1.0), TickAsync, logger);
-			_telemetryTicker = clock.AddSharedTicker("DeviceService.Telemetry", TimeSpan.FromMinutes(10.0), TickTelemetryAsync, logger);
+			_telemetryTicker = clock.AddSharedTicker("DeviceServiceTelemetry", TimeSpan.FromMinutes(10.0), TickTelemetryAsync, logger);
 			_logger = logger;
 
 			_platformMapSingleton = platformMapSingleton;
@@ -181,6 +181,8 @@ namespace Horde.Build.Devices
 		{
 			if (!stoppingToken.IsCancellationRequested)
 			{
+				using IScope scope = GlobalTracer.Instance.BuildSpan("DeviceService.TickTelemetryAsync").StartActive();
+
 				_logger.LogInformation("Updating pool telemetry");
 				await _devices.CreatePoolTelemetrySnapshot();
 			}
