Horde: Actually start the pool telemetry ticker
#rnx

diff --git a/Engine/Source/Programs/Horde/Horde.Build/Devices/DeviceService.cs b/Engine/Source/Programs/Horde/Horde.Build/Devices/DeviceService.cs
--- a/Engine/Source/Programs/Horde/Horde.Build/Devices/DeviceService.cs
+++ b/Engine/Source/Programs/Horde/Horde.Build/Devices/DeviceService.cs
@@ -154,7 +154,7 @@ namespace Horde.Build.Devices
 			_aclService = aclService;
 			_notificationService = notificationService;
 			_ticker = clock.AddSharedTicker<DeviceService>(TimeSpan.FromMinutes(1.0), TickAsync, logger);
-			_telemetryTicker = clock.AddSharedTicker("DeviceServiceTelemetry", TimeSpan.FromMinutes(10.0), TickTelemetryAsync, logger);
+			_telemetryTicker = clock.AddSharedTicker("DeviceService.Telemetry", TimeSpan.FromMinutes(10.0), TickTelemetryAsync, logger);
 			_logger = logger;
 
 			_platformMapSingleton = platformMapSingleton;
@@ -162,10 +162,19 @@ namespace Horde.Build.Devices
 		}
 
 		/// <inheritdoc/>
-		public Task StartAsync(CancellationToken cancellationToken) => _ticker.StartAsync();
+		public async Task StartAsync(CancellationToken cancellationToken)
+		{
+			await _ticker.StartAsync();
+			await _telemetryTicker.StartAsync();
+		}
+		
 
 		/// <inheritdoc/>
-		public Task StopAsync(CancellationToken cancellationToken) => _ticker.StopAsync();
+		public async Task StopAsync(CancellationToken cancellationToken)
+		{
+			await _ticker.StopAsync();
+			await _telemetryTicker.StopAsync();
+		}
 
 		/// <inheritdoc/>
 		public void Dispose() 
