Horde: Allow specifying an arbitrary condition (using the standard Horde condition syntax) for whether to use a Perforce server cluster. Also fix an issue with duplicate server keys causing exceptions.

diff --git a/Engine/Source/Programs/Horde/Horde.Build/Perforce/PerforceLoadBalancer.cs b/Engine/Source/Programs/Horde/Horde.Build/Perforce/PerforceLoadBalancer.cs
--- a/Engine/Source/Programs/Horde/Horde.Build/Perforce/PerforceLoadBalancer.cs
+++ b/Engine/Source/Programs/Horde/Horde.Build/Perforce/PerforceLoadBalancer.cs
@@ -9,6 +9,7 @@ using System.Net.Http;
 using System.Text.Json;
 using System.Threading;
 using System.Threading.Tasks;
+using EpicGames.Horde.Common;
 using Google.Protobuf.WellKnownTypes;
 using Horde.Build.Agents;
 using Horde.Build.Agents.Leases;
@@ -240,6 +241,28 @@ namespace Horde.Build.Perforce
 			return SelectServerAsync(cluster, agent.Properties);
 		}
 
+		/// <summary>
+		/// Evaluate a server condition using standard Horde syntax
+		/// </summary>
+		static bool EvaluateCondition(Condition? condition, IReadOnlyList<string> properties)
+		{
+			return condition != null && !condition.IsEmpty() && condition.Evaluate(x => GetPropertyValues(properties, x));
+		}
+
+		/// <summary>
+		/// Get property values from a list starting with the given key
+		/// </summary>
+		static IEnumerable<string> GetPropertyValues(IReadOnlyList<string> properties, string name)
+		{
+			foreach (string property in properties)
+			{
+				if (property.Length > name.Length && property[name.Length] == '=' && property.StartsWith(name, StringComparison.OrdinalIgnoreCase))
+				{
+					yield return property.Substring(name.Length + 1);
+				}
+			}
+		}
+
 		/// <summary>
 		/// Select a Perforce server to use
 		/// </summary>
@@ -252,11 +275,11 @@ namespace Horde.Build.Perforce
 			List<PerforceServer> validServers = new List<PerforceServer>();
 			if (properties != null)
 			{
-				validServers.AddRange(cluster.Servers.Where(x => x.Properties != null && x.Properties.Count > 0 && x.Properties.All(y => properties.Contains(y))));
+				validServers.AddRange(cluster.Servers.Where(x => (x.Properties != null && x.Properties.Count > 0 && x.Properties.All(y => properties.Contains(y))) || EvaluateCondition(x.Condition, properties)));
 			}
 			if (validServers.Count == 0)
 			{
-				validServers.AddRange(cluster.Servers.Where(x => x.Properties == null || x.Properties.Count == 0));
+				validServers.AddRange(cluster.Servers.Where(x => (x.Properties == null || x.Properties.Count == 0) && (x.Condition == null || x.Condition.IsEmpty())));
 			}
 
 			HashSet<string> validServerNames = new HashSet<string>(validServers.Select(x => x.ServerAndPort), StringComparer.OrdinalIgnoreCase);
@@ -358,7 +381,11 @@ namespace Horde.Build.Perforce
 
 		static void MergeServerList(PerforceServerList serverList, List<PerforceServerEntry> newEntries)
 		{
-			Dictionary<string, PerforceServerEntry> existingEntries = serverList.Servers.ToDictionary(x => x.ServerAndPort, x => x, StringComparer.OrdinalIgnoreCase);
+			Dictionary<string, PerforceServerEntry> existingEntries = new Dictionary<string, PerforceServerEntry>(StringComparer.OrdinalIgnoreCase);
+			foreach (PerforceServerEntry existingEntry in serverList.Servers)
+			{
+				existingEntries[existingEntry.ServerAndPort] = existingEntry;
+			}
 			foreach (PerforceServerEntry newEntry in newEntries)
 			{
 				PerforceServerEntry? existingEntry;
diff --git a/Engine/Source/Programs/Horde/Horde.Build/Server/Globals.cs b/Engine/Source/Programs/Horde/Horde.Build/Server/Globals.cs
--- a/Engine/Source/Programs/Horde/Horde.Build/Server/Globals.cs
+++ b/Engine/Source/Programs/Horde/Horde.Build/Server/Globals.cs
@@ -5,6 +5,7 @@ using System.Collections.Generic;
 using System.ComponentModel.DataAnnotations;
 using System.Linq;
 using System.Security.Cryptography;
+using EpicGames.Horde.Common;
 using EpicGames.Perforce;
 using Horde.Build.Acls;
 using Horde.Build.Server;
@@ -202,6 +203,11 @@ namespace Horde.Build.Server
 		/// </summary>
 		public int MaxConformCount { get; set; }
 
+		/// <summary>
+		/// Optional condition for a machine to be eligable to use this server
+		/// </summary>
+		public Condition? Condition { get; set; }
+
 		/// <summary>
 		/// List of properties for an agent to be eligable to use this server
 		/// </summary>
