Horde: Fix callstacks generating different hashes. Now redacts all hex sequences before hashing events.

diff --git a/Engine/Source/Programs/Horde/Horde.Build.Tests/IssueServiceTests.cs b/Engine/Source/Programs/Horde/Horde.Build.Tests/IssueServiceTests.cs
--- a/Engine/Source/Programs/Horde/Horde.Build.Tests/IssueServiceTests.cs
+++ b/Engine/Source/Programs/Horde/Horde.Build.Tests/IssueServiceTests.cs
@@ -812,6 +812,73 @@ namespace Horde.Build.Tests
 			}
 		}
 
+		[TestMethod]
+		public async Task CallstackTest()
+		{
+			string[] lines1 =
+			{
+				@"LogWindows: Error: begin: stack for UAT",
+				@"LogWindows: Error: === Critical error: ===",
+				@"LogWindows: Error:",
+				@"LogWindows: Error: Fatal error!",
+				@"LogWindows: Error:",
+				@"LogWindows: Error: Unhandled Exception: EXCEPTION_ACCESS_VIOLATION reading address 0x0000000000000070",
+				@"LogWindows: Error:",
+				@"LogWindows: Error: [Callstack] 0x00007ffdaea6afc8 UnrealEditor-Landscape.dll!ALandscape::ALandscape() []",
+				@"LogWindows: Error: [Callstack] 0x00007ffdc005d375 UnrealEditor-CoreUObject.dll!StaticConstructObject_Internal() []",
+				@"LogWindows: Error: [Callstack] 0x00007ffdbfe7f2af UnrealEditor-CoreUObject.dll!FLinkerLoad::CreateExport() []",
+				@"LogWindows: Error: [Callstack] 0x00007ffdbfe7fb7b UnrealEditor-CoreUObject.dll!FLinkerLoad::CreateExportAndPreload() []",
+				@"LogWindows: Error: [Callstack] 0x00007ffdbfea9141 UnrealEditor-CoreUObject.dll!FLinkerLoad::LoadAllObjects() []",
+				@"LogWindows: Error:",
+				@"LogWindows: Error: end: stack for UAT",
+				@"Took 70.6962389s to run UnrealEditor-Cmd.exe, ExitCode=3",
+				@"Copying crash data to d:\build\U5M+Inc\Sync\Engine\Programs\AutomationTool\Saved\Logs\Crashes\UECC-Windows-D7C3D5AD4E079F5DF8FD00B69907CD38_0000...",
+				@"Editor terminated with exit code 3 while running Cook for D:\build\U5M+Inc\Sync\Samples\Games\Lyra\Lyra.uproject; see log d:\build\U5M+Inc\Sync\Engine\Programs\AutomationTool\Saved\Logs\Cook-2022.08.19-17.34.05.txt",
+			};
+
+			IJob job1 = CreateJob(_mainStreamId, 120, "Compile Test", _graph);
+			await ParseEventsAsync(job1, 0, 0, lines1);
+			await UpdateCompleteStep(job1, 0, 0, JobStepOutcome.Failure);
+
+			List<IIssue> issues1 = await IssueCollection.FindIssuesAsync();
+			Assert.AreEqual(1, issues1.Count);
+
+			IIssue issue1 = issues1[0];
+			Assert.AreEqual(issue1.Fingerprints.Count, 1);
+			Assert.AreEqual(issue1.Fingerprints[0].Type, "Hashed");
+
+			// SAME ERROR BUT DIFFERENT CALLSTACK ADDRESSES
+
+			string[] lines2 =
+			{
+				@"LogWindows: Error: begin: stack for UAT",
+				@"LogWindows: Error: === Critical error: ===",
+				@"LogWindows: Error:",
+				@"LogWindows: Error: Fatal error!",
+				@"LogWindows: Error:",
+				@"LogWindows: Error: Unhandled Exception: EXCEPTION_ACCESS_VIOLATION reading address 0x0000000000000070",
+				@"LogWindows: Error:",
+				@"LogWindows: Error: [Callstack] 0x00007ff95973afc8 UnrealEditor-Landscape.dll!ALandscape::ALandscape() []",
+				@"LogWindows: Error: [Callstack] 0x00007ff963ced375 UnrealEditor-CoreUObject.dll!StaticConstructObject_Internal() []",
+				@"LogWindows: Error: [Callstack] 0x00007ff963b0f2af UnrealEditor-CoreUObject.dll!FLinkerLoad::CreateExport() []",
+				@"LogWindows: Error: [Callstack] 0x00007ff963b0fb7b UnrealEditor-CoreUObject.dll!FLinkerLoad::CreateExportAndPreload() []",
+				@"LogWindows: Error: [Callstack] 0x00007ff963b39141 UnrealEditor-CoreUObject.dll!FLinkerLoad::LoadAllObjects() []",
+				@"LogWindows: Error:",
+				@"LogWindows: Error: end: stack for UAT",
+				@"Took 67.0103214s to run UnrealEditor-Cmd.exe, ExitCode=3",
+				@"Copying crash data to d:\build\U5M+Inc\Sync\Engine\Programs\AutomationTool\Saved\Logs\Crashes\UECC-Windows-5F2FFCFE4EAAFE5DDD3E0EAB04336EA0_0000...",
+				@"Editor terminated with exit code 3 while running Cook for D:\build\U5M+Inc\Sync\Samples\Games\Lyra\Lyra.uproject; see log d:\build\U5M+Inc\Sync\Engine\Programs\AutomationTool\Saved\Logs\Cook-2022.08.19-16.53.28.txt",
+			};
+
+			IJob job2 = CreateJob(_mainStreamId, 120, "Compile Test", _graph);
+			await ParseEventsAsync(job2, 0, 0, lines2);
+			await UpdateCompleteStep(job2, 0, 0, JobStepOutcome.Failure);
+
+			List<IIssue> issues2 = await IssueCollection.FindIssuesAsync();
+			Assert.AreEqual(1, issues2.Count);
+			Assert.AreEqual(issues2[0].Id, issues1[0].Id);
+		}
+
 		[TestMethod]
 		public async Task CompileTypeTest()
 		{
diff --git a/Engine/Source/Programs/Horde/Horde.Build/Issues/Handlers/HashedIssueHandler.cs b/Engine/Source/Programs/Horde/Horde.Build/Issues/Handlers/HashedIssueHandler.cs
--- a/Engine/Source/Programs/Horde/Horde.Build/Issues/Handlers/HashedIssueHandler.cs
+++ b/Engine/Source/Programs/Horde/Horde.Build/Issues/Handlers/HashedIssueHandler.cs
@@ -49,6 +49,7 @@ namespace Horde.Build.Issues.Handlers
 		static bool TryGetHash(string message, out Md5Hash hash)
 		{
 			string sanitized = message.ToUpperInvariant();
+			sanitized = Regex.Replace(sanitized, @"0[xX][0-9a-fA-F]+", "H"); // Redact hex strings
 			sanitized = Regex.Replace(sanitized, @"\d[\d.,:]*", "n"); // Redact numbers and timestamp like things
 
 			if (sanitized.Length > 30)
