Horde: Format issue details using BlockKit messages, and add controls underneath.

diff --git a/Engine/Source/Programs/Horde/Horde.Build/Notifications/Sinks/SlackNotificationSink.cs b/Engine/Source/Programs/Horde/Horde.Build/Notifications/Sinks/SlackNotificationSink.cs
--- a/Engine/Source/Programs/Horde/Horde.Build/Notifications/Sinks/SlackNotificationSink.cs
+++ b/Engine/Source/Programs/Horde/Horde.Build/Notifications/Sinks/SlackNotificationSink.cs
@@ -805,17 +805,24 @@ namespace Horde.Build.Notifications.Sinks
 						}
 					}
 
-					StringBuilder summary = new StringBuilder();
-					summary.AppendLine($"From {FormatJobStep(span.FirstFailure, span.NodeName)}:");
+					SlackMessage message = new SlackMessage();
+
+					message.Blocks.Add(new SectionBlock(new TextObject($"From {FormatJobStep(span.FirstFailure, span.NodeName)}:")));
 					foreach (ILogEventData eventDataItem in eventDataItems)
 					{
-						summary.AppendLine(QuoteText(eventDataItem.Message, MaxLineLength));
+						message.Blocks.Add(new SectionBlock(new TextObject(QuoteText(eventDataItem.Message, MaxLineLength)));
 					}
 					if (events.Count > eventDataItems.Count)
 					{
-						summary.AppendLine("```...```");
+						message.Blocks.Add(new SectionBlock(new TextObject("```...```")));
 					}
-					string? summaryTs = await _slackClient.PostMessageAsync(triageChannel, state.Ts, summary.ToString());
+
+					ActionsBlock actions = message.AddActions();
+					actions.AddButton("Assign to Me", value: $"issue_{issue.Id}_ack", style: ButtonStyle.Primary);
+					actions.AddButton("Not Me", value: $"issue_{issue.Id}_decline", style: ButtonStyle.Danger);
+					actions.AddButton("Mark Fixed", value: $"issue_{issue.Id}_markfixed");
+
+					string? summaryTs = await _slackClient.PostMessageAsync(triageChannel, state.Ts, message);
 
 					// Permalink to the summary text so we link inside the thread rather than just to the original message
 					string? permalink = null;
@@ -824,18 +831,6 @@ namespace Horde.Build.Notifications.Sinks
 						permalink = await _slackClient.GetPermalinkAsync(state.Channel, summaryTs);
 					}
 
-					// Post controls for interacting with the issue
-					{
-						SlackMessage message = new SlackMessage();
-
-						ActionsBlock actions = message.AddActions();
-						actions.AddButton("Acknowledge", value: $"issue_{issue.Id}_ack", style: ButtonStyle.Primary);
-						actions.AddButton("Not Me", value: $"issue_{issue.Id}_decline", style: ButtonStyle.Danger);
-						actions.AddButton("Mark Fixed", value: $"issue_{issue.Id}_markfixed");
-
-						await _slackClient.PostMessageAsync(triageChannel, state.Ts, message);
-					}
-
 					// If it has an owner, show that
 					HashSet<UserId> inviteUserIds = new HashSet<UserId>();
 					if (issue.OwnerId != null)
