Horde: Add a mask emoji to a thread if the issue is quarantined.

diff --git a/Engine/Source/Programs/Horde/Horde.Build/Notifications/Sinks/SlackNotificationSink.cs b/Engine/Source/Programs/Horde/Horde.Build/Notifications/Sinks/SlackNotificationSink.cs
--- a/Engine/Source/Programs/Horde/Horde.Build/Notifications/Sinks/SlackNotificationSink.cs
+++ b/Engine/Source/Programs/Horde/Horde.Build/Notifications/Sinks/SlackNotificationSink.cs
@@ -877,6 +877,15 @@ namespace Horde.Build.Notifications.Sinks
 					await _slackClient.RemoveReactionAsync(state.Channel, state.Ts, "eyes");
 				}
 
+				if (issue.QuarantinedByUserId != null)
+				{
+					await _slackClient.AddReactionAsync(state.Channel, state.Ts, "mask");
+				}
+				else
+				{
+					await _slackClient.RemoveReactionAsync(state.Channel, state.Ts, "mask");
+				}
+
 				IIssueSpan? fixFailedSpan = null;
 				if (issue.FixChange != null)
 				{
