Horde: Build fix.

diff --git a/Engine/Source/Programs/Horde/Horde.Build/Notifications/Sinks/SlackNotificationSink.cs b/Engine/Source/Programs/Horde/Horde.Build/Notifications/Sinks/SlackNotificationSink.cs
--- a/Engine/Source/Programs/Horde/Horde.Build/Notifications/Sinks/SlackNotificationSink.cs
+++ b/Engine/Source/Programs/Horde/Horde.Build/Notifications/Sinks/SlackNotificationSink.cs
@@ -810,7 +810,7 @@ namespace Horde.Build.Notifications.Sinks
 					message.Blocks.Add(new SectionBlock(new TextObject($"From {FormatJobStep(span.FirstFailure, span.NodeName)}:")));
 					foreach (ILogEventData eventDataItem in eventDataItems)
 					{
-						message.Blocks.Add(new SectionBlock(new TextObject(QuoteText(eventDataItem.Message, MaxLineLength)));
+						message.Blocks.Add(new SectionBlock(new TextObject(QuoteText(eventDataItem.Message, MaxLineLength))));
 					}
 					if (events.Count > eventDataItems.Count)
 					{
