Horde: Ensure a default AWS region is always set for AWS clients in DI

diff --git a/Engine/Source/Programs/Horde/Horde.Build/Startup.cs b/Engine/Source/Programs/Horde/Horde.Build/Startup.cs
--- a/Engine/Source/Programs/Horde/Horde.Build/Startup.cs
+++ b/Engine/Source/Programs/Horde/Horde.Build/Startup.cs
@@ -12,6 +12,7 @@ using System.Text.Json;
 using System.Text.Json.Serialization;
 using System.Threading;
 using System.Threading.Tasks;
+using Amazon;
 using Amazon.CloudWatch;
 using Amazon.Extensions.NETCore.Setup;
 using EpicGames.AspNet;
@@ -467,6 +468,11 @@ namespace Horde.Build
 			
 			AWSOptions awsOptions = Configuration.GetAWSOptions();
 			services.AddDefaultAWSOptions(awsOptions);
+			if (awsOptions.Region == null && Environment.GetEnvironmentVariable("AWS_REGION") == null)
+			{
+				awsOptions.Region = RegionEndpoint.USEast1;
+			}
+			
 			services.AddAWSService<IAmazonCloudWatch>();
 
 			ConfigureLogStorage(services);
