Horde: Fix bug re-using same name for shared ticker in AutoscaleService

diff --git a/Engine/Source/Programs/Horde/Horde.Build/Agents/Fleet/AutoscaleServiceV2.cs b/Engine/Source/Programs/Horde/Horde.Build/Agents/Fleet/AutoscaleServiceV2.cs
--- a/Engine/Source/Programs/Horde/Horde.Build/Agents/Fleet/AutoscaleServiceV2.cs
+++ b/Engine/Source/Programs/Horde/Horde.Build/Agents/Fleet/AutoscaleServiceV2.cs
@@ -67,7 +67,7 @@ namespace Horde.Build.Agents.Fleet
 			_dogStatsd = dogStatsd;
 			_clock = clock;
 			_ticker = clock.AddSharedTicker<AutoscaleServiceV2>(TimeSpan.FromMinutes(5.0), TickLeaderAsync, logger);
-			_tickerHighFrequency = clock.AddSharedTicker<AutoscaleServiceV2>(TimeSpan.FromSeconds(30), TickHighFrequencyAsync, logger);
+			_tickerHighFrequency = clock.AddSharedTicker("AutoscaleServiceV2.TickHighFrequency", TimeSpan.FromSeconds(30), TickHighFrequencyAsync, logger);
 			_settings = settings.Value;
 			_logger = logger;
 			_defaultScaleOutCooldown = TimeSpan.FromSeconds(settings.Value.AgentPoolScaleOutCooldownSeconds);
@@ -87,7 +87,11 @@ namespace Horde.Build.Agents.Fleet
 		}
 
 		/// <inheritdoc/>
-		public void Dispose() => _ticker.Dispose();
+		public void Dispose()
+		{
+			_ticker.Dispose();
+			_tickerHighFrequency.Dispose();
+		}
 
 		internal async ValueTask TickLeaderAsync(CancellationToken stoppingToken)
 		{
