Horde Storage - Stop scanning references once we have found one reference that pins the blob. Not scanning them in parallel but blobs are still scanned in parallel, decent tradeof to make the abort easier.

diff --git a/Engine/Source/Programs/Horde/HordeStorage/Horde.Storage/Implementation/GC/OrphanBlobCleanupRefs.cs b/Engine/Source/Programs/Horde/HordeStorage/Horde.Storage/Implementation/GC/OrphanBlobCleanupRefs.cs
--- a/Engine/Source/Programs/Horde/HordeStorage/Horde.Storage/Implementation/GC/OrphanBlobCleanupRefs.cs
+++ b/Engine/Source/Programs/Horde/HordeStorage/Horde.Storage/Implementation/GC/OrphanBlobCleanupRefs.cs
@@ -136,19 +136,14 @@ namespace Horde.Storage.Implementation
                     return;
                 }
 
-                await Parallel.ForEachAsync(blobIndex.References, cancellationToken, async (tuple, token) =>
+                foreach ((BucketId, IoHashKey) tuple in blobIndex.References)
                 {
-                    (BucketId bucket, IoHashKey key) = tuple;
-                    if (token.IsCancellationRequested)
-                    {
-                        return;
-                    }
-
                     try
                     {
-                        (ObjectRecord, BlobContents?) _ = await _objectService.Get(blobNamespace,
-                            bucket, key, new string[] { "name" }, doLastAccessTracking:false);
+                        (BucketId bucket, IoHashKey key) = tuple;
+                        (ObjectRecord, BlobContents?) _ = await _objectService.Get(blobNamespace, bucket, key, new string[] { "name" }, doLastAccessTracking: false);
                         found = true;
+                        break;
                     }
                     catch (ObjectNotFoundException)
                     {
@@ -159,7 +154,7 @@ namespace Horde.Storage.Implementation
                         // we do not care if there are missing blobs, as long as the record is valid we keep this blob around
                         found = true;
                     }
-                });
+                }
             });
 
             if (cancellationToken.IsCancellationRequested)
