Horde: Add support for scope annotations in audit logs, and use it to flag any issue updates with the user that initiated them.

diff --git a/Engine/Source/Programs/Horde/Horde.Build/Auditing/AuditLog.cs b/Engine/Source/Programs/Horde/Horde.Build/Auditing/AuditLog.cs
--- a/Engine/Source/Programs/Horde/Horde.Build/Auditing/AuditLog.cs
+++ b/Engine/Source/Programs/Horde/Horde.Build/Auditing/AuditLog.cs
@@ -2,6 +2,7 @@
 
 using System;
 using System.Collections.Generic;
+using System.Text;
 using System.Threading.Channels;
 using System.Threading.Tasks;
 using EpicGames.Core;
@@ -52,27 +53,92 @@ namespace Horde.Build.Auditing
 		{
 			sealed class Scope : IDisposable
 			{
-				public void Dispose() { }
+				readonly AuditLogChannel _owner;
+
+				public LogEvent Message { get; }
+
+				public Scope(AuditLogChannel owner, LogEvent message)
+				{
+					_owner = owner;
+					Message = message;
+
+					_owner._scopes.Add(this);
+				}
+
+				public void Dispose() => _owner._scopes.Remove(this);
 			}
 
 			public readonly AuditLog<TSubject> Outer;
 			public TSubject Subject { get; }
 
+			readonly List<Scope> _scopes = new List<Scope>();
+
 			public AuditLogChannel(AuditLog<TSubject> outer, TSubject subject)
 			{
 				Outer = outer;
 				Subject = subject;
 			}
 
-			public IDisposable BeginScope<TState>(TState state) => new Scope();
+			public IDisposable BeginScope<TState>(TState state) => new Scope(this, LogEvent.FromState(LogLevel.Information, default, state, null, (x, y) => x?.ToString() ?? String.Empty));
 
 			public bool IsEnabled(LogLevel logLevel) => true;
 
+			LogEvent CreateEvent<TState>(LogLevel logLevel, EventId eventId, TState state, Exception? exception, Func<TState, Exception?, string> formatter)
+			{
+				LogEvent logEvent = LogEvent.FromState(logLevel, eventId, state, exception, formatter);
+				if (_scopes.Count > 0)
+				{
+					StringBuilder message = new StringBuilder();
+					StringBuilder format = new StringBuilder();
+					Dictionary<string, object> properties = new Dictionary<string, object>(StringComparer.Ordinal);
+
+					for (int idx = 0; idx < _scopes.Count; idx++)
+					{
+						message.Append('[');
+						format.Append('[');
+						AppendMessage(_scopes[idx].Message, idx + 1, message, format, properties);
+						message.Append(']');
+						format.Append(']');
+					}
+
+					message.Append(' ');
+					format.Append(' ');
+
+					AppendMessage(logEvent, 0, message, format, properties);
+
+					logEvent.Message = message.ToString();
+					logEvent.Format = format.ToString();
+					logEvent.Properties = properties;
+				}
+				return logEvent;
+			}
+
+			static void AppendMessage(LogEvent logEvent, int id, StringBuilder message, StringBuilder format, Dictionary<string, object> properties)
+			{
+				message.Append(logEvent.Message);
+				if (logEvent.Format == null)
+				{
+					format.Append($"{{_scope{id}}}");
+					properties.Add($"_scope{id}", logEvent.Message);
+				}
+				else
+				{
+					format.Append(logEvent.Format);
+					if (logEvent.Properties != null)
+					{
+						foreach (KeyValuePair<string, object> item in logEvent.Properties)
+						{
+							properties[item.Key] = item.Value;
+						}
+					}
+				}
+			}
+
 			public void Log<TState>(LogLevel logLevel, EventId eventId, TState state, Exception? exception, Func<TState, Exception?, string> formatter)
 			{
 				DateTime time = DateTime.UtcNow;
 
-				LogEvent logEvent = LogEvent.FromState(logLevel, eventId, state, exception, formatter);
+				LogEvent logEvent = CreateEvent(logLevel, eventId, state, exception, formatter);
 				string data = logEvent.ToJson();
 				AuditLogMessage message = new AuditLogMessage(Subject, time, logLevel, data);
 				Outer._messageChannel.Writer.TryWrite(message);
diff --git a/Engine/Source/Programs/Horde/Horde.Build/Issues/IssuesController.cs b/Engine/Source/Programs/Horde/Horde.Build/Issues/IssuesController.cs
--- a/Engine/Source/Programs/Horde/Horde.Build/Issues/IssuesController.cs
+++ b/Engine/Source/Programs/Horde/Horde.Build/Issues/IssuesController.cs
@@ -620,9 +620,12 @@ namespace Horde.Build.Issues
 				removeSpans = request.RemoveSpans.ConvertAll(x => ObjectId.Parse(x));
 			}
 
-			if (!await _issueService.UpdateIssueAsync(issueId, request.Summary, request.Description, request.Promoted, newOwnerId, newNominatedById, request.Acknowledged, newDeclinedById, request.FixChange, newResolvedById, addSpans, removeSpans, request.ExternalIssueKey, newQuarantinedById))
+			using (IDisposable scope = _issueCollection.GetLogger(issueId).BeginScope("User {UserId}", User.GetUserId() ?? UserId.Empty))
 			{
-				return NotFound();
+				if (!await _issueService.UpdateIssueAsync(issueId, request.Summary, request.Description, request.Promoted, newOwnerId, newNominatedById, request.Acknowledged, newDeclinedById, request.FixChange, newResolvedById, addSpans, removeSpans, request.ExternalIssueKey, newQuarantinedById))
+				{
+					return NotFound();
+				}
 			}
 			return Ok();
 		}
