Horde Storage - Further reducing the amount of scylla calls made during GC. This was added to speed up the GC (by doing all calls in parallel) but this prevents us from existing early when we find a valid blob (as the other calls are already in flight). If we want to run the gc more aggresivley (spend more compute on it) we can always bump OrphanGCMaxParallelOperations which will check more blobs in parallel but not internally parallize the calls.

diff --git a/Engine/Source/Programs/Horde/HordeStorage/Horde.Storage/Implementation/GC/OrphanBlobCleanupRefs.cs b/Engine/Source/Programs/Horde/HordeStorage/Horde.Storage/Implementation/GC/OrphanBlobCleanupRefs.cs
--- a/Engine/Source/Programs/Horde/HordeStorage/Horde.Storage/Implementation/GC/OrphanBlobCleanupRefs.cs
+++ b/Engine/Source/Programs/Horde/HordeStorage/Horde.Storage/Implementation/GC/OrphanBlobCleanupRefs.cs
@@ -122,18 +122,23 @@ namespace Horde.Storage.Implementation
             bool found = false;
 
             // check all namespaces that share the same storage pool for presence of the blob
-            await Parallel.ForEachAsync(namespacesThatSharePool, cancellationToken, async (blobNamespace, token) =>
+            foreach (NamespaceId blobNamespace in namespacesThatSharePool)
             {
-                if (token.IsCancellationRequested)
+                if (cancellationToken.IsCancellationRequested)
                 {
-                    return;
+                    break;
+                }
+
+                if (found)
+                {
+                    break;
                 }
 
                 BlobInfo? blobIndex = await _blobIndex.GetBlobInfo(blobNamespace, blob);
 
                 if (blobIndex == null)
                 {
-                    return;
+                    break;
                 }
 
                 foreach ((BucketId, IoHashKey) tuple in blobIndex.References)
@@ -155,7 +160,7 @@ namespace Horde.Storage.Implementation
                         found = true;
                     }
                 }
-            });
+            }
 
             if (cancellationToken.IsCancellationRequested)
             {
