Horde: Include changelist numbers when showing who an issue is assigned to.

diff --git a/Engine/Source/Programs/Horde/Horde.Build/Notifications/Sinks/SlackNotificationSink.cs b/Engine/Source/Programs/Horde/Horde.Build/Notifications/Sinks/SlackNotificationSink.cs
--- a/Engine/Source/Programs/Horde/Horde.Build/Notifications/Sinks/SlackNotificationSink.cs
+++ b/Engine/Source/Programs/Horde/Horde.Build/Notifications/Sinks/SlackNotificationSink.cs
@@ -831,17 +831,26 @@ namespace Horde.Build.Notifications.Sinks
 						permalink = await _slackClient.GetPermalinkAsync(state.Channel, summaryTs);
 					}
 
+					// Get the suspects for the issue
+					List<IIssueSuspect> suspects = await _issueService.Collection.FindSuspectsAsync(issue);
+
 					// If it has an owner, show that
 					HashSet<UserId> inviteUserIds = new HashSet<UserId>();
 					if (issue.OwnerId != null)
 					{
 						string mention = await FormatMentionAsync(issue.OwnerId.Value, workflow.AllowMentions);
+
+						string changes = String.Join(", ", suspects.Where(x => x.AuthorId == issue.OwnerId).Select(x => FormatChange(x.Change)));
+						if (changes.Length > 0)
+						{
+							mention += $" ({changes})";
+						}
+
 						await _slackClient.PostMessageAsync(triageChannel, state.Ts, $"Assigned to {mention}");
 						inviteUserIds.Add(issue.OwnerId.Value);
 					}
 					else
 					{
-						List<IIssueSuspect> suspects = await _issueService.Collection.FindSuspectsAsync(issue);
 						IGrouping<UserId, IIssueSuspect>[] suspectGroups = suspects.GroupBy(x => x.AuthorId).ToArray();
 
 						if (suspectGroups.Length > 0 && suspectGroups.Length <= workflow.MaxMentions)
