Horde: Improve handling of Perforce case errors. Now correctly identifies correct suspect for mismatched case.

diff --git a/Engine/Source/Programs/Horde/Horde.Build.Tests/IssueServiceTests.cs b/Engine/Source/Programs/Horde/Horde.Build.Tests/IssueServiceTests.cs
--- a/Engine/Source/Programs/Horde/Horde.Build.Tests/IssueServiceTests.cs
+++ b/Engine/Source/Programs/Horde/Horde.Build.Tests/IssueServiceTests.cs
@@ -271,6 +271,12 @@ namespace Horde.Build.Tests
 			await LogFileService.CreateEventsAsync(new List<NewLogEventData> { new NewLogEventData { LogId = logId, LineIndex = metadata.MaxLineIndex, LineCount = 1, Severity = severity } });
 		}
 
+		private TestJsonLogger CreateLogger(IJob job, int batchIdx, int stepIdx)
+		{
+			LogId logId = job.Batches[batchIdx].Steps[stepIdx].LogId!.Value;
+			return new TestJsonLogger(LogFileService, logId);
+		}
+
 		private async Task ParseEventsAsync(IJob job, int batchIdx, int stepIdx, string[] lines)
 		{
 			LogId logId = job.Batches[batchIdx].Steps[stepIdx].LogId!.Value;
@@ -460,6 +466,48 @@ namespace Horde.Build.Tests
 			}
 		}
 
+		[TestMethod]
+		public async Task PerforceCaseIssueTest()
+		{
+			// #1
+			// Scenario: Job step completes successfully at CL 105
+			// Expected: No issues are created
+			{
+				IJob job = CreateJob(_mainStreamId, 105, "Test Build", _graph);
+				await UpdateCompleteStep(job, 0, 0, JobStepOutcome.Warnings);
+
+				List<IIssue> issues = await IssueCollection.FindIssuesAsync();
+				Assert.AreEqual(0, issues.Count);
+			}
+
+			// #2
+			// Scenario: Job step fails at CL 170 with P4 case error
+			// Expected: Creates issue, identifies source file correctly
+			{
+				IUser chris = await UserCollection.FindOrAddUserByLoginAsync("Chris");
+				_perforce.AddChange(MainStreamName, 150, chris, "Description", new string[] { "Engine/Foo/Bar.txt" });
+
+				IUser john = await UserCollection.FindOrAddUserByLoginAsync("John");
+				_perforce.AddChange(MainStreamName, 160, john, "Description", new string[] { "Engine/Foo/Baz.txt" });
+
+				IJob job = CreateJob(_mainStreamId, 170, "Test Build", _graph);
+				await using (TestJsonLogger logger = CreateLogger(job, 0, 0))
+				{
+					logger.LogWarning(KnownLogEvents.AutomationTool_PerforceCase, "    {DepotFile}", new LogValue(LogValueType.DepotPath, "//UE5/Main/Engine/Foo/Bar.txt"));
+				}
+				await UpdateCompleteStep(job, 0, 0, JobStepOutcome.Warnings);
+
+				List<IIssue> issues = await IssueCollection.FindIssuesAsync();
+				Assert.AreEqual(1, issues.Count);
+				Assert.AreEqual(IssueSeverity.Warning, issues[0].Severity);
+				Assert.AreEqual("PerforceCase", issues[0].Fingerprints[0].Type);
+				Assert.AreEqual("//UE5/Main/Engine/Foo/Bar.txt", issues[0].Fingerprints[0].Keys.First());
+				Assert.AreEqual(chris.Id, issues[0].OwnerId);
+
+				Assert.AreEqual("Inconsistent case for Bar.txt", issues[0].Summary);
+			}
+		}
+
 		[TestMethod]
 		public async Task AutoSdkWarningTest()
 		{
diff --git a/Engine/Source/Programs/Horde/Horde.Build.Tests/Stubs/Services/PerforceServiceStub.cs b/Engine/Source/Programs/Horde/Horde.Build.Tests/Stubs/Services/PerforceServiceStub.cs
--- a/Engine/Source/Programs/Horde/Horde.Build.Tests/Stubs/Services/PerforceServiceStub.cs
+++ b/Engine/Source/Programs/Horde/Horde.Build.Tests/Stubs/Services/PerforceServiceStub.cs
@@ -17,7 +17,7 @@ namespace Horde.Build.Tests.Stubs.Services
 	{
 		public static ChangeFile CreateChangeFile(string path)
 		{
-			return new ChangeFile(path, null!, 0, 0, EpicGames.Core.Md5Hash.Zero, null!);
+			return new ChangeFile(path, $"//UE5/Main/{path}", 0, 0, EpicGames.Core.Md5Hash.Zero, null!);
 		}
 
 		public static void Add(this List<ChangeFile> files, string path)
diff --git a/Engine/Source/Programs/Horde/Horde.Build/Issues/Handlers/PerforceCaseIssueHandler.cs b/Engine/Source/Programs/Horde/Horde.Build/Issues/Handlers/PerforceCaseIssueHandler.cs
--- a/Engine/Source/Programs/Horde/Horde.Build/Issues/Handlers/PerforceCaseIssueHandler.cs
+++ b/Engine/Source/Programs/Horde/Horde.Build/Issues/Handlers/PerforceCaseIssueHandler.cs
@@ -58,18 +58,12 @@ namespace Horde.Build.Issues.Handlers
 		/// <param name="depotFiles">List of source files</param>
 		static void GetSourceFiles(ILogEventData logEventData, HashSet<string> depotFiles)
 		{
-			foreach (ILogEventLine line in logEventData.Lines)
+			foreach (JsonProperty property in logEventData.FindPropertiesOfType(LogValueType.DepotPath))
 			{
-				JsonElement properties;
-				if (line.Data.TryGetProperty("properties", out properties) && properties.ValueKind == JsonValueKind.Object)
+				JsonElement value;
+				if (property.Value.TryGetProperty(LogEventPropertyName.Text.Span, out value) && value.ValueKind == JsonValueKind.String)
 				{
-					foreach (JsonProperty property in properties.EnumerateObject())
-					{
-						if (property.NameEquals("File") && property.Value.ValueKind == JsonValueKind.String)
-						{
-							depotFiles.Add(property.Value.GetString()!);
-						}
-					}
+					depotFiles.Add(value.GetString()!);
 				}
 			}
 		}
diff --git a/Engine/Source/Programs/Horde/Horde.Build/Logs/ILogEvent.cs b/Engine/Source/Programs/Horde/Horde.Build/Logs/ILogEvent.cs
--- a/Engine/Source/Programs/Horde/Horde.Build/Logs/ILogEvent.cs
+++ b/Engine/Source/Programs/Horde/Horde.Build/Logs/ILogEvent.cs
@@ -1,7 +1,9 @@
 // Copyright Epic Games, Inc. All Rights Reserved.
 
 using System.Collections.Generic;
+using System.Linq;
 using System.Text.Json;
+using EpicGames.Core;
 using Horde.Build.Utilities;
 using HordeCommon;
 using Microsoft.Extensions.Logging;
@@ -124,4 +126,55 @@ namespace Horde.Build.Logs
 		/// </summary>
 		public ObjectId? SpanId { get; set; }
 	}
+
+	/// <summary>
+	/// Extensions for parsing log event properties
+	/// </summary>
+	public static class LogEventExtensions
+	{
+		/// <summary>
+		/// Find all properties of the given type in a particular log line
+		/// </summary>
+		/// <param name="data">Line data</param>
+		/// <param name="type">Type of property to return</param>
+		/// <returns></returns>
+		public static IEnumerable<JsonProperty> FindPropertiesOfType(this ILogEventData data, Utf8String type)
+		{
+			return data.Lines.SelectMany(x => FindPropertiesOfType(x, type));
+		}
+
+		/// <summary>
+		/// Find all properties of the given type in a particular log line
+		/// </summary>
+		/// <param name="line">Line data</param>
+		/// <param name="type">Type of property to return</param>
+		/// <returns></returns>
+		public static IEnumerable<JsonProperty> FindPropertiesOfType(this ILogEventLine line, Utf8String type)
+		{
+			JsonElement properties;
+			if (line.Data.TryGetProperty("properties", out properties) && properties.ValueKind == JsonValueKind.Object)
+			{
+				foreach (JsonProperty property in properties.EnumerateObject())
+				{
+					if (property.Value.ValueKind == JsonValueKind.Object)
+					{
+						foreach (JsonProperty subProperty in property.Value.EnumerateObject())
+						{
+							if (subProperty.NameEquals(LogEventPropertyName.Type.Span))
+							{
+								if (subProperty.Value.ValueKind == JsonValueKind.String && subProperty.Value.ValueEquals(type.Span))
+								{
+									yield return property;
+								}
+								else
+								{
+									break;
+								}
+							}
+						}
+					}
+				}
+			}
+		}
+	}
 }
