Horde: Add matchers for shader compiler errors.

diff --git a/Engine/Source/Programs/AutomationTool/AutomationUtils/Matchers/CrashEventMatcher.cs b/Engine/Source/Programs/AutomationTool/AutomationUtils/Matchers/CrashEventMatcher.cs
--- a/Engine/Source/Programs/AutomationTool/AutomationUtils/Matchers/CrashEventMatcher.cs
+++ b/Engine/Source/Programs/AutomationTool/AutomationUtils/Matchers/CrashEventMatcher.cs
@@ -15,6 +15,8 @@ namespace AutomationUtils.Matchers
 	{
 		static readonly Regex s_exitCodePattern = new Regex(@"ExitCode=(3|139|255)(?!\d)");
 
+		static readonly Regex s_appErrorPattern = new Regex(@"^\s*[A-Za-z]+: Error: appError called: ");
+
 		public LogEventMatch? Match(ILogCursor cursor)
 		{
 			if (cursor.Contains("begin: stack for UAT"))
@@ -28,6 +30,11 @@ namespace AutomationUtils.Matchers
 					}
 				}
 			}
+			if (cursor.IsMatch(s_appErrorPattern))
+			{
+				LogEventBuilder builder = new LogEventBuilder(cursor);
+				return builder.ToMatch(LogEventPriority.Normal, LogLevel.Error, KnownLogEvents.Engine_AppError);
+			}
 			if (cursor.Contains("AutomationTool: Stack:"))
 			{
 				LogEventBuilder builder = new LogEventBuilder(cursor);
diff --git a/Engine/Source/Programs/AutomationTool/AutomationUtils/Matchers/ShaderEventMatcher.cs b/Engine/Source/Programs/AutomationTool/AutomationUtils/Matchers/ShaderEventMatcher.cs
new file mode 100644
--- /dev/null
+++ b/Engine/Source/Programs/AutomationTool/AutomationUtils/Matchers/ShaderEventMatcher.cs
@@ -0,0 +1,72 @@
+// Copyright Epic Games, Inc. All Rights Reserved.
+
+using System;
+using System.Text.RegularExpressions;
+using EpicGames.Core;
+using Microsoft.Extensions.Logging;
+
+#nullable enable
+
+namespace AutomationUtils.Matchers
+{
+	/// <summary>
+	/// Matches shader compile errors and annotates with the source file path and revision
+	/// </summary>
+	class ShaderEventMatcher : ILogEventMatcher
+	{
+		const string Prefix = @"^\s*LogShaderCompilers: (?<severity>Error|Warning): ";
+
+		const string FilePattern =
+			@"(?<file>" +
+				// optional drive letter
+				@"(?:[a-zA-Z]:)?" +
+				// any non-colon character
+				@"[^:(\s]+" +
+				// any filename character (not whitespace or slash)
+				@"[^:(\s\\/]" +
+			@")";
+
+		const string LinePattern =
+			@"(?<line>\d+)(?::(?<column>\d+))?";
+
+		static readonly Regex s_pattern = new Regex(Prefix);
+		static readonly Regex s_patternWithFile = new Regex($"{Prefix}{FilePattern}(?:\\({LinePattern}\\))?:");
+
+		/// <inheritdoc/>
+		public LogEventMatch? Match(ILogCursor input)
+		{
+			Match? match;
+			if (input.TryMatch(s_pattern, out match))
+			{
+				LogLevel level = GetLogLevelFromSeverity(match);
+
+				LogEventBuilder builder = new LogEventBuilder(input);
+				builder.Annotate(match.Groups["severity"], LogEventMarkup.Severity);
+
+				Match? fileMatch;
+				if (input.TryMatch(s_patternWithFile, out fileMatch))
+				{
+					builder.AnnotateSourceFile(fileMatch.Groups["file"], "");
+					builder.TryAnnotate(fileMatch.Groups["line"], LogEventMarkup.LineNumber);
+					builder.TryAnnotate(fileMatch.Groups["column"], LogEventMarkup.ColumnNumber);
+				}
+
+				return builder.ToMatch(LogEventPriority.AboveNormal, level, KnownLogEvents.Engine_ShaderCompiler);
+			}
+			return null;
+		}
+
+		static LogLevel GetLogLevelFromSeverity(Match match)
+		{
+			string severity = match.Groups["severity"].Value;
+			if (severity.Equals("Warning", StringComparison.Ordinal))
+			{
+				return LogLevel.Warning;
+			}
+			else
+			{
+				return LogLevel.Error;
+			}
+		}
+	}
+}
diff --git a/Engine/Source/Programs/Horde/Horde.Build.Tests/IssueServiceTests.cs b/Engine/Source/Programs/Horde/Horde.Build.Tests/IssueServiceTests.cs
--- a/Engine/Source/Programs/Horde/Horde.Build.Tests/IssueServiceTests.cs
+++ b/Engine/Source/Programs/Horde/Horde.Build.Tests/IssueServiceTests.cs
@@ -508,6 +508,56 @@ namespace Horde.Build.Tests
 			}
 		}
 
+		[TestMethod]
+		public async Task ShaderIssueTest()
+		{
+			// #1
+			// Scenario: Job step completes successfully at CL 105
+			// Expected: No issues are created
+			{
+				IJob job = CreateJob(_mainStreamId, 105, "Test Build", _graph);
+				await UpdateCompleteStep(job, 0, 0, JobStepOutcome.Warnings);
+
+				List<IIssue> issues = await IssueCollection.FindIssuesAsync();
+				Assert.AreEqual(0, issues.Count);
+			}
+
+			// #2
+			// Scenario: Job step fails at CL 170 with shader compile error
+			// Expected: Creates issue, identifies source file correctly
+			{
+				string[] lines =
+				{
+					@"LogModuleManager: Display: Unable to bootstrap from archive [REDACTED_USER_HOME]/AppData/Local/Temp/UnrealXGEWorkingDir/E196A42E4BEB54E0AE6B0EB1573DBD7A//Bootstrap-5ABC233A-4168-4640-ABFD-61E8349924DD.modules, will fallback on normal initialization",
+					@"LogShaderCompilers: Warning: 1 Shader compiler errors compiling global shaders for platform SF_PS5:",
+					@"LogShaderCompilers: Error: " + FileReference.Combine(_workspaceDir, "Engine/Shaders/Private/Lumen/LumenScreenProbeTracing.usf").FullName + @"(810:95): Shader FScreenProbeTraceMeshSDFsCS, Permutation 95, VF None:	/Engine/Private/Lumen/LumenScreenProbeTracing.usf(810:95): (error, code:5476) - ambiguous call to 'select_internal'. Found 88 possible candidates:",
+					@"LogWindows: Error: appError called: Fatal error: [File:D:\build\U5M+Inc\Sync\Engine\Source\Runtime\Engine\Private\ShaderCompiler\ShaderCompiler.cpp] [Line: 7718]",
+					@"Took 64.0177761s to run UnrealEditor-Cmd.exe, ExitCode=3",
+					@"Copying crash data to d:\build\U5M+Inc\Sync\Engine\Programs\AutomationTool\Saved\Logs\Crashes\UECC-Windows-9AD258A94110A61CB524848FE0D7196D_0000...",
+					@"Editor terminated with exit code 3 while running Cook for D:\build\U5M+Inc\Sync\Samples\Games\ShooterGame\ShooterGame.uproject; see log d:\build\U5M+Inc\Sync\Engine\Programs\AutomationTool\Saved\Logs\Cook-2022.07.22-05.32.05.txt",
+				};
+
+				IUser chris = await UserCollection.FindOrAddUserByLoginAsync("Chris");
+				_perforce.AddChange(MainStreamName, 150, chris, "Description", new string[] { "Engine/Foo/LumenScreenProbeTracing.usf" });
+
+				IUser john = await UserCollection.FindOrAddUserByLoginAsync("John");
+				_perforce.AddChange(MainStreamName, 160, john, "Description", new string[] { "Engine/Foo/Baz.txt" });
+
+				IJob job = CreateJob(_mainStreamId, 170, "Test Build", _graph);
+				await ParseEventsAsync(job, 0, 0, lines);
+				await UpdateCompleteStep(job, 0, 0, JobStepOutcome.Failure);
+
+				List<IIssue> issues = await IssueCollection.FindIssuesAsync();
+				Assert.AreEqual(1, issues.Count);
+				Assert.AreEqual(IssueSeverity.Error, issues[0].Severity);
+				Assert.AreEqual("Shader", issues[0].Fingerprints[0].Type);
+				Assert.AreEqual("LumenScreenProbeTracing.usf", issues[0].Fingerprints[0].Keys.First());
+				Assert.AreEqual(chris.Id, issues[0].OwnerId);
+
+				Assert.AreEqual("Shader compile errors in LumenScreenProbeTracing.usf", issues[0].Summary);
+			}
+		}
+
 		[TestMethod]
 		public async Task AutoSdkWarningTest()
 		{
diff --git a/Engine/Source/Programs/Horde/Horde.Build/Issues/Handlers/ShaderIssueHandler.cs b/Engine/Source/Programs/Horde/Horde.Build/Issues/Handlers/ShaderIssueHandler.cs
new file mode 100644
--- /dev/null
+++ b/Engine/Source/Programs/Horde/Horde.Build/Issues/Handlers/ShaderIssueHandler.cs
@@ -0,0 +1,72 @@
+// Copyright Epic Games, Inc. All Rights Reserved.
+
+using System;
+using System.Collections.Generic;
+using System.Diagnostics.CodeAnalysis;
+using System.Linq;
+using EpicGames.Core;
+using Horde.Build.Jobs;
+using Horde.Build.Jobs.Graphs;
+using Horde.Build.Logs;
+using Microsoft.Extensions.Logging;
+using MongoDB.Driver;
+
+namespace Horde.Build.Issues.Handlers
+{
+	/// <summary>
+	/// Instance of a particular shader compile error
+	/// </summary>
+	class ShaderIssueHandler : SourceFileIssueHandler
+	{
+		/// <inheritdoc/>
+		public override string Type => "Shader";
+
+		/// <inheritdoc/>
+		public override int Priority => 10;
+
+		/// <summary>
+		/// Determines if the given event id matches
+		/// </summary>
+		/// <param name="eventId">The event id to compare</param>
+		/// <returns>True if the given event id matches</returns>
+		public static bool IsMatchingEventId(EventId eventId)
+		{
+			return eventId == KnownLogEvents.Engine_ShaderCompiler;
+		}
+
+		static bool IsMaskedEventId(EventId id) => id == KnownLogEvents.Generic || id == KnownLogEvents.ExitCode || id == KnownLogEvents.Systemic_Xge_BuildFailed || id == KnownLogEvents.Engine_Crash || id == KnownLogEvents.AutomationTool_CrashExitCode || id == KnownLogEvents.Engine_AppError;
+
+		/// <inheritdoc/>
+		public override void TagEvents(IJob job, INode node, IReadOnlyNodeAnnotations annotations, IReadOnlyList<IssueEvent> stepEvents)
+		{
+			bool hasMatches = false;
+			foreach (IssueEvent stepEvent in stepEvents)
+			{
+				if (stepEvent.EventId.HasValue)
+				{
+					EventId eventId = stepEvent.EventId.Value;
+					if (IsMatchingEventId(eventId))
+					{
+						List<string> newFileNames = new List<string>();
+						GetSourceFiles(stepEvent.EventData, newFileNames);
+
+						stepEvent.Fingerprint = new NewIssueFingerprint(Type, newFileNames, null, null);
+						hasMatches = true;
+					}
+					else if (hasMatches && IsMaskedEventId(eventId))
+					{
+						stepEvent.Ignored = true;
+					}
+				}
+			}
+		}
+
+		/// <inheritdoc/>
+		public override string GetSummary(IIssueFingerprint fingerprint, IssueSeverity severity)
+		{
+			string level = (severity == IssueSeverity.Warning) ? "warnings" : "errors";
+			string list = StringUtils.FormatList(fingerprint.Keys.Where(x => !x.StartsWith(NotePrefix, StringComparison.Ordinal)).ToArray(), 2);
+			return $"Shader compile {level} in {list}";
+		}
+	}
+}
diff --git a/Engine/Source/Programs/Shared/EpicGames.Core/KnownLogEvents.cs b/Engine/Source/Programs/Shared/EpicGames.Core/KnownLogEvents.cs
--- a/Engine/Source/Programs/Shared/EpicGames.Core/KnownLogEvents.cs
+++ b/Engine/Source/Programs/Shared/EpicGames.Core/KnownLogEvents.cs
@@ -80,6 +80,16 @@ namespace EpicGames.Core
 		/// </summary>
 		public static EventId Engine_Localization { get; } = new EventId(304);
 
+		/// <summary>
+		/// Engine: appErrorf called
+		/// </summary>
+		public static EventId Engine_AppError { get; } = new EventId(305);
+
+		/// <summary>
+		/// Engine: Shader compiler output
+		/// </summary>
+		public static EventId Engine_ShaderCompiler { get; } = new EventId(310);
+
 		/// <summary>
 		/// UAT error
 		/// </summary>
