Horde: Stop building/running Horde.Storage within Horde.Build hosts. Both projects will be developed separately for now.

diff --git a/Engine/Source/Programs/Horde/Horde.Build/Commands/ServerCommand.cs b/Engine/Source/Programs/Horde/Horde.Build/Commands/ServerCommand.cs
--- a/Engine/Source/Programs/Horde/Horde.Build/Commands/ServerCommand.cs
+++ b/Engine/Source/Programs/Horde/Horde.Build/Commands/ServerCommand.cs
@@ -41,21 +41,8 @@ namespace Horde.Build.Commands
 		{
 			using (X509Certificate2? grpcCertificate = ReadGrpcCertificate(_hordeSettings))
 			{
-				List<IHost> hosts = new List<IHost>();
-				hosts.Add(CreateHostBuilderWithCert(_args, _config, _hordeSettings, grpcCertificate).Build());
-#if WITH_HORDE_STORAGE
-				IHostBuilder storageHostBuilder = Horde.Storage.Program.CreateHostBuilder(_args);
-				storageHostBuilder.ConfigureWebHostDefaults(builder =>
-				{
-					builder.ConfigureKestrel(options =>
-					{
-						options.ListenAnyIP(57000);
-						options.ListenAnyIP(57001, configure => configure.UseHttps());
-					});
-				});
-				hosts.Add(storageHostBuilder.Build());
-#endif
-				await Task.WhenAll(hosts.Select(x => x.RunAsync()));
+				using IHost host = CreateHostBuilderWithCert(_args, _config, _hordeSettings, grpcCertificate).Build();
+				await host.RunAsync();
 				return 0;
 			}
 		}
diff --git a/Engine/Source/Programs/Horde/Horde.Build/Horde.Build.csproj b/Engine/Source/Programs/Horde/Horde.Build/Horde.Build.csproj
--- a/Engine/Source/Programs/Horde/Horde.Build/Horde.Build.csproj
+++ b/Engine/Source/Programs/Horde/Horde.Build/Horde.Build.csproj
@@ -1,4 +1,4 @@
-﻿<Project Sdk="Microsoft.NET.Sdk.Web">
+<Project Sdk="Microsoft.NET.Sdk.Web">
 
   <PropertyGroup>
     <TargetFramework>net6.0</TargetFramework>
@@ -10,14 +10,6 @@
     <RunAnalyzersDuringBuild>False</RunAnalyzersDuringBuild>
   </PropertyGroup>
 
-  <PropertyGroup Condition="'$(WithHordeStorage)' == ''">
-    <WithHordeStorage>true</WithHordeStorage>
-  </PropertyGroup>
-
-  <PropertyGroup Condition="$(WithHordeStorage)">
-    <DefineConstants>$(DefineConstants);WITH_HORDE_STORAGE</DefineConstants>
-  </PropertyGroup>
-
   <ItemGroup>
     <PackageReference Include="AWSSDK.AutoScaling" Version="3.7.10.43" />
     <PackageReference Include="AWSSDK.CloudWatch" Version="3.7.6.3" />
@@ -113,10 +105,6 @@
     <ProjectReference Include="..\..\Shared\EpicGames.Slack\EpicGames.Slack.csproj" />
   </ItemGroup>
 
-  <ItemGroup Condition="$(WithHordeStorage)">
-    <ProjectReference Include="..\HordeStorage\Horde.Storage\Horde.Storage.csproj" />
-  </ItemGroup>
-
   <ItemGroup>
     <Content Include="..\..\..\ThirdParty\Perforce\P4Api.Net\native\windows\release\p4bridge.*" Condition=" '$(OS)' == 'Windows_NT' ">
       <CopyToOutputDirectory>PreserveNewest</CopyToOutputDirectory>
