Horde: Use read-only clients for replication, to reduce thrashing in the P4 journal.

diff --git a/Engine/Source/Programs/Horde/Horde.Build/Perforce/ReplicationService.cs b/Engine/Source/Programs/Horde/Horde.Build/Perforce/ReplicationService.cs
--- a/Engine/Source/Programs/Horde/Horde.Build/Perforce/ReplicationService.cs
+++ b/Engine/Source/Programs/Horde/Horde.Build/Perforce/ReplicationService.cs
@@ -1032,11 +1032,12 @@ namespace Horde.Build.Perforce
 
 				InfoRecord serverInfo = await perforce.GetInfoAsync(InfoOptions.ShortOutput);
 
-				ClientRecord newClient = new ClientRecord($"Horde.Build_{serverInfo.ClientHost}_{stream.Id}", perforce.Settings.UserName, "/p4/");
+				ClientRecord newClient = new ClientRecord($"Horde.Build_Rep_{serverInfo.ClientHost}_{stream.Id}", perforce.Settings.UserName, "/p4/");
 				newClient.Description = "Created to mirror Perforce content to Horde Storage";
 				newClient.Owner = perforce.Settings.UserName;
 				newClient.Host = serverInfo.ClientHost;
 				newClient.Stream = stream.Config.ReplicationStream ?? stream.Name;
+				newClient.Type = "readonly";
 				await perforce.CreateClientAsync(newClient);
 				_logger.LogInformation("Created client {ClientName} for {StreamName}", newClient.Name, stream.Name);
 
