Horde: Add tests for RedisTaskScheduler

diff --git a/Engine/Source/Programs/Horde/Horde.Build.Tests/Compute/RedisTaskSchedulerTest.cs b/Engine/Source/Programs/Horde/Horde.Build.Tests/Compute/RedisTaskSchedulerTest.cs
new file mode 100644
--- /dev/null
+++ b/Engine/Source/Programs/Horde/Horde.Build.Tests/Compute/RedisTaskSchedulerTest.cs
@@ -0,0 +1,90 @@
+// Copyright Epic Games, Inc. All Rights Reserved.
+
+using System.Threading;
+using System.Threading.Tasks;
+using Horde.Build.Compute;
+using Horde.Build.Server;
+using Microsoft.Extensions.Logging;
+using Microsoft.VisualStudio.TestTools.UnitTesting;
+using StackExchange.Redis;
+
+namespace Horde.Build.Tests.Compute
+{
+	[TestClass]
+	public class RedisTaskSchedulerTest : TestSetup
+	{
+		private readonly RedisTaskScheduler<string, string> _scheduler;
+
+		public RedisTaskSchedulerTest()
+		{
+			using ILoggerFactory loggerFactory = LoggerFactory.Create(builder => { builder.AddConsole(); });
+			RedisService redisService = GetRedisServiceSingleton();
+			_scheduler = new RedisTaskScheduler<string, string>(redisService.ConnectionPool, new RedisKey("myBaseKey"), loggerFactory.CreateLogger<RedisTaskScheduler<string, string>>());
+		}
+
+		[TestMethod]
+		public async Task EnqueueAndDequeueTask()
+		{
+			await _scheduler.EnqueueAsync("queue", "task1", true);
+			string? dequeuedTask = await _scheduler.DequeueAsync("queue");
+			Assert.AreEqual("task1", dequeuedTask);
+		}
+		
+		[TestMethod]
+		[Ignore("Currently fails as nothing picks up tasks(?)")]
+		public async Task DequeueWithPredicateFromMiddleOfQueue()
+		{
+			using CancellationTokenSource cts = new (2000);
+			await _scheduler.EnqueueAsync("queue", "task1", true);
+			await _scheduler.EnqueueAsync("queue", "task2", true);
+			await _scheduler.EnqueueAsync("queue", "task3", true);
+			Task<(string, string)?> dequeueTask = await _scheduler.DequeueAsync((x) => ValueTask.FromResult(x == "task2"), cts.Token);
+
+			(string queueId, string task) = (await dequeueTask.WaitAsync(cts.Token))!.Value;
+
+			Assert.AreEqual("queue", queueId);
+			Assert.AreEqual("task2", task);
+		}
+		
+		[TestMethod]
+		public async Task EnqueueTaskAtFront()
+		{
+			await _scheduler.EnqueueAsync("queue", "task1", true);
+			await _scheduler.EnqueueAsync("queue", "task2", true);
+			Assert.AreEqual("task2", await _scheduler.DequeueAsync("queue"));
+		}
+		
+		[TestMethod]
+		public async Task EnqueueTaskAtBack()
+		{
+			await _scheduler.EnqueueAsync("queue", "task1", false);
+			await _scheduler.EnqueueAsync("queue", "task2", false);
+			Assert.AreEqual("task1", await _scheduler.DequeueAsync("queue"));
+		}
+		
+		[TestMethod]
+		public async Task EnqueueTasksOnSeparateQueues()
+		{
+			await _scheduler.EnqueueAsync("queue1", "task1", true);
+			await _scheduler.EnqueueAsync("queue2", "task2", true);
+			Assert.AreEqual("task1", await _scheduler.DequeueAsync("queue1"));
+			Assert.AreEqual("task2", await _scheduler.DequeueAsync("queue2"));
+		}
+		
+		[TestMethod]
+		public async Task DequeueEmptyQueue()
+		{
+			Assert.IsNull(await _scheduler.DequeueAsync("queue"));
+		}
+		
+		[TestMethod]
+		public async Task DequeuingMakesQueueActive()
+		{
+			await _scheduler.EnqueueAsync("queue", "task1", true);
+			Assert.AreEqual(1, (await _scheduler.GetInactiveQueuesAsync()).Count);
+			
+			await _scheduler.DequeueAsync("queue");
+			Assert.AreEqual(0, (await _scheduler.GetInactiveQueuesAsync()).Count);
+		}
+	}
+}
\ No newline at end of file
