Branch snapshot for CL 21539030

[CL 21539030 in ue5-main branch]

diff --git a/Engine/Source/Programs/Horde/Horde.Agent.Tests/remote-exec-test-bin/Dockerfile b/Engine/Source/Programs/Horde/Horde.Agent.Tests/remote-exec-test-bin/Dockerfile
new file mode 100644
--- /dev/null
+++ b/Engine/Source/Programs/Horde/Horde.Agent.Tests/remote-exec-test-bin/Dockerfile
@@ -0,0 +1,26 @@
+FROM alpine:3.13 as builder
+
+# Install Zig and dependencies
+RUN apk update && apk add curl xz libc-dev samurai git cmake
+ARG ZIGVER=0.9.1
+RUN mkdir -p /deps
+WORKDIR /deps
+RUN curl https://ziglang.org/deps/zig+llvm+lld+clang-$(uname -m)-linux-musl-$ZIGVER.tar.xz  -O && \
+    tar xf zig+llvm+lld+clang-$(uname -m)-linux-musl-$ZIGVER.tar.xz && \
+    mv zig+llvm+lld+clang-$(uname -m)-linux-musl-$ZIGVER/ local/
+
+# Build the binary for different platforms
+ADD remote-exec-test-bin.c /
+RUN /deps/local/bin/zig cc -o /remote-exec-test-x86_64-windows.exe /remote-exec-test-bin.c -target x86_64-windows-gnu
+RUN /deps/local/bin/zig cc -o /remote-exec-test-x86_64-linux /remote-exec-test-bin.c -target x86_64-linux-gnu
+RUN /deps/local/bin/zig cc -o /remote-exec-test-x86_64-macos /remote-exec-test-bin.c -target x86_64-macos-gnu
+RUN /deps/local/bin/zig cc -o /remote-exec-test-aarch64-linux /remote-exec-test-bin.c -target aarch64-linux-gnu
+RUN /deps/local/bin/zig cc -o /remote-exec-test-aarch64-macos /remote-exec-test-bin.c -target aarch64-macos-gnu
+
+# Copy binaries to a scratch image so that they can be cleanly be copied back to host system
+FROM scratch
+COPY --from=builder /remote-exec-test-x86_64-windows.exe /
+COPY --from=builder /remote-exec-test-x86_64-linux /
+COPY --from=builder /remote-exec-test-x86_64-macos /
+COPY --from=builder /remote-exec-test-aarch64-linux /
+COPY --from=builder /remote-exec-test-aarch64-macos /
\ No newline at end of file
diff --git a/Engine/Source/Programs/Horde/Horde.Build/Dockerfile b/Engine/Source/Programs/Horde/Horde.Build/Dockerfile
new file mode 100644
--- /dev/null
+++ b/Engine/Source/Programs/Horde/Horde.Build/Dockerfile
@@ -0,0 +1,66 @@
+#
+# NOTE: This Dockerfile must be built with the Engine directory as CWD.
+#       In the Horde build script we stage the required files before running:
+#
+#         docker build -f Source/Programs/Horde/HordeServer/Dockerfile . 
+#
+#       Running directly from the source tree will cause a lot of unnecessary
+#       files to be uploaded to Docker before building.
+
+# Build image
+FROM mcr.microsoft.com/dotnet/sdk:6.0 AS build-env
+WORKDIR /app
+COPY Source/Programs/Horde/nuget.config ./Source/Programs/Horde/nuget.config
+RUN apt-get update && apt-get install -y build-essential && apt-get install -y curl
+
+# Pre-built binaries for Redis are hard to come by. But it's small and fast enough to build it here.
+RUN cd /tmp && \
+	wget --no-verbose http://download.redis.io/releases/redis-5.0.9.tar.gz && \
+	tar xzf redis-5.0.9.tar.gz && \
+	cd redis-5.0.9 && \
+	make && \
+	mv src/redis-server /usr/local/bin/
+
+# Since the .deb does not install in this image, just download it and extract the static binary
+RUN wget https://repo.mongodb.org/apt/debian/dists/buster/mongodb-org/4.4/main/binary-amd64/mongodb-org-server_4.4.0_amd64.deb && \
+	dpkg -x mongodb-org-server_4.4.0_amd64.deb /tmp/mongodb && \
+	cp /tmp/mongodb/usr/bin/mongod /usr/local/bin/mongod
+
+# Restore all the C# projects
+COPY Source/Programs/Shared/EpicGames.Core/*.csproj ./Source/Programs/Shared/EpicGames.Core/
+COPY Source/Programs/Shared/EpicGames.Horde/*.csproj ./Source/Programs/Shared/EpicGames.Horde/
+COPY Source/Programs/Shared/EpicGames.IoHash/*.csproj ./Source/Programs/Shared/EpicGames.IoHash/
+COPY Source/Programs/Shared/EpicGames.Perforce/*.csproj ./Source/Programs/Shared/EpicGames.Perforce/
+COPY Source/Programs/Shared/EpicGames.Perforce.Managed/*.csproj ./Source/Programs/Shared/EpicGames.Perforce.Managed/
+COPY Source/Programs/Shared/EpicGames.Redis/*.csproj ./Source/Programs/Shared/EpicGames.Redis/
+COPY Source/Programs/Shared/EpicGames.Serialization/*.csproj ./Source/Programs/Shared/EpicGames.Serialization/
+COPY Source/Programs/Shared/EpicGames.Serialization.Tests/*.csproj ./Source/Programs/Shared/EpicGames.Serialization.Tests/
+COPY Source/Programs/Horde/Horde.Agent/*.csproj ./Source/Programs/Horde/Horde.Agent/
+COPY Source/Programs/Horde/Horde.Build/*.csproj ./Source/Programs/Horde/Horde.Build/
+COPY Source/Programs/Horde/Horde.Build.Tests/*.csproj ./Source/Programs/Horde/Horde.Build.Tests/
+COPY Source/ThirdParty/Perforce/P4Api.Net/ ./Source/ThirdParty/Perforce/P4Api.Net/
+RUN dotnet restore Source/Programs/Horde/Horde.Build/Horde.Build.csproj -p:WithHordeStorage=false
+
+# Build the server
+COPY Binaries/DotNET/EpicGames.Perforce.Native/ Binaries/DotNET/EpicGames.Perforce.Native/
+COPY Source/Programs/Shared/ Source/Programs/Shared/
+COPY Source/Programs/Horde/Horde.Build/ Source/Programs/Horde/Horde.Build/
+COPY Source/Programs/AutomationTool/AutomationUtils/Matchers/ Source/Programs/AutomationTool/AutomationUtils/Matchers/
+COPY Source/Programs/UnrealBuildTool/Matchers/ Source/Programs/UnrealBuildTool/Matchers/
+ARG msbuild_args
+RUN dotnet publish Source/Programs/Horde/Horde.Build/Horde.Build.csproj -c Release -o out -p:WithHordeStorage=false $msbuild_args -p:RunAnalyzersDuringBuild=True
+
+# Run tests
+COPY Source/Programs/Horde/Horde.Agent/ Source/Programs/Horde/Horde.Agent/
+COPY Source/Programs/Horde/Horde.Build.Tests/ Source/Programs/Horde/Horde.Build.Tests/
+COPY Source/Programs/Shared/EpicGames.Serialization.Tests/ Source/Programs/Shared/EpicGames.Serialization.Tests/
+RUN redis-server --save "" --appendonly no --daemonize yes && \
+	mongod --noauth --quiet --fork --dbpath /tmp/mongodb --logpath /tmp/mongod.log && \
+	dotnet test Source/Programs/Horde/Horde.Build.Tests/Horde.Build.Tests.csproj --blame-hang-timeout 2m -p:WithHordeStorage=false && \
+	dotnet test Source/Programs/Shared/EpicGames.Serialization.Tests/EpicGames.Serialization.Tests.csproj --blame-hang-timeout 2m -p:WithHordeStorage=false
+
+# Create the runtime image
+FROM mcr.microsoft.com/dotnet/aspnet:6.0 AS runtime
+WORKDIR /app
+COPY --from=build-env /app/out .
+ENTRYPOINT [ "dotnet", "Horde.Build.dll" ]
diff --git a/Engine/Source/Programs/Horde/HordeStorage/Benchmarks/benchmark-suite/dockerfile b/Engine/Source/Programs/Horde/HordeStorage/Benchmarks/benchmark-suite/dockerfile
new file mode 100644
--- /dev/null
+++ b/Engine/Source/Programs/Horde/HordeStorage/Benchmarks/benchmark-suite/dockerfile
@@ -0,0 +1,17 @@
+FROM ubuntu AS build
+
+WORKDIR /app
+ENV DEBIAN_FRONTEND=noninteractive
+RUN apt-get update
+RUN apt-get install -y python3 pip
+
+COPY vegeta_12.8.4_linux_386.tar .
+RUN tar -xf vegeta_12.8.4_linux_386.tar
+
+COPY requirements.txt .
+RUN python3 -m pip install -r requirements.txt
+COPY *.py ./
+
+# Labels used by datadog to understand the log output we produce
+LABEL "com.datadoghq.ad.logs"='[{"source": "python", "service": "jupiter_benchmark"}]'
+ENTRYPOINT ["python3", "main.py"]
\ No newline at end of file
diff --git a/Engine/Source/Programs/Horde/HordeStorage/dockerfile b/Engine/Source/Programs/Horde/HordeStorage/dockerfile
new file mode 100644
--- /dev/null
+++ b/Engine/Source/Programs/Horde/HordeStorage/dockerfile
@@ -0,0 +1,62 @@
+FROM mcr.microsoft.com/dotnet/sdk:6.0 AS build
+
+WORKDIR /app
+
+ARG Project
+ARG UE_P4_CHANGELIST=0
+
+# Copy csprojs for nuget restore
+COPY Programs/Shared/EpicGames.Core/*.csproj ./Programs/Shared/EpicGames.Core/
+COPY Programs/Shared/EpicGames.Horde/*.csproj ./Programs/Shared/EpicGames.Horde/
+COPY Programs/Shared/EpicGames.IoHash/*.csproj ./Programs/Shared/EpicGames.IoHash/
+COPY Programs/Shared/EpicGames.Serialization/*.csproj ./Programs/Shared/EpicGames.Serialization/
+COPY Programs/Horde/HordeStorage/Jupiter.Common/*.csproj ./Programs/Horde/HordeStorage/Jupiter.Common/
+COPY Programs/Horde/HordeStorage/${Project}/*.csproj ./Programs/Horde/HordeStorage/${Project}/
+RUN dotnet restore Programs/Horde/HordeStorage/${Project}/${Project}.csproj
+# all nuget prereqs are now done
+
+# copy the source we need
+COPY Programs/Shared/ Programs/Shared/
+COPY Programs/Horde/HordeStorage/ Programs/Horde/HordeStorage/
+
+# build and create the dll
+RUN dotnet publish -c Release -o out Programs/Horde/HordeStorage/${Project}/${Project}.csproj
+
+# create the actual deployable image
+FROM mcr.microsoft.com/dotnet/aspnet:6.0 AS runtime
+# add curl for the other commands below, plus it is useful to have
+RUN apt-get update && apt-get install curl -y
+
+# Add the datadog tracer
+RUN curl -LO [REDACTED_PROFILE_URL]/dd-trace-dotnet/releases/download/v2.11.0/datadog-dotnet-apm_2.11.0_amd64.deb
+RUN dpkg -i ./datadog-dotnet-apm_2.11.0_amd64.deb
+RUN /opt/datadog/createLogPath.sh
+
+# install remote debugger
+RUN apt-get install unzip && \
+    curl -sSL https://aka.ms/getvsdbgsh | /bin/sh /dev/stdin -v latest -l /vsdbg
+# install ps for remote debugging with vscode
+RUN apt-get -y install procps
+
+WORKDIR /app
+
+COPY --from=build /app/out ./
+ARG Project
+ENV ASSEMBLY_PROJECT_NAME ${Project}
+
+# datadog APM requirments
+ENV CORECLR_ENABLE_PROFILING=1
+ENV CORECLR_PROFILER={846F5F1C-F9AE-4B07-969E-05C26BC060D8}
+ENV CORECLR_PROFILER_PATH=/opt/datadog/Datadog.Trace.ClrProfiler.Native.so
+ENV DD_INTEGRATIONS=/opt/datadog/integrations.json
+ENV DD_DOTNET_TRACER_HOME=/opt/datadog
+# enable datadog continus profiler
+ENV LD_PRELOAD=/opt/datadog/continuousprofiler/Datadog.Linux.ApiWrapper.x64.so
+
+ENV DD_LOGS_INJECTION=true
+ENV DD_RUNTIME_METRICS_ENABLED=true
+# do not consider http 400-499 statuses errors
+ENV DD_HTTP_CLIENT_ERROR_STATUSES=500-599
+
+
+ENTRYPOINT dotnet $ASSEMBLY_PROJECT_NAME.dll
\ No newline at end of file
