Horde: Add additional job information to device telemetry
#rnx

diff --git a/Engine/Source/Programs/Horde/Horde.Build.Tests/DeviceServiceTests.cs b/Engine/Source/Programs/Horde/Horde.Build.Tests/DeviceServiceTests.cs
--- a/Engine/Source/Programs/Horde/Horde.Build.Tests/DeviceServiceTests.cs
+++ b/Engine/Source/Programs/Horde/Horde.Build.Tests/DeviceServiceTests.cs
@@ -136,6 +136,8 @@ namespace Horde.Build.Tests
 			// create a reservation
 			GetLegacyReservationResponse reservation = ResultToValue(await DeviceController!.CreateDeviceReservationV1Async(request));			
 			Assert.AreEqual(1, reservation.DeviceNames.Length);
+			Assert.AreEqual("hello2", reservation.JobName);
+			Assert.AreEqual("abcd", reservation.StepId);
 
 			// get the device in the reservation, and make sure it is the right platform
 			GetLegacyDeviceResponse device = ResultToValue(await DeviceController!.GetDeviceV1Async(reservation.DeviceNames[0]));
@@ -155,6 +157,7 @@ namespace Horde.Build.Tests
 			Assert.AreEqual(telemetry[0].Telemetry.Count, 1);
 			Assert.AreEqual(telemetry[0].Telemetry[0].StreamId, "ue5-main");
 			Assert.AreEqual(telemetry[0].Telemetry[0].StepId, "abcd");
+			Assert.AreEqual(telemetry[0].Telemetry[0].JobName, "hello2");
 		}
 
 		[TestMethod]
diff --git a/Engine/Source/Programs/Horde/Horde.Build/Devices/DeviceCollection.cs b/Engine/Source/Programs/Horde/Horde.Build/Devices/DeviceCollection.cs
--- a/Engine/Source/Programs/Horde/Horde.Build/Devices/DeviceCollection.cs
+++ b/Engine/Source/Programs/Horde/Horde.Build/Devices/DeviceCollection.cs
@@ -113,6 +113,12 @@ namespace Horde.Build.Devices
 			[BsonIgnoreIfNull]
 			public string? StepId { get; set; }
 
+			[BsonIgnoreIfNull]
+			public string? JobName { get; set; }
+
+			[BsonIgnoreIfNull]
+			public string? StepName { get; set; }
+
 			/// <summary>
 			/// Reservations held by a user, requires a token like download code
 			/// </summary>
@@ -151,7 +157,7 @@ namespace Horde.Build.Devices
 
 			}
 
-			public DeviceReservationDocument(ObjectId id, DevicePoolId poolId, List<DeviceId> devices, List<string> requestedDevicePlatforms, DateTime createTimeUtc, string? hostname, string? reservationDetails, string? streamId, string? jobId, string? stepId)
+			public DeviceReservationDocument(ObjectId id, DevicePoolId poolId, List<DeviceId> devices, List<string> requestedDevicePlatforms, DateTime createTimeUtc, string? hostname, string? reservationDetails, string? streamId, string? jobId, string? stepId, string? jobName, string? stepName)
 			{
 				Id = id;
 				PoolId = poolId;
@@ -164,6 +170,8 @@ namespace Horde.Build.Devices
 				StreamId = streamId;
 				JobId = jobId;
 				StepId = stepId;
+				JobName = jobName;
+				StepName = stepName;
 
 				LegacyGuid = Guid.NewGuid().ToString();
 			}
@@ -281,12 +289,24 @@ namespace Horde.Build.Devices
 			[BsonIgnoreIfNull,BsonElement("job")]
 			public string? JobId { get; set; }
 
+			/// <summary>
+			/// The job name
+			/// </summary>
+			[BsonIgnoreIfNull]
+			public string? JobName { get; set; }
+
 			/// <summary>
 			/// The job's step id
 			/// </summary>
 			[BsonIgnoreIfNull, BsonElement("step")]
 			public string? StepId { get; set; }
 
+			/// <summary>
+			/// The step name
+			/// </summary>
+			[BsonIgnoreIfNull]
+			public string? StepName { get; set; }
+
 			/// <summary>
 			/// Reservation Id (transient, reservations are deleted upon expiration)
 			/// </summary>
@@ -316,7 +336,7 @@ namespace Horde.Build.Devices
 			{
 			}
 
-			public DeviceTelemetryDocument(DeviceId deviceId, ObjectId? reservationId = null, DateTime? reservationStartTime = null, string? streamId = null, string? jobId = null, string? stepId = null)
+			public DeviceTelemetryDocument(DeviceId deviceId, ObjectId? reservationId = null, DateTime? reservationStartTime = null, string? streamId = null, string? jobId = null, string? stepId = null, string? jobName = null, string? stepName = null)
 			{
 				TelemetryId = ObjectId.GenerateNewId();
 				CreateTimeUtc = DateTime.UtcNow;
@@ -328,6 +348,8 @@ namespace Horde.Build.Devices
 				StreamId = streamId;
 				JobId = jobId;
 				StepId = stepId;
+				JobName = jobName;
+				StepName = stepName;
 			}
 
 		}
@@ -772,7 +794,7 @@ namespace Horde.Build.Devices
 		}
 
 		/// <inheritdoc/>
-		public async Task<IDeviceReservation?> TryAddReservationAsync(DevicePoolId poolId, List<DeviceRequestData> request, string? hostname, string? reservationDetails, string? streamId, string? jobId, string? stepId)
+		public async Task<IDeviceReservation?> TryAddReservationAsync(DevicePoolId poolId, List<DeviceRequestData> request, string? hostname, string? reservationDetails, string? streamId, string? jobId, string? stepId, string? jobName, string? stepName)
 		{
 
 			if (request.Count == 0)
@@ -898,14 +920,14 @@ namespace Horde.Build.Devices
 			List<string> requestedPlatforms = deviceIds.Select(x => platformRequestMap[x]).ToList();
 
 			// Create new reservation
-			DeviceReservationDocument newReservation = new DeviceReservationDocument(ObjectId.GenerateNewId(), poolId, deviceIds, requestedPlatforms, reservationTimeUtc, hostname, reservationDetails, streamId, jobId, stepId);
+			DeviceReservationDocument newReservation = new DeviceReservationDocument(ObjectId.GenerateNewId(), poolId, deviceIds, requestedPlatforms, reservationTimeUtc, hostname, reservationDetails, streamId, jobId, stepId, jobName, stepName);
 			await _reservations.InsertOneAsync(newReservation);
 
 			// Create device telemetry data for reservation
 			List<DeviceTelemetryDocument> telemetry = new List<DeviceTelemetryDocument>();
 			foreach (DeviceId deviceId in deviceIds)
 			{
-				telemetry.Add(new DeviceTelemetryDocument(deviceId, newReservation.Id, newReservation.CreateTimeUtc, streamId, jobId, stepId));
+				telemetry.Add(new DeviceTelemetryDocument(deviceId, newReservation.Id, newReservation.CreateTimeUtc, streamId, jobId, stepId, jobName, stepName));
 			}
 
 			if (telemetry.Count > 0)
diff --git a/Engine/Source/Programs/Horde/Horde.Build/Devices/DeviceMessages.cs b/Engine/Source/Programs/Horde/Horde.Build/Devices/DeviceMessages.cs
--- a/Engine/Source/Programs/Horde/Horde.Build/Devices/DeviceMessages.cs
+++ b/Engine/Source/Programs/Horde/Horde.Build/Devices/DeviceMessages.cs
@@ -520,6 +520,16 @@ namespace Horde.Build.Devices
 		/// </summary>
 		public string? StepId { get; set; }
 
+		/// <summary>
+		/// Job mame holding reservation
+		/// </summary>
+		public string? JobName { get; set; }
+
+		/// <summary>
+		/// Job step holding reservation
+		/// </summary>
+		public string? StepName { get; set; }
+
 		/// <summary>
 		/// Reservations held by a user, requires a token
 		/// </summary>
@@ -567,11 +577,21 @@ namespace Horde.Build.Devices
 		/// </summary>
 		public string? JobId { get; set; }
 
+		/// <summary>
+		/// The job name which utilized device
+		/// </summary>
+		public string? JobName { get; set; }
+
 		/// <summary>
 		/// The job's step id
 		/// </summary>
 		public string? StepId { get; set; }
 
+		/// <summary>
+		/// The job name which utilized device
+		/// </summary>
+		public string? StepName { get; set; }
+
 		/// <summary>
 		/// If this telemetry has a reservation, the start time of the reservation
 		/// </summary>
@@ -597,6 +617,8 @@ namespace Horde.Build.Devices
 			StreamId = data.StreamId;
 			JobId = data.JobId;
 			StepId = data.StepId;
+			JobName = data.JobName;
+			StepName = data.StepName;
 			ReservationStartUtc = data.ReservationStartUtc;
 			ReservationFinishUtc = data.ReservationFinishUtc;
 			ProblemTimeUtc = data.ProblemTimeUtc;
@@ -814,6 +836,16 @@ namespace Horde.Build.Devices
 		/// </summary>
 		public string? StepId { get; set; } = null;
 
+		/// <summary>
+		/// The job name of reservation request 
+		/// </summary>
+		public string? JobName { get; set; } = null;
+
+		/// <summary>
+		/// The step name of reservation request 
+		/// </summary>
+		public string? StepName { get; set; } = null;
+
 		/// <summary>
 		/// The legacy guid of the reservation
 		/// </summary>
diff --git a/Engine/Source/Programs/Horde/Horde.Build/Devices/DeviceService.cs b/Engine/Source/Programs/Horde/Horde.Build/Devices/DeviceService.cs
--- a/Engine/Source/Programs/Horde/Horde.Build/Devices/DeviceService.cs
+++ b/Engine/Source/Programs/Horde/Horde.Build/Devices/DeviceService.cs
@@ -407,17 +407,35 @@ namespace Horde.Build.Devices
 		public async Task<IDeviceReservation?> TryCreateReservationAsync(DevicePoolId pool, List<DeviceRequestData> request, string? hostname = null, string? reservationDetails = null, string? jobId = null, string? stepId = null)
 		{
 			string? streamId = null;
+			string? jobName = null;
+			string? stepName = null;
 
 			if (jobId != null)
 			{
 				IJob? job = await _jobService.GetJobAsync(new JobId(jobId));
 				if (job != null)
-				{
+				{					
 					streamId = job.StreamId.ToString();
+					jobName = job.Name;
+
+					if (stepId != null)
+					{
+						SubResourceId id = SubResourceId.Parse(stepId);
+						IGraph graph = await _jobService.GetGraphAsync(job);
+						foreach (IJobStepBatch batch in job.Batches)
+						{
+							IJobStep? step;							
+							if (batch.TryGetStep(id, out step))
+							{
+								stepName = graph.Groups[batch.GroupIdx].Nodes[step.NodeIdx].Name;
+								break;
+							}
+						}
+					}
 				}
 			}
 
-			return await _devices.TryAddReservationAsync(pool, request, hostname, reservationDetails, streamId, jobId, stepId);
+			return await _devices.TryAddReservationAsync(pool, request, hostname, reservationDetails, streamId, jobId, stepId, jobName, stepName);
 		}
 
 		/// <summary>
diff --git a/Engine/Source/Programs/Horde/Horde.Build/Devices/DevicesController.cs b/Engine/Source/Programs/Horde/Horde.Build/Devices/DevicesController.cs
--- a/Engine/Source/Programs/Horde/Horde.Build/Devices/DevicesController.cs
+++ b/Engine/Source/Programs/Horde/Horde.Build/Devices/DevicesController.cs
@@ -656,6 +656,8 @@ namespace Horde.Build.Devices
 					Devices = reservation.Devices.Select(x => x.ToString()).ToList(),
 					JobId = reservation.JobId?.ToString(),
 					StepId = reservation.StepId?.ToString(),
+					JobName = reservation.JobName?.ToString(),
+					StepName = reservation.StepName?.ToString(),
 					UserId = reservation.UserId?.ToString(),
 					Hostname = reservation.Hostname,
 					ReservationDetails = reservation.ReservationDetails,
@@ -965,6 +967,10 @@ namespace Horde.Build.Devices
 			response.HostName = request.Hostname;
 			response.StartDateTime = reservation.CreateTimeUtc.ToString("O", System.Globalization.CultureInfo.InvariantCulture);
 			response.Duration = $"{request.Duration}";
+			response.JobId = reservation.JobId;
+			response.StepId = reservation.StepId;
+			response.JobName = reservation.JobName;
+			response.StepName = reservation.StepName;
 
 			return new JsonResult(response, new JsonSerializerOptions() { PropertyNamingPolicy = null });
 
diff --git a/Engine/Source/Programs/Horde/Horde.Build/Devices/IDevice.cs b/Engine/Source/Programs/Horde/Horde.Build/Devices/IDevice.cs
--- a/Engine/Source/Programs/Horde/Horde.Build/Devices/IDevice.cs
+++ b/Engine/Source/Programs/Horde/Horde.Build/Devices/IDevice.cs
@@ -49,6 +49,16 @@ namespace Horde.Build.Devices
 		/// </summary>
 		public string? StepId { get; }
 
+		/// <summary>
+		/// Job name holding reservation
+		/// </summary>
+		public string? JobName { get; }
+
+		/// <summary>
+		/// Job step name holding reservation
+		/// </summary>
+		public string? StepName { get; }
+
 		/// <summary>
 		/// Reservations held by a user, requires a token
 		/// </summary>
@@ -226,11 +236,21 @@ namespace Horde.Build.Devices
 		/// </summary>
 		public string? JobId { get; }
 
+		/// <summary>
+		/// The job name which utilized device
+		/// </summary>
+		public string? JobName { get; }
+
 		/// <summary>
 		/// The job's step id
 		/// </summary>
 		public string? StepId { get; }
 
+		/// <summary>
+		/// The job name which utilized device
+		/// </summary>
+		public string? StepName { get; }
+
 		/// <summary>
 		/// Reservation Id (transient, reservations are deleted upon expiration)
 		/// </summary>
diff --git a/Engine/Source/Programs/Horde/Horde.Build/Devices/IDeviceCollection.cs b/Engine/Source/Programs/Horde/Horde.Build/Devices/IDeviceCollection.cs
--- a/Engine/Source/Programs/Horde/Horde.Build/Devices/IDeviceCollection.cs
+++ b/Engine/Source/Programs/Horde/Horde.Build/Devices/IDeviceCollection.cs
@@ -186,7 +186,9 @@ namespace Horde.Build.Devices
 		/// <param name="streamId">The Stream Id associated with the job</param>
 		/// <param name="jobId">The Job Id associated with the job</param>
 		/// <param name="stepId">The Step Id associated with the job</param>
-		Task<IDeviceReservation?> TryAddReservationAsync(DevicePoolId poolId, List<DeviceRequestData> request, string? hostname, string? reservationDetails, string? streamId, string? jobId, string? stepId);
+		/// <param name="jobName">The Job name associated with the job</param>
+		/// <param name="stepName">The Step name associated with the job</param>		
+		Task<IDeviceReservation?> TryAddReservationAsync(DevicePoolId poolId, List<DeviceRequestData> request, string? hostname, string? reservationDetails, string? streamId, string? jobId, string? stepId, string? jobName, string? stepName);
 
 		/// <summary>
 		/// Gets a reservation by guid for legacy clients
