Horde: Code to track down null reference.

diff --git a/Engine/Source/Programs/Horde/Horde.Build/Perforce/PerforceLoadBalancer.cs b/Engine/Source/Programs/Horde/Horde.Build/Perforce/PerforceLoadBalancer.cs
--- a/Engine/Source/Programs/Horde/Horde.Build/Perforce/PerforceLoadBalancer.cs
+++ b/Engine/Source/Programs/Horde/Horde.Build/Perforce/PerforceLoadBalancer.cs
@@ -246,7 +246,19 @@ namespace Horde.Build.Perforce
 		/// </summary>
 		static bool EvaluateCondition(Condition? condition, IReadOnlyList<string> properties)
 		{
-			return condition != null && !condition.IsEmpty() && condition.Evaluate(x => GetPropertyValues(properties, x));
+			if (condition == null)
+			{
+				return false;
+			}
+			if (condition.IsEmpty())
+			{
+				return false;
+			}
+			if (properties == null)
+			{
+				throw new ArgumentException(nameof(properties));
+			}
+			return condition.Evaluate(x => GetPropertyValues(properties, x));
 		}
 
 		/// <summary>
