Horde: Add a BSON serializer for condition objects.

diff --git a/Engine/Source/Programs/Horde/Horde.Build/Startup.cs b/Engine/Source/Programs/Horde/Horde.Build/Startup.cs
--- a/Engine/Source/Programs/Horde/Horde.Build/Startup.cs
+++ b/Engine/Source/Programs/Horde/Horde.Build/Startup.cs
@@ -828,6 +828,7 @@ namespace Horde.Build
 
 				// Register the custom serializers
 				BsonSerializer.RegisterSerializer(new RefIdBsonSerializer());
+				BsonSerializer.RegisterSerializer(new ConditionSerializer());
 				BsonSerializer.RegisterSerializationProvider(new BsonSerializationProvider());
 				BsonSerializer.RegisterSerializationProvider(new StringIdSerializationProvider());
 				BsonSerializer.RegisterSerializationProvider(new ObjectIdSerializationProvider());
diff --git a/Engine/Source/Programs/Horde/Horde.Build/Utilities/ConditionSerializer.cs b/Engine/Source/Programs/Horde/Horde.Build/Utilities/ConditionSerializer.cs
new file mode 100644
--- /dev/null
+++ b/Engine/Source/Programs/Horde/Horde.Build/Utilities/ConditionSerializer.cs
@@ -0,0 +1,35 @@
+// Copyright Epic Games, Inc. All Rights Reserved.
+
+using System;
+using EpicGames.Horde.Common;
+using MongoDB.Bson.Serialization;
+
+namespace Horde.Build.Utilities
+{
+	/// <summary>
+	/// Serializer for Condition objects
+	/// </summary>
+	public sealed class ConditionSerializer : IBsonSerializer<Condition>
+	{
+		/// <inheritdoc/>
+		public Type ValueType => typeof(Condition);
+
+		/// <inheritdoc/>
+		void IBsonSerializer.Serialize(BsonSerializationContext context, BsonSerializationArgs args, object value) => BsonClassMapSerializer(context, args, (Condition)value);
+
+		/// <inheritdoc/>
+		object IBsonSerializer.Deserialize(BsonDeserializationContext context, BsonDeserializationArgs args) => Deserialize(context, args);
+
+		/// <inheritdoc/>
+		public Condition Deserialize(BsonDeserializationContext context, BsonDeserializationArgs args)
+		{
+			return Condition.Parse(context.Reader.ReadString());
+		}
+
+		/// <inheritdoc/>
+		public void Serialize(BsonSerializationContext context, BsonSerializationArgs args, Condition value)
+		{
+			context.Writer.WriteString(value.ToString());
+		}
+	}
+}
diff --git a/Engine/Source/Programs/Horde/Horde.Build/Utilities/ContentHashSerializer.cs b/Engine/Source/Programs/Horde/Horde.Build/Utilities/ContentHashSerializer.cs
--- a/Engine/Source/Programs/Horde/Horde.Build/Utilities/ContentHashSerializer.cs
+++ b/Engine/Source/Programs/Horde/Horde.Build/Utilities/ContentHashSerializer.cs
@@ -7,7 +7,7 @@ using MongoDB.Bson.Serialization;
 namespace Horde.Build.Utilities
 {
 	/// <summary>
-	/// Serializer for JobStepRefId objects
+	/// Serializer for ContentHash objects
 	/// </summary>
 	public sealed class ContentHashSerializer : IBsonSerializer<ContentHash>
 	{
