Build fix.

diff --git a/Engine/Source/Programs/Horde/Horde.Build/Utilities/ConditionSerializer.cs b/Engine/Source/Programs/Horde/Horde.Build/Utilities/ConditionSerializer.cs
--- a/Engine/Source/Programs/Horde/Horde.Build/Utilities/ConditionSerializer.cs
+++ b/Engine/Source/Programs/Horde/Horde.Build/Utilities/ConditionSerializer.cs
@@ -15,7 +15,7 @@ namespace Horde.Build.Utilities
 		public Type ValueType => typeof(Condition);
 
 		/// <inheritdoc/>
-		void IBsonSerializer.Serialize(BsonSerializationContext context, BsonSerializationArgs args, object value) => BsonClassMapSerializer(context, args, (Condition)value);
+		void IBsonSerializer.Serialize(BsonSerializationContext context, BsonSerializationArgs args, object value) => Serialize(context, args, (Condition)value);
 
 		/// <inheritdoc/>
 		object IBsonSerializer.Deserialize(BsonDeserializationContext context, BsonDeserializationArgs args) => Deserialize(context, args);
