Horde: Support for deserializing null conditions.

diff --git a/Engine/Source/Programs/Horde/Horde.Build/Utilities/ConditionSerializer.cs b/Engine/Source/Programs/Horde/Horde.Build/Utilities/ConditionSerializer.cs
--- a/Engine/Source/Programs/Horde/Horde.Build/Utilities/ConditionSerializer.cs
+++ b/Engine/Source/Programs/Horde/Horde.Build/Utilities/ConditionSerializer.cs
@@ -2,6 +2,7 @@
 
 using System;
 using EpicGames.Horde.Common;
+using MongoDB.Bson;
 using MongoDB.Bson.Serialization;
 
 namespace Horde.Build.Utilities
@@ -23,7 +24,14 @@ namespace Horde.Build.Utilities
 		/// <inheritdoc/>
 		public Condition Deserialize(BsonDeserializationContext context, BsonDeserializationArgs args)
 		{
-			return Condition.Parse(context.Reader.ReadString());
+			if (context.Reader.CurrentBsonType == BsonType.Null)
+			{
+				return null!;
+			}
+			else
+			{
+				return Condition.Parse(context.Reader.ReadString());
+			}
 		}
 
 		/// <inheritdoc/>
