Horde: Deserialization fix.

diff --git a/Engine/Source/Programs/Horde/Horde.Build/Utilities/ConditionSerializer.cs b/Engine/Source/Programs/Horde/Horde.Build/Utilities/ConditionSerializer.cs
--- a/Engine/Source/Programs/Horde/Horde.Build/Utilities/ConditionSerializer.cs
+++ b/Engine/Source/Programs/Horde/Horde.Build/Utilities/ConditionSerializer.cs
@@ -24,7 +24,7 @@ namespace Horde.Build.Utilities
 		/// <inheritdoc/>
 		public Condition Deserialize(BsonDeserializationContext context, BsonDeserializationArgs args)
 		{
-			if (context.Reader.CurrentBsonType == BsonType.Null)
+			if (context.Reader.ReadBsonType() == BsonType.Null)
 			{
 				return null!;
 			}
