Horde: Remove custom handling for null types. Will fix in DB.

diff --git a/Engine/Source/Programs/Horde/Horde.Build/Utilities/ConditionSerializer.cs b/Engine/Source/Programs/Horde/Horde.Build/Utilities/ConditionSerializer.cs
--- a/Engine/Source/Programs/Horde/Horde.Build/Utilities/ConditionSerializer.cs
+++ b/Engine/Source/Programs/Horde/Horde.Build/Utilities/ConditionSerializer.cs
@@ -24,14 +24,7 @@ namespace Horde.Build.Utilities
 		/// <inheritdoc/>
 		public Condition Deserialize(BsonDeserializationContext context, BsonDeserializationArgs args)
 		{
-			if (context.Reader.ReadBsonType() == BsonType.Null)
-			{
-				return null!;
-			}
-			else
-			{
-				return Condition.Parse(context.Reader.ReadString());
-			}
+			return Condition.Parse(context.Reader.ReadString());
 		}
 
 		/// <inheritdoc/>
