Horde: Handle null values for conditions.

diff --git a/Engine/Source/Programs/Horde/Horde.Build/Server/Globals.cs b/Engine/Source/Programs/Horde/Horde.Build/Server/Globals.cs
--- a/Engine/Source/Programs/Horde/Horde.Build/Server/Globals.cs
+++ b/Engine/Source/Programs/Horde/Horde.Build/Server/Globals.cs
@@ -206,11 +206,13 @@ namespace Horde.Build.Server
 		/// <summary>
 		/// Optional condition for a machine to be eligable to use this server
 		/// </summary>
+		[BsonIgnoreIfNull]
 		public Condition? Condition { get; set; }
 
 		/// <summary>
 		/// List of properties for an agent to be eligable to use this server
 		/// </summary>
+		[BsonIgnoreIfNull]
 		public List<string>? Properties { get; set; }
 	}
 
diff --git a/Engine/Source/Programs/Horde/Horde.Build/Utilities/ConditionSerializer.cs b/Engine/Source/Programs/Horde/Horde.Build/Utilities/ConditionSerializer.cs
--- a/Engine/Source/Programs/Horde/Horde.Build/Utilities/ConditionSerializer.cs
+++ b/Engine/Source/Programs/Horde/Horde.Build/Utilities/ConditionSerializer.cs
@@ -24,6 +24,11 @@ namespace Horde.Build.Utilities
 		/// <inheritdoc/>
 		public Condition Deserialize(BsonDeserializationContext context, BsonDeserializationArgs args)
 		{
+			if(context.Reader.CurrentBsonType == BsonType.Null)
+			{
+				context.Reader.ReadNull();
+				return null!
+			}
 			return Condition.Parse(context.Reader.ReadString());
 		}
 
