Fix typo.

diff --git a/Engine/Source/Programs/Horde/Horde.Build/Utilities/ConditionSerializer.cs b/Engine/Source/Programs/Horde/Horde.Build/Utilities/ConditionSerializer.cs
--- a/Engine/Source/Programs/Horde/Horde.Build/Utilities/ConditionSerializer.cs
+++ b/Engine/Source/Programs/Horde/Horde.Build/Utilities/ConditionSerializer.cs
@@ -27,7 +27,7 @@ namespace Horde.Build.Utilities
 			if(context.Reader.CurrentBsonType == BsonType.Null)
 			{
 				context.Reader.ReadNull();
-				return null!
+				return null!;
 			}
 			return Condition.Parse(context.Reader.ReadString());
 		}
