Horde: Improve device pool telemetry to include allocated devices in states
#rnx

diff --git a/Engine/Source/Programs/Horde/Horde.Build.Tests/DeviceServiceTests.cs b/Engine/Source/Programs/Horde/Horde.Build.Tests/DeviceServiceTests.cs
--- a/Engine/Source/Programs/Horde/Horde.Build.Tests/DeviceServiceTests.cs
+++ b/Engine/Source/Programs/Horde/Horde.Build.Tests/DeviceServiceTests.cs
@@ -223,21 +223,23 @@ namespace Horde.Build.Tests
 				if (t.PlatformId == "testdeviceplatform1")
 				{
 					Assert.IsNull(t.Available);
-					Assert.AreEqual(1, t.Reserved);
-					Assert.AreEqual(1, t.Maintenance);
-
-					Assert.IsTrue(t.StreamDevices!.ContainsKey("ue5-main"));
-					Assert.AreEqual(t.StreamDevices["ue5-main"].Count, 1);
-					Assert.AreEqual(t.StreamDevices["ue5-main"][0], reservedDevice.Id);
+					Assert.AreEqual(1, t.Reserved?.Count);
+					Assert.AreEqual(1, t.Maintenance?.Count);
+
+					Assert.AreEqual(true, t.Reserved?.ContainsKey("ue5-main"));
+					Assert.AreEqual(t.Reserved?["ue5-main"].Count, 1);
+					Assert.AreEqual(t.Reserved?["ue5-main"][0].DeviceId, reservedDevice.Id);
+					Assert.AreEqual(t.Reserved?["ue5-main"][0].JobName, "hello2");
+					Assert.AreEqual(t.Reserved?["ue5-main"][0].StepId, "abcd");
 				}
 				else if (t.PlatformId == "testdeviceplatform3")
 				{
-					Assert.AreEqual(1, t.Available);
-					Assert.AreEqual(1, t.Disabled);
+					Assert.AreEqual(1, t.Available?.Count);
+					Assert.AreEqual(1, t.Disabled?.Count);
 				}
 				else
 				{
-					Assert.AreEqual(2, t.Available);
+					Assert.AreEqual(2, t.Available?.Count);
 				}
 			}
 
@@ -245,12 +247,12 @@ namespace Horde.Build.Tests
 			{
 				if (t.PlatformId == "testdeviceplatform2")
 				{
-					Assert.AreEqual(1, t.Available);
-					Assert.AreEqual(1, t.Problem);
+					Assert.AreEqual(1, t.Available?.Count);
+					Assert.AreEqual(1, t.Problem?.Count);
 				}
 				else
 				{
-					Assert.AreEqual(2, t.Available);
+					Assert.AreEqual(2, t.Available?.Count);
 				}
 			}
 		}
diff --git a/Engine/Source/Programs/Horde/Horde.Build/Devices/DeviceCollection.cs b/Engine/Source/Programs/Horde/Horde.Build/Devices/DeviceCollection.cs
--- a/Engine/Source/Programs/Horde/Horde.Build/Devices/DeviceCollection.cs
+++ b/Engine/Source/Programs/Horde/Horde.Build/Devices/DeviceCollection.cs
@@ -12,6 +12,7 @@ using Horde.Build.Users;
 using Horde.Build.Utilities;
 using MongoDB.Bson;
 using MongoDB.Bson.Serialization.Attributes;
+using MongoDB.Bson.Serialization.Options;
 using MongoDB.Driver;
 
 namespace Horde.Build.Devices
@@ -354,15 +355,39 @@ namespace Horde.Build.Devices
 
 		}
 
-		class StreamDeviceTelemetry : IStreamDeviceTelemetry
+		class DeviceReservationPoolTelemetryDocument : IDevicePoolReservationTelemetry
 		{
 			/// <inheritdoc/>
-			public StreamId StreamId { get; set; } = new StreamId();
+			[BsonRequired]
+			public DeviceId DeviceId { get; set; }
 
-			/// <inheritdoc/>
-			public List<DeviceId> DeviceIds { get; set; } = new List<DeviceId>();
-		}
+			[BsonIgnoreIfNull]
+			public string? JobId { get; set; }
+
+			[BsonIgnoreIfNull]
+			public string? StepId { get; set; }
+
+			[BsonIgnoreIfNull]
+			public string? JobName { get; set; }
+
+			[BsonIgnoreIfNull]
+			public string? StepName { get; set; }
+
+			[BsonConstructor]
+			private DeviceReservationPoolTelemetryDocument()
+			{
+			}
 
+			public DeviceReservationPoolTelemetryDocument(DeviceId deviceId, string? jobId, string? stepId, string? jobName, string? stepName)
+			{
+				DeviceId = deviceId;
+				JobId = jobId;
+				StepId = stepId;
+				JobName = jobName;
+				StepName = stepName;
+			}
+
+		}
 
 		class DevicePlatformTelemetryDocument : IDevicePlatformTelemetry
 		{
@@ -372,73 +397,68 @@ namespace Horde.Build.Devices
 			[BsonRequired, BsonId, BsonElement("tid")]
 			public ObjectId TelemetryId { get; set; }
 
-			/// <summary>
-			/// The time this telemetry data was created
-			/// </summary>
 			[BsonRequired, BsonElement("c")]
 			public DateTime CreateTimeUtc { get; set; }
 
-			/// <summary>
-			/// The platform id 
-			/// </summary>
 			[BsonRequired, BsonElement("pid")]
 			public DevicePlatformId PlatformId { get; set; }
 
-			/// <summary>
-			/// Number of available devices of this platform 
-			/// </summary>
-			[BsonIgnoreIfDefault, BsonDefaultValue(0), BsonElement("a")]
-			public int Available { get; set; } = 0;
-
-			/// <summary>
-			/// Number of reserved devices of this platform 
-			/// </summary>
-			[BsonIgnoreIfDefault, BsonDefaultValue(0), BsonElement("r")]
-			public int Reserved { get; set; } = 0;
+			[BsonIgnoreIfNull, BsonElement("a")]
+			public List<DeviceId>? Available { get; set; }
+			IReadOnlyList<DeviceId>? IDevicePlatformTelemetry.Available => Available;
 
-			/// <summary>
-			/// Number of devices in maintenance state
-			/// </summary>
-			[BsonIgnoreIfDefault, BsonDefaultValue(0), BsonElement("m")]
-			public int Maintenance { get; set; } = 0;
+			[BsonIgnoreIfNull, BsonElement("m")]
+			public List<DeviceId>? Maintenance { get; set; }
+			IReadOnlyList<DeviceId>? IDevicePlatformTelemetry.Maintenance => Maintenance;
 
-			/// <summary>
-			/// Number of devices in problem state
-			/// </summary>			
-			[BsonIgnoreIfDefault, BsonDefaultValue(0), BsonElement("p")]
-			public int Problem { get; set; } = 0;
+			[BsonIgnoreIfNull, BsonElement("p")]
+			public List<DeviceId>? Problem { get; set; }
+			IReadOnlyList<DeviceId>? IDevicePlatformTelemetry.Problem => Problem;
 
-			/// <summary>
-			/// Number of devices that are disabled
-			/// </summary>
-			[BsonIgnoreIfDefault, BsonDefaultValue(0), BsonElement("d")]
-			public int Disabled { get; set; } = 0;
+			[BsonIgnoreIfNull, BsonElement("d")]
+			public List<DeviceId>? Disabled { get; set; }
+			IReadOnlyList<DeviceId>? IDevicePlatformTelemetry.Disabled => Disabled;
 
-			/// <summary>
-			/// Stream devices that are reserved
-			/// </summary>
-			[BsonIgnoreIfNull]
-			public List<StreamDeviceTelemetry> StreamDevices { get; set; } = new List<StreamDeviceTelemetry>();
-			IReadOnlyList<IStreamDeviceTelemetry> IDevicePlatformTelemetry.StreamDevices => StreamDevices;
+			[BsonIgnoreIfNull, BsonElement("r"), BsonDictionaryOptions(DictionaryRepresentation.ArrayOfDocuments)]
+			public Dictionary<StreamId, List<DeviceReservationPoolTelemetryDocument>>? Reserved { get; set; }
+			IReadOnlyDictionary<StreamId, IReadOnlyList<IDevicePoolReservationTelemetry>>? IDevicePlatformTelemetry.Reserved => Reserved?.ToDictionary(kvp => kvp.Key, kvp => kvp.Value as IReadOnlyList<IDevicePoolReservationTelemetry>) ?? new Dictionary<StreamId, IReadOnlyList<IDevicePoolReservationTelemetry>>();
 
 			[BsonConstructor]
 			private DevicePlatformTelemetryDocument()
 			{
 			}
 
-			public DevicePlatformTelemetryDocument(DevicePlatformId platformId, int available, int reserved, int maintenance, int problem, int disabled, List<StreamDeviceTelemetry> streamDevices)
+			public DevicePlatformTelemetryDocument(DevicePlatformId platformId, List<DeviceId>? available, Dictionary<StreamId, List<DeviceReservationPoolTelemetryDocument>>? reserved, List<DeviceId>? maintenance, List<DeviceId>? problem, List<DeviceId>? disabled)
 			{
 				TelemetryId = ObjectId.GenerateNewId();
-				CreateTimeUtc = DateTime.UtcNow;
+				CreateTimeUtc = DateTime.UtcNow;				
 				PlatformId = platformId;
-				Available = available;
-				Reserved = reserved;
-				Maintenance = maintenance;	
-				Problem = problem;
-				Disabled = disabled;
-				StreamDevices = streamDevices;
-			}
 
+				if (available != null && available.Count > 0)
+				{
+					Available = available;
+				}
+
+				if (reserved != null && reserved.Count > 0)
+				{
+					Reserved = reserved;
+				}
+
+				if (maintenance != null && maintenance.Count > 0)
+				{
+					Maintenance = maintenance;
+				}
+
+				if (problem != null && problem.Count > 0)
+				{
+					Problem = problem;
+				}
+
+				if (disabled != null && disabled.Count > 0)
+				{
+					Disabled = disabled;
+				}				
+			}
 		}
 
 		/// <summary>
@@ -1097,15 +1117,15 @@ namespace Horde.Build.Devices
 			return results.ConvertAll<IDeviceTelemetry>(x => x);
 		}
 
-		internal class DevicePoolTelemetryHelper
+		struct DevicePoolTelemetryHelper
 		{
-			public int available = 0;
-			public int reserved = 0;
-			public int maintenance = 0;
-			public int problem = 0;
-			public int disabled = 0;
-			public Dictionary<StreamId, List<DeviceId>> streamDevices = new Dictionary<StreamId, List<DeviceId>>();
+			public List<DeviceId> available = new List<DeviceId>();
+			public List<DeviceId> problem = new List<DeviceId>();
+			public List<DeviceId> maintenance = new List<DeviceId>();
+			public List<DeviceId> disabled = new List<DeviceId>();
+			public Dictionary<StreamId, List<DeviceReservationPoolTelemetryDocument>> reserved = new Dictionary<StreamId, List<DeviceReservationPoolTelemetryDocument>>();
 
+			public DevicePoolTelemetryHelper() { }
 		}
 
 		/// <summary>
@@ -1139,7 +1159,6 @@ namespace Horde.Build.Devices
 			foreach (IDevicePool pool in pools)
 			{								
 				List<IDevice> poolDevices = devices.Where(x => x.PoolId == pool.Id).ToList();
-
 				HashSet<DevicePlatformId> platforms = new HashSet<DevicePlatformId>();
 				poolDevices.ForEach(d => platforms.Add(d.PlatformId));
 
@@ -1155,54 +1174,50 @@ namespace Horde.Build.Devices
 
 					if (reservedDevices.Contains(device))
 					{
-						helper.reserved++;
 						IDeviceReservation? reservation = reservations.FirstOrDefault(x => x.Devices.Contains(device.Id));
 						if (reservation != null && reservation.StreamId != null)
 						{
 							StreamId streamId = new StreamId(reservation.StreamId);
-							List<DeviceId>? deviceIds;
-							if (!helper.streamDevices.TryGetValue(streamId, out deviceIds))
+							List<DeviceReservationPoolTelemetryDocument>? rdevices;
+							if (!helper.reserved.TryGetValue(streamId, out rdevices))
 							{
-								deviceIds = new List<DeviceId>();
-								helper.streamDevices[streamId] = deviceIds;
+								rdevices = new List<DeviceReservationPoolTelemetryDocument>();
+								helper.reserved[streamId] = rdevices;
 							}
-							deviceIds.Add(device.Id);
+							rdevices.Add(new DeviceReservationPoolTelemetryDocument(device.Id, reservation.JobId, reservation.StepId, reservation.JobName, reservation.StepName));
 						}
 						continue;
 					}
 
 					if (problemDevices.Contains(device))
 					{
-						helper.problem++;
+						helper.problem.Add(device.Id);
 						continue;
 					}
 
 					if (maintenanceDevices.Contains(device))
 					{
-						helper.maintenance++;
+						helper.maintenance.Add(device.Id);
 						continue;
 					}
 
 					if (disabledDevices.Contains(device))
 					{
-						helper.disabled++;
+						helper.disabled.Add(device.Id);
 						continue;
 					}
 
-					helper.available++;
+					helper.available.Add(device.Id);
 				}
 
 				List<DevicePlatformTelemetryDocument> platformTelemtry = new List<DevicePlatformTelemetryDocument>();
+
+
 				foreach (KeyValuePair<DevicePlatformId, DevicePoolTelemetryHelper> platform in helpers)
 				{
 					DevicePoolTelemetryHelper helper = platform.Value;
-					List<StreamDeviceTelemetry> streamDevices = new List<StreamDeviceTelemetry>();
-					foreach (KeyValuePair<StreamId, List<DeviceId>> i in helper.streamDevices)
-					{
-						streamDevices.Add(new StreamDeviceTelemetry() { StreamId = i.Key, DeviceIds = i.Value });
-					}
 
-					platformTelemtry.Add(new DevicePlatformTelemetryDocument(platform.Key, helper.available, helper.reserved, helper.maintenance, helper.problem, helper.disabled, streamDevices));
+					platformTelemtry.Add(new DevicePlatformTelemetryDocument(platform.Key, helper.available, helper.reserved, helper.maintenance, helper.problem, helper.disabled));
 				}
 
 				poolTelemetry[pool.Id] = platformTelemtry;
diff --git a/Engine/Source/Programs/Horde/Horde.Build/Devices/DeviceMessages.cs b/Engine/Source/Programs/Horde/Horde.Build/Devices/DeviceMessages.cs
--- a/Engine/Source/Programs/Horde/Horde.Build/Devices/DeviceMessages.cs
+++ b/Engine/Source/Programs/Horde/Horde.Build/Devices/DeviceMessages.cs
@@ -656,6 +656,55 @@ namespace Horde.Build.Devices
 	}
 
 
+	/// <summary>
+	/// Stream device telemetry for pool snapshot
+	/// </summary>
+	public class GetDevicePoolReservationTelemetryResponse
+	{
+		/// <summary>
+		/// Device id for reservation
+		/// </summary>
+		public string DeviceId { get; set; }
+
+		/// <summary>
+		/// Job id associated with reservation
+		/// </summary>
+		public string? JobId { get; set; }
+
+		/// <summary>
+		/// The step id of reservation
+		/// </summary>
+		public string? StepId { get; set; }
+
+		/// <summary>
+		/// The name of the job holding reservation
+		/// </summary>
+		public string? JobName { get; set; }
+
+		/// <summary>
+		/// The name of the step holding reservation
+		/// </summary>
+		public string? StepName { get; set; }
+
+		/// <summary>
+		/// constructor
+		/// </summary>
+		/// <param name="deviceId"></param>
+		/// <param name="jobId"></param>
+		/// <param name="stepId"></param>
+		/// <param name="jobName"></param>
+		/// <param name="stepName"></param>
+		public GetDevicePoolReservationTelemetryResponse(string deviceId, string? jobId, string? stepId, string? jobName, string? stepName)
+		{
+			DeviceId = deviceId;
+			JobId = jobId;
+			StepId = stepId;
+			JobName = jobName;
+			StepName = stepName;
+		}
+
+	}
+
 	/// <summary>
 	/// Device telemetry respponse
 	/// </summary>
@@ -667,47 +716,65 @@ namespace Horde.Build.Devices
 		public string PlatformId { get; set; }
 
 		/// <summary>
-		/// Number of available devices of this platform 
-		/// </summary>
-		public int? Available { get; set; }
-
-		/// <summary>
-		/// Number of reserved devices of this platform 
+		/// Available devices of this platform 
 		/// </summary>
-		public int? Reserved { get; set; }
+		public List<string>? Available { get; set; }
 
 		/// <summary>
-		/// Number of devices in maintenance state
+		/// Devices in maintenance state
 		/// </summary>
-		public int? Maintenance { get; set; }
+		public List<string>? Maintenance { get; set; }
 
 		/// <summary>
-		/// Number of devices in problem state
+		/// Devices in problem state
 		/// </summary>
-		public int? Problem { get; set; }
+		public List<string>? Problem { get; set; }
 
 		/// <summary>
 		/// Number of devices in disabled state
 		/// </summary>
-		public int? Disabled { get; set; }
+		public List<string>? Disabled { get; set; }
 
 		/// <summary>
-		/// StreamId to reserved device ids
+		/// Reserved devices
 		/// </summary>
-		public Dictionary<string, List<string>>? StreamDevices { get; set; }
+		public Dictionary<string, List<GetDevicePoolReservationTelemetryResponse>>? Reserved { get; set; }
 
 		/// <summary>
 		/// Constructor
 		/// </summary>
-		public GetDevicePlatformTelemetryResponse(string platformId, int available, int reserved, int maintenance, int problem, int disabled, Dictionary<string, List<string>> streamDevices)
+		public GetDevicePlatformTelemetryResponse(string platformId, List<string>? available, List<string>? maintenance, List<string>? problem, List<string>? disabled, IReadOnlyDictionary<string, IReadOnlyList<IDevicePoolReservationTelemetry>>? reserved)
 		{
 			PlatformId = platformId;
-			Available = available == 0 ? null : available;
-			Reserved = reserved == 0 ? null : reserved;
-			Maintenance = maintenance == 0 ? null : maintenance;
-			Problem = problem == 0 ? null : problem;
-			Disabled = disabled == 0 ? null : disabled;
-			StreamDevices = streamDevices.Count == 0 ? null : streamDevices;
+			if (available != null && available.Count > 0)
+			{
+				Available = available;
+			}
+			if (maintenance != null && maintenance.Count > 0)
+			{
+				Maintenance = maintenance;
+			}
+			if (problem != null && problem.Count > 0)
+			{
+				Problem = problem;
+			}
+			if (disabled != null && disabled.Count > 0)
+			{
+				Disabled = disabled;
+			}
+			if (reserved != null && reserved.Count > 0)
+			{
+				Reserved = new Dictionary<string, List<GetDevicePoolReservationTelemetryResponse>>();
+
+				foreach (KeyValuePair<string, IReadOnlyList<IDevicePoolReservationTelemetry>> r in reserved)
+				{
+					Reserved[r.Key] = new List<GetDevicePoolReservationTelemetryResponse>();
+					foreach (IDevicePoolReservationTelemetry telemetry in r.Value)
+					{
+						Reserved[r.Key].Add(new GetDevicePoolReservationTelemetryResponse(telemetry.DeviceId.ToString(), telemetry.JobId, telemetry.StepId, telemetry.JobName, telemetry.StepName));
+					}
+				}
+			}
 		}
 	}
 
diff --git a/Engine/Source/Programs/Horde/Horde.Build/Devices/DevicesController.cs b/Engine/Source/Programs/Horde/Horde.Build/Devices/DevicesController.cs
--- a/Engine/Source/Programs/Horde/Horde.Build/Devices/DevicesController.cs
+++ b/Engine/Source/Programs/Horde/Horde.Build/Devices/DevicesController.cs
@@ -526,12 +526,8 @@ namespace Horde.Build.Devices
 
 					foreach (IDevicePlatformTelemetry telemetry in pool.Value)
 					{
-						Dictionary<string, List<string>> streamDevices = new Dictionary<string, List<string>>();
-						foreach(IStreamDeviceTelemetry i in telemetry.StreamDevices)
-						{
-							streamDevices[i.StreamId.ToString()] = i.DeviceIds.Select(x => x.ToString()).ToList();
-						}
-						platformTelemetry.Add(new GetDevicePlatformTelemetryResponse(telemetry.PlatformId.ToString(), telemetry.Available, telemetry.Reserved, telemetry.Maintenance, telemetry.Problem, telemetry.Disabled, streamDevices));
+						IReadOnlyDictionary<string, IReadOnlyList<IDevicePoolReservationTelemetry>>? reserved = telemetry.Reserved?.ToDictionary(kvp => kvp.Key.ToString(), kvp => kvp.Value);
+						platformTelemetry.Add(new GetDevicePlatformTelemetryResponse(telemetry.PlatformId.ToString(), telemetry.Available?.Select(d => d.ToString()).ToList(), telemetry.Maintenance?.Select(d => d.ToString()).ToList(), telemetry.Problem?.Select(d => d.ToString()).ToList(), telemetry.Disabled?.Select(d => d.ToString()).ToList(), reserved));
 					}
 
 					poolData[pool.Key.ToString()] = platformTelemetry;
diff --git a/Engine/Source/Programs/Horde/Horde.Build/Devices/IDevice.cs b/Engine/Source/Programs/Horde/Horde.Build/Devices/IDevice.cs
--- a/Engine/Source/Programs/Horde/Horde.Build/Devices/IDevice.cs
+++ b/Engine/Source/Programs/Horde/Horde.Build/Devices/IDevice.cs
@@ -358,18 +358,32 @@ namespace Horde.Build.Devices
 	/// <summary>
 	/// Stream device telemetry for pool snapshot
 	/// </summary>
-	public interface IStreamDeviceTelemetry
+	public interface IDevicePoolReservationTelemetry
 	{
 		/// <summary>
-		/// Associated stream
+		/// Device id for reservation
 		/// </summary>
-		public StreamId StreamId { get;  }
+		public DeviceId DeviceId { get;}
 
 		/// <summary>
-		/// Device ids contained in the snapshot
+		/// Job id associated with reservation
 		/// </summary>
-		public List<DeviceId> DeviceIds { get; }
+		public string? JobId { get; }
+
+		/// <summary>
+		/// The step id of reservation
+		/// </summary>
+		public string? StepId { get; }
+
+		/// <summary>
+		/// The name of the job holding reservation
+		/// </summary>
+		public string? JobName { get; }
 
+		/// <summary>
+		/// The name of the step holding reservation
+		/// </summary>
+		public string? StepName { get; }
 	}
 
 
@@ -384,34 +398,29 @@ namespace Horde.Build.Devices
 		public DevicePlatformId PlatformId { get; }
 
 		/// <summary>
-		/// Number of available devices of this platform 
-		/// </summary>
-		public int Available { get; }
-
-		/// <summary>
-		/// Number of reserved devices of this platform 
+		/// Available devices of this platform 
 		/// </summary>
-		public int Reserved { get; }
+		public IReadOnlyList<DeviceId>? Available { get; }
 
 		/// <summary>
 		/// Number of devices in maintenance state
 		/// </summary>
-		public int Maintenance { get; }
+		public IReadOnlyList<DeviceId>? Maintenance { get; }
 
 		/// <summary>
 		/// Number of devices in problem state
 		/// </summary>
-		public int Problem { get; }
+		public IReadOnlyList<DeviceId>? Problem { get; }
 
 		/// <summary>
 		/// Number of devices in disabled state
 		/// </summary>
-		public int Disabled { get; }
+		public IReadOnlyList<DeviceId>? Disabled { get; }
 
 		/// <summary>
-		/// Devices reserved in streams
+		/// Number of reserved devices of this platform 
 		/// </summary>
-		public IReadOnlyList<IStreamDeviceTelemetry> StreamDevices { get; }
+		public IReadOnlyDictionary<StreamId, IReadOnlyList<IDevicePoolReservationTelemetry>>? Reserved { get; }
 
 	}
 
