Horde: Remove unnecessary bson id from pool telemetry data
#rnx

diff --git a/Engine/Source/Programs/Horde/Horde.Build/Devices/DeviceCollection.cs b/Engine/Source/Programs/Horde/Horde.Build/Devices/DeviceCollection.cs
--- a/Engine/Source/Programs/Horde/Horde.Build/Devices/DeviceCollection.cs
+++ b/Engine/Source/Programs/Horde/Horde.Build/Devices/DeviceCollection.cs
@@ -391,12 +391,6 @@ namespace Horde.Build.Devices
 
 		class DevicePlatformTelemetryDocument : IDevicePlatformTelemetry
 		{
-			/// <summary>
-			/// Id of telemetry document
-			/// </summary>
-			[BsonRequired, BsonId, BsonElement("tid")]
-			public ObjectId TelemetryId { get; set; }
-
 			[BsonRequired, BsonElement("c")]
 			public DateTime CreateTimeUtc { get; set; }
 
@@ -429,8 +423,7 @@ namespace Horde.Build.Devices
 			}
 
 			public DevicePlatformTelemetryDocument(DevicePlatformId platformId, List<DeviceId>? available, Dictionary<StreamId, List<DeviceReservationPoolTelemetryDocument>>? reserved, List<DeviceId>? maintenance, List<DeviceId>? problem, List<DeviceId>? disabled)
-			{
-				TelemetryId = ObjectId.GenerateNewId();
+			{				
 				CreateTimeUtc = DateTime.UtcNow;				
 				PlatformId = platformId;
 
