Horde: Remove unneeded timestamp for pool platform telemetry while at it
#rnx

diff --git a/Engine/Source/Programs/Horde/Horde.Build/Devices/DeviceCollection.cs b/Engine/Source/Programs/Horde/Horde.Build/Devices/DeviceCollection.cs
--- a/Engine/Source/Programs/Horde/Horde.Build/Devices/DeviceCollection.cs
+++ b/Engine/Source/Programs/Horde/Horde.Build/Devices/DeviceCollection.cs
@@ -391,9 +391,6 @@ namespace Horde.Build.Devices
 
 		class DevicePlatformTelemetryDocument : IDevicePlatformTelemetry
 		{
-			[BsonRequired, BsonElement("c")]
-			public DateTime CreateTimeUtc { get; set; }
-
 			[BsonRequired, BsonElement("pid")]
 			public DevicePlatformId PlatformId { get; set; }
 
@@ -423,8 +420,7 @@ namespace Horde.Build.Devices
 			}
 
 			public DevicePlatformTelemetryDocument(DevicePlatformId platformId, List<DeviceId>? available, Dictionary<StreamId, List<DeviceReservationPoolTelemetryDocument>>? reserved, List<DeviceId>? maintenance, List<DeviceId>? problem, List<DeviceId>? disabled)
-			{				
-				CreateTimeUtc = DateTime.UtcNow;				
+			{								
 				PlatformId = platformId;
 
 				if (available != null && available.Count > 0)
