Horde: Add a context block underneath the issue assignment buttons showing who has acknowledged/declined the issue. Also add modals to confirm reassigning an existing issue, and declining an issue for which the user is not a suspect.

diff --git a/Engine/Source/Programs/Horde/Horde.Build/Issues/IIssueCollection.cs b/Engine/Source/Programs/Horde/Horde.Build/Issues/IIssueCollection.cs
--- a/Engine/Source/Programs/Horde/Horde.Build/Issues/IIssueCollection.cs
+++ b/Engine/Source/Programs/Horde/Horde.Build/Issues/IIssueCollection.cs
@@ -420,9 +420,9 @@ namespace Horde.Build.Issues
 		/// <summary>
 		/// Finds the suspects for an issue
 		/// </summary>
-		/// <param name="issue">The issue to retrieve suspects for</param>
+		/// <param name="issueId">The issue to retrieve suspects for</param>
 		/// <returns>List of suspects</returns>
-		Task<List<IIssueSuspect>> FindSuspectsAsync(IIssue issue);
+		Task<List<IIssueSuspect>> FindSuspectsAsync(int issueId);
 
 		/// <summary>
 		/// Searches for open issues
@@ -595,6 +595,17 @@ namespace Horde.Build.Issues
 	/// </summary>
 	public static class IssueCollectionExtensions
 	{
+		/// <summary>
+		/// Find suspects for the given issue
+		/// </summary>
+		/// <param name="issueCollection"></param>
+		/// <param name="issue"></param>
+		/// <returns></returns>
+		public static Task<List<IIssueSuspect>> FindSuspectsAsync(this IIssueCollection issueCollection, IIssue issue)
+		{
+			return issueCollection.FindSuspectsAsync(issue.Id);
+		}
+
 		/// <summary>
 		/// Find steps for the given spans
 		/// </summary>
@@ -640,7 +651,5 @@ namespace Horde.Build.Issues
 
 			return await issueCollection.FindIssuesAsync(ids: issueIds, ownerId: ownerId, resolved: resolved, promoted: promoted, index: index, count: count);
 		}
-
-
 	}
 }
diff --git a/Engine/Source/Programs/Horde/Horde.Build/Issues/IssueCollection.cs b/Engine/Source/Programs/Horde/Horde.Build/Issues/IssueCollection.cs
--- a/Engine/Source/Programs/Horde/Horde.Build/Issues/IssueCollection.cs
+++ b/Engine/Source/Programs/Horde/Horde.Build/Issues/IssueCollection.cs
@@ -684,9 +684,9 @@ namespace Horde.Build.Issues
 		}
 
 		/// <inheritdoc/>
-		public Task<List<IIssueSuspect>> FindSuspectsAsync(IIssue issue)
+		public Task<List<IIssueSuspect>> FindSuspectsAsync(int issueId)
 		{
-			return _issueSuspects.Find(x => x.IssueId == issue.Id).ToListAsync<IssueSuspect, IIssueSuspect>();
+			return _issueSuspects.Find(x => x.IssueId == issueId).ToListAsync<IssueSuspect, IIssueSuspect>();
 		}
 
 		class ProjectedIssueId
@@ -853,6 +853,11 @@ namespace Horde.Build.Issues
 		{
 			Issue issueDocument = (Issue)issue;
 
+			if (newDeclinedById != null && newDeclinedById == issueDocument.OwnerId)
+			{
+				newOwnerId = UserId.Empty;
+			}
+
 			DateTime utcNow = DateTime.UtcNow;
 
 			List<UpdateDefinition<Issue>> updates = new List<UpdateDefinition<Issue>>();
diff --git a/Engine/Source/Programs/Horde/Horde.Build/Notifications/Sinks/SlackNotificationSink.cs b/Engine/Source/Programs/Horde/Horde.Build/Notifications/Sinks/SlackNotificationSink.cs
--- a/Engine/Source/Programs/Horde/Horde.Build/Notifications/Sinks/SlackNotificationSink.cs
+++ b/Engine/Source/Programs/Horde/Horde.Build/Notifications/Sinks/SlackNotificationSink.cs
@@ -777,6 +777,9 @@ namespace Horde.Build.Notifications.Sinks
 					text = $"~{text}~";
 				}
 
+				// Get the suspects for the issue
+				List<IIssueSuspect> suspects = await _issueService.Collection.FindSuspectsAsync(issue);
+
 				(MessageStateDocument state, bool isNew) = await SendOrUpdateMessageAsync(triageChannel, eventId, null, text);
 				if (isNew)
 				{
@@ -817,28 +820,76 @@ namespace Horde.Build.Notifications.Sinks
 						message.Blocks.Add(new SectionBlock(new TextObject("```...```")));
 					}
 
+					string? summaryTs = await _slackClient.PostMessageAsync(triageChannel, state.Ts, message);
+
+					// Permalink to the summary text so we link inside the thread rather than just to the original message
+					string? permalink = null;
+					if (summaryTs != null)
+					{
+						permalink = await _slackClient.GetPermalinkAsync(state.Channel, summaryTs);
+					}
+
+					await SetMessageTimestampAsync(state.Id, state.Channel, state.Ts, permalink);
+				}
+
+				// Post a message containing the controls and status
+				{
+					SlackMessage message = new SlackMessage();
+
+					if (workflow.TriageInstructions != null)
+					{
+						message.AddSection(workflow.TriageInstructions);
+					}
+
 					ActionsBlock actions = message.AddActions();
 					actions.AddButton("Assign to Me", value: $"issue_{issue.Id}_ack", style: ButtonStyle.Primary);
 					actions.AddButton("Not Me", value: $"issue_{issue.Id}_decline", style: ButtonStyle.Danger);
 					actions.AddButton("Mark Fixed", value: $"issue_{issue.Id}_markfixed");
 
-					if (workflow.TriageInstructions != null)
+					string? context = null;
+					if (issue.OwnerId != null)
 					{
-						message.AddSection(workflow.TriageInstructions);
+						string user = await FormatNameAsync(issue.OwnerId.Value);
+						if (issue.AcknowledgedAt == null)
+						{
+							context = $"Assigned to {user} (unacknowledged).";
+						}
+						else
+						{
+							context = $"Acknowledged by {user}.";
+						}
 					}
+					else if(suspects.Any(x => x.DeclinedAt != null))
+					{
+						HashSet<UserId> userIds = new HashSet<UserId>();
+						foreach (IIssueSuspect suspect in suspects)
+						{
+							if (suspect.DeclinedAt != null)
+							{
+								userIds.Add(suspect.AuthorId);
+							}
+						}
 
-					string? summaryTs = await _slackClient.PostMessageAsync(triageChannel, state.Ts, message);
+						List<string> users = new List<string>();
+						foreach (UserId userId in userIds)
+						{
+							users.Add(await FormatNameAsync(userId));
+						}
+						users.Sort(StringComparer.OrdinalIgnoreCase);
 
-					// Permalink to the summary text so we link inside the thread rather than just to the original message
-					string? permalink = null;
-					if (summaryTs != null)
+						context = $"Declined by {StringUtils.FormatList(users)}.";
+					}
+
+					if (context != null)
 					{
-						permalink = await _slackClient.GetPermalinkAsync(state.Channel, summaryTs);
+						message.AddContext(context);
 					}
 
-					// Get the suspects for the issue
-					List<IIssueSuspect> suspects = await _issueService.Collection.FindSuspectsAsync(issue);
+					await SendOrUpdateMessageAsync(triageChannel, state.Ts, $"{eventId}_buttons", null, message);
+				}
 
+				if (isNew)
+				{
 					// If it has an owner, show that
 					HashSet<UserId> inviteUserIds = new HashSet<UserId>();
 					if (issue.OwnerId != null)
@@ -874,8 +925,6 @@ namespace Horde.Build.Notifications.Sinks
 						}
 					}
 
-					await SetMessageTimestampAsync(state.Id, state.Channel, state.Ts, permalink);
-
 					if (_environment.IsProduction() && workflow.AllowMentions)
 					{
 						await InviteUsersAsync(state.Channel, inviteUserIds);
@@ -1218,6 +1267,19 @@ namespace Horde.Build.Notifications.Sinks
 			}
 		}
 
+		async Task<string> FormatNameAsync(UserId userId, bool allowMentions)
+		{
+			IUser? user = await _userCollection.GetUserAsync(userId);
+			if (user == null)
+			{
+				return $"User {userId}";
+			}
+			else
+			{
+				return user.Name;
+			}
+		}
+
 		async Task<string> FormatMentionAsync(UserId userId, bool allowMentions)
 		{
 			IUser? user = await _userCollection.GetUserAsync(userId);
@@ -1920,6 +1982,11 @@ namespace Horde.Build.Notifications.Sinks
 		}
 
 		private async Task<(MessageStateDocument, bool)> SendOrUpdateMessageAsync(string recipient, string eventId, UserId? userId, SlackMessage message)
+		{
+			return await SendOrUpdateMessageAsync(recipient, null, eventId, userId, message);
+		}
+
+		private async Task<(MessageStateDocument, bool)> SendOrUpdateMessageAsync(string recipient, string? threadTs, string eventId, UserId? userId, SlackMessage message)
 		{
 			string requestDigest = ContentHash.MD5(JsonSerializer.Serialize(message, new JsonSerializerOptions { DefaultIgnoreCondition = JsonIgnoreCondition.WhenWritingNull })).ToString();
 
@@ -1935,14 +2002,28 @@ namespace Horde.Build.Notifications.Sinks
 				_logger.LogInformation("Sending new slack message to {SlackUser} (msg: {MessageId})", recipient, state.Id);
 
 				state.Channel = recipient;
-				state.Ts = await _slackClient.PostMessageAsync(recipient, message);
+				if (threadTs == null)
+				{
+					state.Ts = await _slackClient.PostMessageAsync(recipient, message);
+				}
+				else
+				{
+					state.Ts = await _slackClient.PostMessageAsync(recipient, threadTs, message);
+				}
 
 				await SetMessageTimestampAsync(state.Id, state.Channel, state.Ts);
 			}
 			else if (!String.IsNullOrEmpty(state.Ts))
 			{
 				_logger.LogInformation("Updating existing slack message {MessageId} for user {SlackUser} ({Channel}, {MessageTs})", state.Id, recipient, state.Channel, state.Ts);
-				await _slackClient.UpdateMessageAsync(state.Channel, state.Ts, message);
+				if (threadTs == null)
+				{
+					await _slackClient.UpdateMessageAsync(state.Channel, state.Ts, message);
+				}
+				else
+				{
+					await _slackClient.UpdateMessageAsync(state.Channel, state.Ts, threadTs, message);
+				}
 			}
 			return (state, isNew);
 		}
@@ -2122,7 +2203,7 @@ namespace Horde.Build.Notifications.Sinks
 					if (payload.View != null && payload.View.State != null && payload.View.CallbackId != null)
 					{
 						Match? match;
-						if(TryMatch(payload.View.CallbackId, @"^issue_(\d+)_markfixed_([a-fA-F0-9]{24})$", out match))
+						if (TryMatch(payload.View.CallbackId, @"^issue_(\d+)_markfixed_([a-fA-F0-9]{24})$", out match))
 						{
 							int issueId = Int32.Parse(match.Groups[1].Value, NumberStyles.Integer, CultureInfo.InvariantCulture);
 							UserId userId = match.Groups[2].Value.ToObjectId<IUser>();
@@ -2142,6 +2223,12 @@ namespace Horde.Build.Notifications.Sinks
 								_logger.LogInformation("Marked issue {IssueId} fixed by user {UserId} in {Change}", issueId, userId, fixChange);
 							}
 						}
+						else if (TryMatch(payload.View.CallbackId, @"^issue_(\d+)_ack_([a-fA-F0-9]{24})$", out match))
+						{
+							int issueId = Int32.Parse(match.Groups[1].Value, NumberStyles.Integer, CultureInfo.InvariantCulture);
+							UserId userId = match.Groups[2].Value.ToObjectId<IUser>();
+							await _issueService.UpdateIssueAsync(issueId, acknowledged: true, ownerId: userId);
+						}
 					}
 				}
 			}
@@ -2216,10 +2303,54 @@ namespace Horde.Build.Notifications.Sinks
 
 			if (String.Equals(verb, "ack", StringComparison.Ordinal))
 			{
+				IIssue? issue = await _issueService.Collection.GetIssueAsync(issueId);
+				if (issue == null)
+				{
+					_logger.LogWarning("Unable to find issue {IssueId}", issueId);
+					return;
+				}
+
+				if (issue.OwnerId != null && issue.OwnerId != user.Id)
+				{
+					IUser? owner = await _userCollection.GetUserAsync(issue.OwnerId.Value);
+					if (owner != null)
+					{
+						SlackView view = new SlackView($"Issue {issueId}");
+						view.CallbackId = $"issue_{issueId}_ack_{user.Id}";
+						view.AddSection($"This issue is current assigned to {owner.Name}. Assign it to yourself?");
+						view.Close = "Cancel";
+						view.Submit = "Assign to Me";
+
+						await _slackClient.OpenViewAsync(triggerId, view);
+						return;
+					}
+				}
+
 				await _issueService.UpdateIssueAsync(issueId, acknowledged: true, ownerId: user.Id);
 			}
 			else if (String.Equals(verb, "decline", StringComparison.Ordinal))
 			{
+				IIssue? issue = await _issueService.Collection.GetIssueAsync(issueId);
+				if (issue == null)
+				{
+					_logger.LogWarning("Unable to find issue {IssueId}", issueId);
+					return;
+				}
+
+				if (issue.OwnerId != user.Id)
+				{
+					List<IIssueSuspect> suspects = await _issueService.Collection.FindSuspectsAsync(issue);
+					if (!suspects.Any(x => x.AuthorId == user.Id))
+					{
+						SlackView view = new SlackView($"Issue {issueId}");
+						view.AddSection("You are not currently listed as a suspect for this issue.");
+						view.Close = "Cancel";
+
+						await _slackClient.OpenViewAsync(triggerId, view);
+						return;
+					}
+				}
+
 				await _issueService.UpdateIssueAsync(issueId, declinedById: user.Id);
 			}
 			else if (String.Equals(verb, "markfixed", StringComparison.Ordinal))
