Horde: Change log view clicks to open in same tab
#rnx

diff --git a/Engine/Source/Programs/Horde/Horde.Build/Devices/DeviceCollection.cs b/Engine/Source/Programs/Horde/Horde.Build/Devices/DeviceCollection.cs
--- a/Engine/Source/Programs/Horde/Horde.Build/Devices/DeviceCollection.cs
+++ b/Engine/Source/Programs/Horde/Horde.Build/Devices/DeviceCollection.cs
@@ -263,7 +263,7 @@ namespace Horde.Build.Devices
 			/// <summary>
 			/// Id of telemetry document
 			/// </summary>
-			[BsonRequired, BsonId, BsonElement("tid")]
+			[BsonRequired, BsonId]
 			public ObjectId TelemetryId { get; set; }
 
 			/// <summary>
@@ -358,19 +358,19 @@ namespace Horde.Build.Devices
 		class DeviceReservationPoolTelemetryDocument : IDevicePoolReservationTelemetry
 		{
 			/// <inheritdoc/>
-			[BsonRequired]
+			[BsonRequired, BsonElement("did")]
 			public DeviceId DeviceId { get; set; }
 
-			[BsonIgnoreIfNull]
+			[BsonIgnoreIfNull, BsonElement("ji")]
 			public string? JobId { get; set; }
 
-			[BsonIgnoreIfNull]
+			[BsonIgnoreIfNull, BsonElement("si")]
 			public string? StepId { get; set; }
 
-			[BsonIgnoreIfNull]
+			[BsonIgnoreIfNull, BsonElement("jn")]
 			public string? JobName { get; set; }
 
-			[BsonIgnoreIfNull]
+			[BsonIgnoreIfNull, BsonElement("sn")]
 			public string? StepName { get; set; }
 
 			[BsonConstructor]
diff --git a/Engine/Source/Programs/Horde/HordeDashboard/src/components/jobDetailsV2/JobDetailViewSteps.tsx b/Engine/Source/Programs/Horde/HordeDashboard/src/components/jobDetailsV2/JobDetailViewSteps.tsx
--- a/Engine/Source/Programs/Horde/HordeDashboard/src/components/jobDetailsV2/JobDetailViewSteps.tsx
+++ b/Engine/Source/Programs/Horde/HordeDashboard/src/components/jobDetailsV2/JobDetailViewSteps.tsx
@@ -267,7 +267,7 @@ export const StepsPanelInner: React.FC<{ jobDetails: JobDetailsV2, depStepId?: s
 
          case 'ViewLogColumn':
             return <React.Fragment>
-               {!disabled && <Link onClick={(ev) => { ev.stopPropagation(); }} to={`/log/${item.batch?.logId}`} target="_blank"> <Stack horizontal horizontalAlign={"end"} verticalAlign="center" tokens={{ childrenGap: 0, padding: 0 }} style={{ width: "100%", height: "100%" }}>
+               {!disabled && <Link onClick={(ev) => { ev.stopPropagation(); }} to={`/log/${item.batch?.logId}`}> <Stack horizontal horizontalAlign={"end"} verticalAlign="center" tokens={{ childrenGap: 0, padding: 0 }} style={{ width: "100%", height: "100%" }}>
                   <Text styles={{ root: { margin: '0px', padding: '0px', paddingRight: '14px' } }} className={"view-log-link"}>View Log</Text>
                </Stack></Link>}
             </React.Fragment>;
@@ -287,7 +287,7 @@ export const StepsPanelInner: React.FC<{ jobDetails: JobDetailsV2, depStepId?: s
 
          // see data-automation-key in onRenderRow selectors
          return <React.Fragment>
-            {!disabled && <Link onClick={(ev) => { ev.stopPropagation(); }} to={`/log/${item.step!.logId}`} target="_blank"> <Stack horizontal verticalAlign="center" tokens={{ childrenGap: 0, padding: 0 }} style={{ paddingLeft: 33, width: "100%", height: "100%" }}>
+            {!disabled && <Link onClick={(ev) => { ev.stopPropagation(); }} to={`/log/${item.step!.logId}`}> <Stack horizontal verticalAlign="center" tokens={{ childrenGap: 0, padding: 0 }} style={{ paddingLeft: 33, width: "100%", height: "100%" }}>
                <Text styles={{ root: { margin: '0px', padding: '0px', paddingRight: '14px' } }} className={"view-log-link"}>View Log</Text>
             </Stack></Link>}
          </React.Fragment>;
diff --git a/Engine/Source/Programs/Horde/HordeDashboard/src/components/jobDetailsV2/JobOperationsBar.tsx b/Engine/Source/Programs/Horde/HordeDashboard/src/components/jobDetailsV2/JobOperationsBar.tsx
--- a/Engine/Source/Programs/Horde/HordeDashboard/src/components/jobDetailsV2/JobOperationsBar.tsx
+++ b/Engine/Source/Programs/Horde/HordeDashboard/src/components/jobDetailsV2/JobOperationsBar.tsx
@@ -155,7 +155,7 @@ export const JobOperations: React.FC<{ jobDetails: JobDetailsV2 }> = observer(({
                <StepOperations jobDetails={jobDetails} stepId={stepId} />
             </Stack>}
 
-            {!!step && !viewLogDisabled && <Link to={`${viewLogDisabled ? "" : logUrl}`} target="_blank">
+            {!!step && !viewLogDisabled && <Link to={`${viewLogDisabled ? "" : logUrl}`}>
                <Stack styles={{ root: { paddingTop: 4 } }}>
                   <CommandBarButton disabled={viewLogDisabled} className={hordeClasses.commandBarSmall} styles={{ root: { padding: "10px 8px 10px 8px" } }} iconProps={{ iconName: "AlignLeft" }} text="View Log" />
                </Stack>
