Horde: Add a setting for an alias to ping if an issue needs triaging.

diff --git a/Engine/Source/Programs/Horde/Horde.Build/Notifications/Sinks/SlackNotificationSink.cs b/Engine/Source/Programs/Horde/Horde.Build/Notifications/Sinks/SlackNotificationSink.cs
--- a/Engine/Source/Programs/Horde/Horde.Build/Notifications/Sinks/SlackNotificationSink.cs
+++ b/Engine/Source/Programs/Horde/Horde.Build/Notifications/Sinks/SlackNotificationSink.cs
@@ -923,6 +923,12 @@ namespace Horde.Build.Notifications.Sinks
 							string suspectMessage = $"Possibly {StringUtils.FormatList(suspectList, "or")}.";
 							await _slackClient.PostMessageAsync(triageChannel, state.Ts, suspectMessage);
 						}
+
+						if (workflow.TriageAlias != null)
+						{
+							string triageMessage = $"(<@{workflow.TriageAlias}> for triage).";
+							await _slackClient.PostMessageAsync(triageChannel, state.Ts, triageMessage);
+						}
 					}
 
 					if (_environment.IsProduction() && workflow.AllowMentions)
diff --git a/Engine/Source/Programs/Horde/Horde.Build/Server/Config.cs b/Engine/Source/Programs/Horde/Horde.Build/Server/Config.cs
--- a/Engine/Source/Programs/Horde/Horde.Build/Server/Config.cs
+++ b/Engine/Source/Programs/Horde/Horde.Build/Server/Config.cs
@@ -379,6 +379,11 @@ namespace Horde.Build.Server
 		/// </summary>
 		public string? TriageInstructions { get; set; }
 
+		/// <summary>
+		/// User id of a Slack user/alias to ping if there is nobody assigned to an issue by default.
+		/// </summary>
+		public string? TriageAlias { get; set; }
+
 		/// <summary>
 		/// Maximum number of people to mention on a triage thread
 		/// </summary>
