Horde: Support for escalating an issue to a named recipient if it is not resolved for a certain period of time.

diff --git a/Engine/Source/Programs/Horde/Horde.Build/Notifications/Sinks/SlackNotificationSink.cs b/Engine/Source/Programs/Horde/Horde.Build/Notifications/Sinks/SlackNotificationSink.cs
--- a/Engine/Source/Programs/Horde/Horde.Build/Notifications/Sinks/SlackNotificationSink.cs
+++ b/Engine/Source/Programs/Horde/Horde.Build/Notifications/Sinks/SlackNotificationSink.cs
@@ -15,6 +15,7 @@ using System.Text.RegularExpressions;
 using System.Threading;
 using System.Threading.Tasks;
 using EpicGames.Core;
+using EpicGames.Redis;
 using EpicGames.Slack;
 using EpicGames.Slack.Blocks;
 using EpicGames.Slack.Elements;
@@ -215,6 +216,8 @@ namespace Horde.Build.Notifications.Sinks
 		readonly HashSet<string>? _allowUsers;
 		readonly IExternalIssueService _externalIssueService;
 		readonly JsonSerializerOptions _jsonSerializerOptions;
+		readonly RedisSortedSet<int> _escalateIssues;
+		readonly IClock _clock;
 		readonly ILogger _logger;
 
 		readonly HttpClient _httpClient;
@@ -228,7 +231,7 @@ namespace Horde.Build.Notifications.Sinks
 		/// <summary>
 		/// Constructor
 		/// </summary>
-		public SlackNotificationSink(MongoService mongoService, IssueService issueService, IUserCollection userCollection, ILogFileService logFileService, StreamService streamService, IExternalIssueService externalIssueService, IWebHostEnvironment environment, IOptions<ServerSettings> settings, ILogger<SlackNotificationSink> logger)
+		public SlackNotificationSink(MongoService mongoService, RedisService redisService, IssueService issueService, IUserCollection userCollection, ILogFileService logFileService, StreamService streamService, IExternalIssueService externalIssueService, IWebHostEnvironment environment, IOptions<ServerSettings> settings, IClock clock, ILogger<SlackNotificationSink> logger)
 		{
 			_issueService = issueService;
 			_userCollection = userCollection;
@@ -239,6 +242,8 @@ namespace Horde.Build.Notifications.Sinks
 			_settings = settings.Value;
 			_messageStates = mongoService.Database.GetCollection<MessageStateDocument>("Slack");
 			_slackUsers = mongoService.Database.GetCollection<SlackUserDocument>("Slack.UsersV2");
+			_escalateIssues = new RedisSortedSet<int>(redisService.ConnectionPool, "slack/escalate");
+			_clock = clock;
 			_logger = logger;
 
 			_httpClient = new HttpClient();
@@ -935,6 +940,12 @@ namespace Horde.Build.Notifications.Sinks
 					{
 						await InviteUsersAsync(state.Channel, inviteUserIds);
 					}
+
+					if (workflow.EscalateAlias != null && workflow.EscalateTimes.Count > 0)
+					{
+						DateTime escalateTime = DateTime.UtcNow + TimeSpan.FromMinutes(workflow.EscalateTimes[0]);
+						await _escalateIssues.AddAsync(issue.Id, escalateTime.Ticks);
+					}
 				}
 
 				if (issue.AcknowledgedAt != null)
@@ -1273,19 +1284,6 @@ namespace Horde.Build.Notifications.Sinks
 			}
 		}
 
-		async Task<string> FormatNameAsync(UserId userId, bool allowMentions)
-		{
-			IUser? user = await _userCollection.GetUserAsync(userId);
-			if (user == null)
-			{
-				return $"User {userId}";
-			}
-			else
-			{
-				return user.Name;
-			}
-		}
-
 		async Task<string> FormatMentionAsync(UserId userId, bool allowMentions)
 		{
 			IUser? user = await _userCollection.GetUserAsync(userId);
@@ -2034,33 +2032,135 @@ namespace Horde.Build.Notifications.Sinks
 			return (state, isNew);
 		}
 
+		async ValueTask EscalateAsync(CancellationToken cancellationToken)
+		{
+			DateTime utcNow = DateTime.UtcNow;
+
+			int[] issueIds = await _escalateIssues.RangeByScoreAsync(0, utcNow.Ticks);
+			foreach (int issueId in issueIds)
+			{
+				cancellationToken.ThrowIfCancellationRequested();
+
+				double? nextTime = await EscalateSingleIssueAsync(issueId, utcNow);
+
+				await _escalateIssues.RemoveAsync(issueId);
+
+				if (nextTime != null)
+				{
+					await _escalateIssues.AddAsync(issueId, nextTime.Value);
+				}
+			}
+		}
+
+		async Task<double?> EscalateSingleIssueAsync(int issueId, DateTime utcNow)
+		{
+			IIssue? issue = await _issueService.Collection.GetIssueAsync(issueId);
+			if (issue == null)
+			{
+				return null;
+			}
+
+			List<IIssueSpan> spans = await _issueService.Collection.FindSpansAsync(issueId);
+			if (spans.Count == 0)
+			{
+				return null;
+			}
+
+			IIssueSpan span = spans[0];
+
+			WorkflowId? workflowId = span.LastFailure.Annotations.WorkflowId;
+			if (workflowId == null)
+			{
+				return null;
+			}
+
+			IStream? stream = await _streamService.GetStreamAsync(span.StreamId);
+			if (stream == null || !stream.Config.TryGetWorkflow(workflowId.Value, out WorkflowConfig? workflow))
+			{
+				return null;
+			}
+			if (workflow.TriageChannel == null || workflow.EscalateAlias == null || workflow.EscalateTimes.Count == 0)
+			{
+				return null;
+			}
+
+			if (issue.QuarantineTimeUtc == null)
+			{
+				MessageStateDocument? state = await GetMessageStateAsync(workflow.TriageChannel, GetTriageThreadEventId(issueId));
+				if (state == null)
+				{
+					return null;
+				}
+
+				TimeSpan openTime = utcNow - span.FirstFailure.StepTime;
+
+				string openTimeStr;
+				if (openTime < TimeSpan.FromHours(1.0))
+				{
+					openTimeStr = $"{(int)openTime.TotalMinutes}m";
+				}
+				else if (openTime < TimeSpan.FromDays(1.0))
+				{
+					openTimeStr = $"{(int)openTime.TotalHours}h";
+				}
+				else
+				{
+					openTimeStr = $"{(int)openTime.TotalDays}d";
+				}
+
+				await _slackClient.PostMessageAsync(workflow.TriageChannel, state.Ts, $"<@{workflow.EscalateAlias}>: Issue has not been resolved after {openTimeStr}.");
+			}
+
+			DateTime nextEscalationTime = span.FirstFailure.StepTime;
+			for (int idx = 0; idx < workflow.EscalateTimes.Count && nextEscalationTime < utcNow; idx++)
+			{
+				nextEscalationTime += TimeSpan.FromMinutes(workflow.EscalateTimes[idx]);
+			}
+			if (nextEscalationTime < utcNow)
+			{
+				nextEscalationTime = utcNow + TimeSpan.FromMinutes(workflow.EscalateTimes[^1]);
+			}
+
+			return nextEscalationTime.Ticks;
+		}
+
 		/// <inheritdoc/>
 		protected override async Task ExecuteAsync(CancellationToken stoppingToken)
 		{
 			if (!String.IsNullOrEmpty(_settings.SlackSocketToken))
 			{
-				while (!stoppingToken.IsCancellationRequested)
+				using ITicker ticker = _clock.AddSharedTicker<SlackNotificationSink>(TimeSpan.FromMinutes(1.0), EscalateAsync, _logger);
+				await ticker.StartAsync();
+
+				try
 				{
-					try
+					while (!stoppingToken.IsCancellationRequested)
 					{
-						Uri? webSocketUrl = await GetWebSocketUrlAsync(stoppingToken);
-						if (webSocketUrl == null)
+						try
+						{
+							Uri? webSocketUrl = await GetWebSocketUrlAsync(stoppingToken);
+							if (webSocketUrl == null)
+							{
+								await Task.Delay(TimeSpan.FromSeconds(5.0), stoppingToken);
+								continue;
+							}
+							await HandleSocketAsync(webSocketUrl, stoppingToken);
+						}
+						catch (OperationCanceledException)
 						{
+							break;
+						}
+						catch (Exception ex)
+						{
+							_logger.LogError(ex, "Exception while updating Slack socket");
 							await Task.Delay(TimeSpan.FromSeconds(5.0), stoppingToken);
-							continue;
 						}
-						await HandleSocketAsync(webSocketUrl, stoppingToken);
-					}
-					catch (OperationCanceledException)
-					{
-						break;
-					}
-					catch (Exception ex)
-					{
-						_logger.LogError(ex, "Exception while updating Slack socket");
-						await Task.Delay(TimeSpan.FromSeconds(5.0), stoppingToken);
 					}
 				}
+				finally
+				{
+					await ticker.StopAsync();
+				}
 			}
 		}
 
diff --git a/Engine/Source/Programs/Horde/Horde.Build/Server/Config.cs b/Engine/Source/Programs/Horde/Horde.Build/Server/Config.cs
--- a/Engine/Source/Programs/Horde/Horde.Build/Server/Config.cs
+++ b/Engine/Source/Programs/Horde/Horde.Build/Server/Config.cs
@@ -384,6 +384,16 @@ namespace Horde.Build.Server
 		/// </summary>
 		public string? TriageAlias { get; set; }
 
+		/// <summary>
+		/// Alias to ping if an issue has not been resolved for a certain amount of time
+		/// </summary>
+		public string? EscalateAlias { get; set; }
+
+		/// <summary>
+		/// Times after an issue has been opened to escalate to the alias above, in minutes. Continues to notify on the last interval once reaching the end of the list.
+		/// </summary>
+		public List<int> EscalateTimes { get; } = new List<int>{ 120 };
+
 		/// <summary>
 		/// Maximum number of people to mention on a triage thread
 		/// </summary>
