Horde: Add support for mentioning subteams in Slack messages.

diff --git a/Engine/Source/Programs/Horde/Horde.Build/Notifications/Sinks/SlackNotificationSink.cs b/Engine/Source/Programs/Horde/Horde.Build/Notifications/Sinks/SlackNotificationSink.cs
--- a/Engine/Source/Programs/Horde/Horde.Build/Notifications/Sinks/SlackNotificationSink.cs
+++ b/Engine/Source/Programs/Horde/Horde.Build/Notifications/Sinks/SlackNotificationSink.cs
@@ -931,7 +931,7 @@ namespace Horde.Build.Notifications.Sinks
 
 						if (workflow.TriageAlias != null)
 						{
-							string triageMessage = $"(<@{workflow.TriageAlias}> for triage).";
+							string triageMessage = $"({FormatUserOrGroupMention(workflow.TriageAlias)} for triage).";
 							await _slackClient.PostMessageAsync(triageChannel, state.Ts, triageMessage);
 						}
 					}
@@ -1284,6 +1284,18 @@ namespace Horde.Build.Notifications.Sinks
 			}
 		}
 
+		static string FormatUserOrGroupMention(string userOrGroup)
+		{
+			if (userOrGroup.StartsWith("S", StringComparison.OrdinalIgnoreCase))
+			{
+				return $"<!subteam^{userOrGroup}>";
+			}
+			else
+			{
+				return $"<@{userOrGroup}>";
+			}
+		}
+
 		async Task<string> FormatMentionAsync(UserId userId, bool allowMentions)
 		{
 			IUser? user = await _userCollection.GetUserAsync(userId);
@@ -2108,7 +2120,7 @@ namespace Horde.Build.Notifications.Sinks
 					openTimeStr = $"{(int)openTime.TotalDays}d";
 				}
 
-				await _slackClient.PostMessageAsync(workflow.TriageChannel, state.Ts, $"<@{workflow.EscalateAlias}>: Issue has not been resolved after {openTimeStr}.");
+				await _slackClient.PostMessageAsync(workflow.TriageChannel, state.Ts, $"{FormatUserOrGroupMention(workflow.EscalateAlias)}> - Issue has not been resolved after {openTimeStr}.");
 			}
 
 			DateTime nextEscalationTime = span.FirstFailure.StepTime;
