Horde: Add tracing for task expiration in ComputeService

diff --git a/Engine/Source/Programs/Horde/Horde.Build/Compute/ComputeService.cs b/Engine/Source/Programs/Horde/Horde.Build/Compute/ComputeService.cs
--- a/Engine/Source/Programs/Horde/Horde.Build/Compute/ComputeService.cs
+++ b/Engine/Source/Programs/Horde/Horde.Build/Compute/ComputeService.cs
@@ -23,6 +23,8 @@ using HordeCommon.Rpc.Tasks;
 using Microsoft.Extensions.Caching.Memory;
 using Microsoft.Extensions.Hosting;
 using Microsoft.Extensions.Logging;
+using OpenTracing;
+using OpenTracing.Util;
 using StackExchange.Redis;
 
 namespace Horde.Build.Compute
@@ -170,7 +172,10 @@ namespace Horde.Build.Compute
 		/// <returns></returns>
 		async ValueTask ExpireTasksAsync(CancellationToken cancellationToken)
 		{
+			using IScope scope = GlobalTracer.Instance.BuildSpan("ComputeService.ExpireTasksAsync").StartActive();
 			List<QueueKey> queueKeys = await _taskScheduler.GetInactiveQueuesAsync();
+			scope.Span.SetTag("numQueueKeys", queueKeys.Count);
+			
 			foreach (QueueKey queueKey in queueKeys)
 			{
 				_logger.LogInformation("Inactive queue: {QueueKey}", queueKey);
