Horde: Disable unfurling on report messages. Also only notify triage alias if no suspects are identified.

diff --git a/Engine/Source/Programs/Horde/Horde.Build/Notifications/Sinks/SlackNotificationSink.cs b/Engine/Source/Programs/Horde/Horde.Build/Notifications/Sinks/SlackNotificationSink.cs
--- a/Engine/Source/Programs/Horde/Horde.Build/Notifications/Sinks/SlackNotificationSink.cs
+++ b/Engine/Source/Programs/Horde/Horde.Build/Notifications/Sinks/SlackNotificationSink.cs
@@ -928,10 +928,9 @@ namespace Horde.Build.Notifications.Sinks
 							string suspectMessage = $"Possibly {StringUtils.FormatList(suspectList, "or")}.";
 							await _slackClient.PostMessageAsync(triageChannel, state.Ts, suspectMessage);
 						}
-
-						if (workflow.TriageAlias != null)
+						else if (workflow.TriageAlias != null)
 						{
-							string triageMessage = $"({FormatUserOrGroupMention(workflow.TriageAlias)} for triage).";
+							string triageMessage = $"(cc {FormatUserOrGroupMention(workflow.TriageAlias)} for triage).";
 							await _slackClient.PostMessageAsync(triageChannel, state.Ts, triageMessage);
 						}
 					}
@@ -1598,7 +1597,10 @@ namespace Horde.Build.Notifications.Sinks
 				body.Append(text);
 			}
 
-			await SendOrUpdateMessageAsync(channel, eventId, null, body.ToString());
+			SlackMessage message = body.ToString();
+			message.UnfurlLinks = false;
+			message.UnfurlMedia = false;
+			await SendOrUpdateMessageAsync(channel, eventId, null, message);
 		}
 
 		string GetSeverityPrefix(IssueSeverity severity)
