Horde: Fix typo in escalation message.

diff --git a/Engine/Source/Programs/Horde/Horde.Build/Notifications/Sinks/SlackNotificationSink.cs b/Engine/Source/Programs/Horde/Horde.Build/Notifications/Sinks/SlackNotificationSink.cs
--- a/Engine/Source/Programs/Horde/Horde.Build/Notifications/Sinks/SlackNotificationSink.cs
+++ b/Engine/Source/Programs/Horde/Horde.Build/Notifications/Sinks/SlackNotificationSink.cs
@@ -2122,7 +2122,7 @@ namespace Horde.Build.Notifications.Sinks
 					openTimeStr = $"{(int)openTime.TotalDays}d";
 				}
 
-				await _slackClient.PostMessageAsync(workflow.TriageChannel, state.Ts, $"{FormatUserOrGroupMention(workflow.EscalateAlias)}> - Issue has not been resolved after {openTimeStr}.");
+				await _slackClient.PostMessageAsync(workflow.TriageChannel, state.Ts, $"{FormatUserOrGroupMention(workflow.EscalateAlias)} - Issue has not been resolved after {openTimeStr}.");
 			}
 
 			DateTime nextEscalationTime = span.FirstFailure.StepTime;
