Horde: Register ComputeService as a hosted service to enable background ticker

diff --git a/Engine/Source/Programs/Horde/Horde.Build/Startup.cs b/Engine/Source/Programs/Horde/Horde.Build/Startup.cs
--- a/Engine/Source/Programs/Horde/Horde.Build/Startup.cs
+++ b/Engine/Source/Programs/Horde/Horde.Build/Startup.cs
@@ -647,6 +647,7 @@ namespace Horde.Build
 			services.AddHostedService<ConformTaskSource>(provider => provider.GetRequiredService<ConformTaskSource>());
 			services.AddSingleton<ComputeService>();
 			services.AddSingleton<IComputeService, ComputeService>();
+			services.AddHostedService(provider => provider.GetRequiredService<ComputeService>());
 
 			services.AddSingleton<ITaskSource, UpgradeTaskSource>();
 			services.AddSingleton<ITaskSource, ShutdownTaskSource>();
