Horde: Stop escalating issues once they're resolved for a workflow.

diff --git a/Engine/Source/Programs/Horde/Horde.Build/Notifications/Sinks/SlackNotificationSink.cs b/Engine/Source/Programs/Horde/Horde.Build/Notifications/Sinks/SlackNotificationSink.cs
--- a/Engine/Source/Programs/Horde/Horde.Build/Notifications/Sinks/SlackNotificationSink.cs
+++ b/Engine/Source/Programs/Horde/Horde.Build/Notifications/Sinks/SlackNotificationSink.cs
@@ -2098,6 +2098,11 @@ namespace Horde.Build.Notifications.Sinks
 				return null;
 			}
 
+			if (!IsIssueOpenForWorkflow(stream.Id, span.TemplateRefId, workflow.Id, spans))
+			{
+				return null;
+			}
+
 			if (issue.QuarantineTimeUtc == null)
 			{
 				MessageStateDocument? state = await GetMessageStateAsync(workflow.TriageChannel, GetTriageThreadEventId(issueId));
