Horde: Prevent same notification message being sent more than once.

diff --git a/Engine/Source/Programs/Horde/Horde.Build/Notifications/Sinks/SlackNotificationSink.cs b/Engine/Source/Programs/Horde/Horde.Build/Notifications/Sinks/SlackNotificationSink.cs
--- a/Engine/Source/Programs/Horde/Horde.Build/Notifications/Sinks/SlackNotificationSink.cs
+++ b/Engine/Source/Programs/Horde/Horde.Build/Notifications/Sinks/SlackNotificationSink.cs
@@ -942,8 +942,8 @@ namespace Horde.Build.Notifications.Sinks
 
 					if (workflow.EscalateAlias != null && workflow.EscalateTimes.Count > 0)
 					{
-						DateTime escalateTime = DateTime.UtcNow + TimeSpan.FromMinutes(workflow.EscalateTimes[0]);
-						await _escalateIssues.AddAsync(issue.Id, escalateTime.Ticks);
+						DateTime escalateTime = span.FirstFailure.StepTime + TimeSpan.FromMinutes(workflow.EscalateTimes[0]);
+						await _escalateIssues.AddAsync(issue.Id, escalateTime.Ticks, StackExchange.Redis.When.NotExists);
 					}
 				}
 
@@ -2056,10 +2056,11 @@ namespace Horde.Build.Notifications.Sinks
 				cancellationToken.ThrowIfCancellationRequested();
 
 				double? nextTime = await EscalateSingleIssueAsync(issueId, utcNow);
-
-				await _escalateIssues.RemoveAsync(issueId);
-
-				if (nextTime != null)
+				if (nextTime == null)
+				{
+					await _escalateIssues.RemoveAsync(issueId);
+				}
+				else
 				{
 					await _escalateIssues.AddAsync(issueId, nextTime.Value);
 				}
