Horde: Ping the triage alias once all suspects have declined an issue.

diff --git a/Engine/Source/Programs/Horde/Horde.Build/Notifications/Sinks/SlackNotificationSink.cs b/Engine/Source/Programs/Horde/Horde.Build/Notifications/Sinks/SlackNotificationSink.cs
--- a/Engine/Source/Programs/Horde/Horde.Build/Notifications/Sinks/SlackNotificationSink.cs
+++ b/Engine/Source/Programs/Horde/Horde.Build/Notifications/Sinks/SlackNotificationSink.cs
@@ -928,11 +928,6 @@ namespace Horde.Build.Notifications.Sinks
 							string suspectMessage = $"Possibly {StringUtils.FormatList(suspectList, "or")}.";
 							await _slackClient.PostMessageAsync(triageChannel, state.Ts, suspectMessage);
 						}
-						else if (workflow.TriageAlias != null)
-						{
-							string triageMessage = $"(cc {FormatUserOrGroupMention(workflow.TriageAlias)} for triage).";
-							await _slackClient.PostMessageAsync(triageChannel, state.Ts, triageMessage);
-						}
 					}
 
 					if (_environment.IsProduction() && workflow.AllowMentions)
@@ -947,6 +942,12 @@ namespace Horde.Build.Notifications.Sinks
 					}
 				}
 
+				if (workflow.TriageAlias != null && issue.OwnerId == null && suspects.All(x => x.DeclinedAt != null))
+				{
+					string triageMessage = $"(cc {FormatUserOrGroupMention(workflow.TriageAlias)} for triage).";
+					await SendOrUpdateMessageAsync(triageChannel, state.Ts, eventId + "_triage", null, triageMessage);
+				}
+
 				if (issue.AcknowledgedAt != null)
 				{
 					await _slackClient.AddReactionAsync(state.Channel, state.Ts, "eyes");
