Horde Storage - Made the blob index consistency check (which uses ExistsInRootStore) not use the blob index for S3 by adding ability to override if we can do optimizations like that. This should signficantly reduce the amount of blob index lookups that happen when the consistency check is enabled.

diff --git a/Engine/Source/Programs/Horde/HordeStorage/Horde.Storage/Implementation/Blob/AmazonS3Store.cs b/Engine/Source/Programs/Horde/HordeStorage/Horde.Storage/Implementation/Blob/AmazonS3Store.cs
--- a/Engine/Source/Programs/Horde/HordeStorage/Horde.Storage/Implementation/Blob/AmazonS3Store.cs
+++ b/Engine/Source/Programs/Horde/HordeStorage/Horde.Storage/Implementation/Blob/AmazonS3Store.cs
@@ -166,10 +166,10 @@ namespace Horde.Storage.Implementation
             return new BlobContents(new InstrumentedStream(response.ResponseStream, "s3-get"), response.ContentLength);
         }
 
-        public async Task<bool> Exists(NamespaceId ns, BlobIdentifier blobIdentifier)
+        public async Task<bool> Exists(NamespaceId ns, BlobIdentifier blobIdentifier, bool forceCheck)
         {
             NamespacePolicy policies = _namespacePolicyResolver.GetPoliciesForNs(ns);
-            if (_settings.UseBlobIndexForExistsCheck && policies.UseBlobIndexForSlowExists)
+            if (_settings.UseBlobIndexForExistsCheck && policies.UseBlobIndexForSlowExists && !forceCheck)
             {
                 return await _blobIndex.BlobExistsInRegion(ns, blobIdentifier);
             }
diff --git a/Engine/Source/Programs/Horde/HordeStorage/Horde.Storage/Implementation/Blob/AzureBlobStore.cs b/Engine/Source/Programs/Horde/HordeStorage/Horde.Storage/Implementation/Blob/AzureBlobStore.cs
--- a/Engine/Source/Programs/Horde/HordeStorage/Horde.Storage/Implementation/Blob/AzureBlobStore.cs
+++ b/Engine/Source/Programs/Horde/HordeStorage/Horde.Storage/Implementation/Blob/AzureBlobStore.cs
@@ -135,7 +135,7 @@ namespace Horde.Storage.Implementation
             }
         }
 
-        public async Task<bool> Exists(NamespaceId ns, BlobIdentifier blobIdentifier)
+        public async Task<bool> Exists(NamespaceId ns, BlobIdentifier blobIdentifier, bool forceCheck)
         {
             BlobContainerClient container = new BlobContainerClient(_connectionString, SanitizeNamespace(ns));
             if (!await container.ExistsAsync())
diff --git a/Engine/Source/Programs/Horde/HordeStorage/Horde.Storage/Implementation/Blob/CachedBlobIndex.cs b/Engine/Source/Programs/Horde/HordeStorage/Horde.Storage/Implementation/Blob/CachedBlobIndex.cs
--- a/Engine/Source/Programs/Horde/HordeStorage/Horde.Storage/Implementation/Blob/CachedBlobIndex.cs
+++ b/Engine/Source/Programs/Horde/HordeStorage/Horde.Storage/Implementation/Blob/CachedBlobIndex.cs
@@ -58,7 +58,7 @@ public class CachedBlobIndex : IBlobIndex
 
     public async Task<bool> BlobExistsInRegion(NamespaceId ns, BlobIdentifier blobIdentifier)
     {
-        return await _fileSystemStore.Exists(ns, blobIdentifier);
+        return await _fileSystemStore.Exists(ns, blobIdentifier, forceCheck: false);
     }
 
     public async Task AddRefToBlobs(NamespaceId ns, BucketId bucket, IoHashKey key, BlobIdentifier[] blobs)
diff --git a/Engine/Source/Programs/Horde/HordeStorage/Horde.Storage/Implementation/Blob/FileSystemStore.cs b/Engine/Source/Programs/Horde/HordeStorage/Horde.Storage/Implementation/Blob/FileSystemStore.cs
--- a/Engine/Source/Programs/Horde/HordeStorage/Horde.Storage/Implementation/Blob/FileSystemStore.cs
+++ b/Engine/Source/Programs/Horde/HordeStorage/Horde.Storage/Implementation/Blob/FileSystemStore.cs
@@ -147,7 +147,7 @@ namespace Horde.Storage.Implementation
             }
         }
 
-        public Task<bool> Exists(NamespaceId ns, BlobIdentifier blob)
+        public Task<bool> Exists(NamespaceId ns, BlobIdentifier blob, bool forceCheck)
         {
             FileInfo filePath = GetFilesystemPath(ns, blob);
 
diff --git a/Engine/Source/Programs/Horde/HordeStorage/Horde.Storage/Implementation/Blob/IBlobService.cs b/Engine/Source/Programs/Horde/HordeStorage/Horde.Storage/Implementation/Blob/IBlobService.cs
--- a/Engine/Source/Programs/Horde/HordeStorage/Horde.Storage/Implementation/Blob/IBlobService.cs
+++ b/Engine/Source/Programs/Horde/HordeStorage/Horde.Storage/Implementation/Blob/IBlobService.cs
@@ -408,7 +408,7 @@ public class BlobService : IBlobService
     {
         IBlobStore store = _blobStores.Last();
 
-        using IScope scope = Tracer.Instance.StartActive("HierarchicalStore.ObjectExists");
+        using IScope scope = Tracer.Instance.StartActive("HierarchicalStore.ObjectExistsInRoot");
         scope.Span.ResourceName = blob.ToString();
         scope.Span.SetTag("BlobStore", store.GetType().Name);
         if (await store.Exists(ns, blob))
diff --git a/Engine/Source/Programs/Horde/HordeStorage/Horde.Storage/Implementation/Blob/IBlobStore.cs b/Engine/Source/Programs/Horde/HordeStorage/Horde.Storage/Implementation/Blob/IBlobStore.cs
--- a/Engine/Source/Programs/Horde/HordeStorage/Horde.Storage/Implementation/Blob/IBlobStore.cs
+++ b/Engine/Source/Programs/Horde/HordeStorage/Horde.Storage/Implementation/Blob/IBlobStore.cs
@@ -22,7 +22,7 @@ namespace Horde.Storage.Implementation
         Task<BlobIdentifier> PutObject(NamespaceId ns, Stream content, BlobIdentifier identifier);
 
         Task<BlobContents> GetObject(NamespaceId ns, BlobIdentifier blob, LastAccessTrackingFlags flags = LastAccessTrackingFlags.DoTracking);
-        Task<bool> Exists(NamespaceId ns, BlobIdentifier blob);
+        Task<bool> Exists(NamespaceId ns, BlobIdentifier blob, bool forceCheck = false);
 
         // Delete a object
         Task DeleteObject(NamespaceId ns, BlobIdentifier blob);
diff --git a/Engine/Source/Programs/Horde/HordeStorage/Horde.Storage/Implementation/Blob/MemoryBlobStore.cs b/Engine/Source/Programs/Horde/HordeStorage/Horde.Storage/Implementation/Blob/MemoryBlobStore.cs
--- a/Engine/Source/Programs/Horde/HordeStorage/Horde.Storage/Implementation/Blob/MemoryBlobStore.cs
+++ b/Engine/Source/Programs/Horde/HordeStorage/Horde.Storage/Implementation/Blob/MemoryBlobStore.cs
@@ -103,7 +103,7 @@ namespace Horde.Storage.Implementation
             return Task.CompletedTask;
         }
 
-        public Task<bool> Exists(NamespaceId ns, BlobIdentifier blob)
+        public Task<bool> Exists(NamespaceId ns, BlobIdentifier blob, bool forceCheck = false)
         {
             if (!_blobs.TryGetValue(ns, value: out ConcurrentDictionary<BlobIdentifier, BlobContainer>? namespaceContainer))
             {
diff --git a/Engine/Source/Programs/Horde/HordeStorage/Horde.Storage/Implementation/Blob/MemoryCacheBlobStore.cs b/Engine/Source/Programs/Horde/HordeStorage/Horde.Storage/Implementation/Blob/MemoryCacheBlobStore.cs
--- a/Engine/Source/Programs/Horde/HordeStorage/Horde.Storage/Implementation/Blob/MemoryCacheBlobStore.cs
+++ b/Engine/Source/Programs/Horde/HordeStorage/Horde.Storage/Implementation/Blob/MemoryCacheBlobStore.cs
@@ -94,7 +94,7 @@ namespace Horde.Storage.Implementation
             return Task.CompletedTask;
         }
 
-        public Task<bool> Exists(NamespaceId ns, BlobIdentifier blobIdentifier)
+        public Task<bool> Exists(NamespaceId ns, BlobIdentifier blobIdentifier, bool forceCheck = false)
         {
             return Task.FromResult(_memoryCache.TryGetValue(BuildKey(ns, blobIdentifier), out _));
         }
diff --git a/Engine/Source/Programs/Horde/HordeStorage/Horde.Storage/Implementation/Blob/RelayBlobStore.cs b/Engine/Source/Programs/Horde/HordeStorage/Horde.Storage/Implementation/Blob/RelayBlobStore.cs
--- a/Engine/Source/Programs/Horde/HordeStorage/Horde.Storage/Implementation/Blob/RelayBlobStore.cs
+++ b/Engine/Source/Programs/Horde/HordeStorage/Horde.Storage/Implementation/Blob/RelayBlobStore.cs
@@ -74,7 +74,7 @@ namespace Horde.Storage.Implementation
             return new BlobContents(await response.Content.ReadAsStreamAsync(), contentLength.Value);
         }
 
-        public async Task<bool> Exists(NamespaceId ns, BlobIdentifier blob)
+        public async Task<bool> Exists(NamespaceId ns, BlobIdentifier blob, bool forceCheck)
         {
             using HttpRequestMessage headObjectRequest = BuildHttpRequest(HttpMethod.Head, new Uri($"api/v1/blobs/{ns}/{blob}", UriKind.Relative));
             HttpResponseMessage response = await HttpClient.SendAsync(headObjectRequest);
