Horde: Add the issue number and link to notifications about issues not being resolved.

diff --git a/Engine/Source/Programs/Horde/Horde.Build/Notifications/Sinks/SlackNotificationSink.cs b/Engine/Source/Programs/Horde/Horde.Build/Notifications/Sinks/SlackNotificationSink.cs
--- a/Engine/Source/Programs/Horde/Horde.Build/Notifications/Sinks/SlackNotificationSink.cs
+++ b/Engine/Source/Programs/Horde/Horde.Build/Notifications/Sinks/SlackNotificationSink.cs
@@ -2129,7 +2129,8 @@ namespace Horde.Build.Notifications.Sinks
 					openTimeStr = $"{(int)openTime.TotalDays}d";
 				}
 
-				await _slackClient.PostMessageAsync(workflow.TriageChannel, state.Ts, $"{FormatUserOrGroupMention(workflow.EscalateAlias)} - Issue has not been resolved after {openTimeStr}.");
+				Uri issueUrl = GetIssueUrl(issue, span.FirstFailure);
+				await _slackClient.PostMessageAsync(workflow.TriageChannel, state.Ts, $"{FormatUserOrGroupMention(workflow.EscalateAlias)} - Issue <{issueUrl}|{issue.Id}> has not been resolved after {openTimeStr}.");
 			}
 
 			DateTime nextEscalationTime = span.FirstFailure.StepTime;
