Horde: Do not escalate issues that are marked fixed and have not yet failed.

diff --git a/Engine/Source/Programs/Horde/Horde.Build/Notifications/Sinks/SlackNotificationSink.cs b/Engine/Source/Programs/Horde/Horde.Build/Notifications/Sinks/SlackNotificationSink.cs
--- a/Engine/Source/Programs/Horde/Horde.Build/Notifications/Sinks/SlackNotificationSink.cs
+++ b/Engine/Source/Programs/Horde/Horde.Build/Notifications/Sinks/SlackNotificationSink.cs
@@ -736,6 +736,22 @@ namespace Horde.Build.Notifications.Sinks
 			await UpdateReportsAsync(issue, details.Spans);
 		}
 
+		static IIssueSpan? GetFixFailedSpan(IIssue issue, IReadOnlyList<IIssueSpan> spans)
+		{
+			if (issue.FixChange != null)
+			{
+				HashSet<StreamId> originStreams = new HashSet<StreamId>(issue.Streams.Where(x => x.MergeOrigin ?? false).Select(x => x.StreamId));
+				foreach (IIssueSpan originSpan in spans.Where(x => originStreams.Contains(x.StreamId)).OrderBy(x => x.StreamId.ToString()))
+				{
+					if (originSpan.LastFailure.Change > issue.FixChange.Value)
+					{
+						return originSpan;
+					}
+				}
+			}
+			return null;
+		}
+
 		async Task InviteUsersAsync(string channel, IEnumerable<UserId> userIds)
 		{
 			List<string> slackUserIds = new List<string>();
@@ -934,12 +950,12 @@ namespace Horde.Build.Notifications.Sinks
 					{
 						await InviteUsersAsync(state.Channel, inviteUserIds);
 					}
+				}
 
-					if (workflow.EscalateAlias != null && workflow.EscalateTimes.Count > 0)
-					{
-						DateTime escalateTime = span.FirstFailure.StepTime + TimeSpan.FromMinutes(workflow.EscalateTimes[0]);
-						await _escalateIssues.AddAsync(issue.Id, escalateTime.Ticks, StackExchange.Redis.When.NotExists);
-					}
+				if (workflow.EscalateAlias != null && workflow.EscalateTimes.Count > 0)
+				{
+					DateTime escalateTime = span.FirstFailure.StepTime + TimeSpan.FromMinutes(workflow.EscalateTimes[0]);
+					await _escalateIssues.AddAsync(issue.Id, escalateTime.Ticks, StackExchange.Redis.When.NotExists);
 				}
 
 				if (workflow.TriageAlias != null && issue.OwnerId == null && suspects.All(x => x.DeclinedAt != null))
@@ -969,15 +985,7 @@ namespace Horde.Build.Notifications.Sinks
 				IIssueSpan? fixFailedSpan = null;
 				if (issue.FixChange != null)
 				{
-					HashSet<StreamId> originStreams = new HashSet<StreamId>(issue.Streams.Where(x => x.MergeOrigin ?? false).Select(x => x.StreamId));
-					foreach (IIssueSpan originSpan in spans.Where(x => originStreams.Contains(x.StreamId)).OrderBy(x => x.StreamId.ToString()))
-					{
-						if (originSpan.LastFailure.Change > issue.FixChange.Value)
-						{
-							fixFailedSpan = originSpan;
-							break;
-						}
-					}
+					fixFailedSpan = GetFixFailedSpan(issue, spans);
 
 					if (fixFailedSpan == null)
 					{
@@ -2104,6 +2112,10 @@ namespace Horde.Build.Notifications.Sinks
 			{
 				return null;
 			}
+			if (issue.FixChange != null && GetFixFailedSpan(issue, spans) == null)
+			{
+				return null;
+			}
 
 			if (issue.QuarantineTimeUtc == null)
 			{
