Horde: Use seconds from Unix epoch for issue alerts, to prevent precision issues with tick counts converted to doubles in Redis.

diff --git a/Engine/Source/Programs/Horde/Horde.Build/Notifications/Sinks/SlackNotificationSink.cs b/Engine/Source/Programs/Horde/Horde.Build/Notifications/Sinks/SlackNotificationSink.cs
--- a/Engine/Source/Programs/Horde/Horde.Build/Notifications/Sinks/SlackNotificationSink.cs
+++ b/Engine/Source/Programs/Horde/Horde.Build/Notifications/Sinks/SlackNotificationSink.cs
@@ -955,7 +955,7 @@ namespace Horde.Build.Notifications.Sinks
 				if (workflow.EscalateAlias != null && workflow.EscalateTimes.Count > 0)
 				{
 					DateTime escalateTime = span.FirstFailure.StepTime + TimeSpan.FromMinutes(workflow.EscalateTimes[0]);
-					await _escalateIssues.AddAsync(issue.Id, escalateTime.Ticks, StackExchange.Redis.When.NotExists);
+					await _escalateIssues.AddAsync(issue.Id, (escalateTime - DateTime.UnixEpoch).TotalSeconds, StackExchange.Redis.When.NotExists);
 				}
 
 				if (workflow.TriageAlias != null && issue.OwnerId == null && suspects.All(x => x.DeclinedAt != null))
@@ -2059,7 +2059,7 @@ namespace Horde.Build.Notifications.Sinks
 		{
 			DateTime utcNow = DateTime.UtcNow;
 
-			int[] issueIds = await _escalateIssues.RangeByScoreAsync(0, utcNow.Ticks);
+			int[] issueIds = await _escalateIssues.RangeByScoreAsync(0, (utcNow - DateTime.UnixEpoch).TotalSeconds);
 			foreach (int issueId in issueIds)
 			{
 				cancellationToken.ThrowIfCancellationRequested();
@@ -2155,7 +2155,7 @@ namespace Horde.Build.Notifications.Sinks
 				nextEscalationTime = utcNow + TimeSpan.FromMinutes(workflow.EscalateTimes[^1]);
 			}
 
-			return nextEscalationTime.Ticks;
+			return (nextEscalationTime - DateTime.UnixEpoch).TotalSeconds;
 		}
 
 		/// <inheritdoc/>
