Horde: Additional logging for issue escalation.

diff --git a/Engine/Source/Programs/Horde/Horde.Build/Notifications/Sinks/SlackNotificationSink.cs b/Engine/Source/Programs/Horde/Horde.Build/Notifications/Sinks/SlackNotificationSink.cs
--- a/Engine/Source/Programs/Horde/Horde.Build/Notifications/Sinks/SlackNotificationSink.cs
+++ b/Engine/Source/Programs/Horde/Horde.Build/Notifications/Sinks/SlackNotificationSink.cs
@@ -955,7 +955,10 @@ namespace Horde.Build.Notifications.Sinks
 				if (workflow.EscalateAlias != null && workflow.EscalateTimes.Count > 0)
 				{
 					DateTime escalateTime = span.FirstFailure.StepTime + TimeSpan.FromMinutes(workflow.EscalateTimes[0]);
-					await _escalateIssues.AddAsync(issue.Id, (escalateTime - DateTime.UnixEpoch).TotalSeconds, StackExchange.Redis.When.NotExists);
+					if (await _escalateIssues.AddAsync(issue.Id, (escalateTime - DateTime.UnixEpoch).TotalSeconds, StackExchange.Redis.When.NotExists))
+					{
+						_logger.LogInformation("First escalation time for issue {IssueId} is {Time}", issue.Id, escalateTime);
+					}
 				}
 
 				if (workflow.TriageAlias != null && issue.OwnerId == null && suspects.All(x => x.DeclinedAt != null))
@@ -2146,13 +2149,23 @@ namespace Horde.Build.Notifications.Sinks
 			}
 
 			DateTime nextEscalationTime = span.FirstFailure.StepTime;
-			for (int idx = 0; idx < workflow.EscalateTimes.Count && nextEscalationTime < utcNow; idx++)
+			for (int idx = 0;;idx++)
 			{
+				if (idx >= workflow.EscalateTimes.Count)
+				{
+					nextEscalationTime = utcNow + TimeSpan.FromMinutes(workflow.EscalateTimes[^1]);
+					_logger.LogInformation("Next escalation time for issue {IssueId} is {NextEscalationTime} (now + {Time}m [{Idx}/{Count}])", issue.Id, nextEscalationTime, workflow.EscalateTimes[^1], workflow.EscalateTimes.Count, workflow.EscalateTimes.Count);
+					break;
+				}
+
+				DateTime prevEscalationTime = nextEscalationTime;
 				nextEscalationTime += TimeSpan.FromMinutes(workflow.EscalateTimes[idx]);
-			}
-			if (nextEscalationTime < utcNow)
-			{
-				nextEscalationTime = utcNow + TimeSpan.FromMinutes(workflow.EscalateTimes[^1]);
+
+				if (nextEscalationTime > utcNow)
+				{
+					_logger.LogInformation("Next escalation time for issue {IssueId} is {NextEscalationTime} ({PrevEscalationTime} + {Time}m [{Idx}/{Count}])", issue.Id, nextEscalationTime, prevEscalationTime, workflow.EscalateTimes[idx], idx + 1, workflow.EscalateTimes.Count);
+					break;
+				}
 			}
 
 			return (nextEscalationTime - DateTime.UnixEpoch).TotalSeconds;
