Horde: Allow setting escalation times property in workflow config. Preventing MongoDB from serializing data correctly.

diff --git a/Engine/Source/Programs/Horde/Horde.Build/Server/Config.cs b/Engine/Source/Programs/Horde/Horde.Build/Server/Config.cs
--- a/Engine/Source/Programs/Horde/Horde.Build/Server/Config.cs
+++ b/Engine/Source/Programs/Horde/Horde.Build/Server/Config.cs
@@ -392,7 +392,7 @@ namespace Horde.Build.Server
 		/// <summary>
 		/// Times after an issue has been opened to escalate to the alias above, in minutes. Continues to notify on the last interval once reaching the end of the list.
 		/// </summary>
-		public List<int> EscalateTimes { get; } = new List<int>{ 120 };
+		public List<int> EscalateTimes { get; set; } = new List<int>{ 120 };
 
 		/// <summary>
 		/// Maximum number of people to mention on a triage thread
