Horde Storage - Further improvements to error handling when sends / receives gets aborted / is running slowly.

diff --git a/Engine/Source/Programs/Horde/HordeStorage/Horde.Storage/Controllers/CompressedBlobController.cs b/Engine/Source/Programs/Horde/HordeStorage/Horde.Storage/Controllers/CompressedBlobController.cs
--- a/Engine/Source/Programs/Horde/HordeStorage/Horde.Storage/Controllers/CompressedBlobController.cs
+++ b/Engine/Source/Programs/Horde/HordeStorage/Horde.Storage/Controllers/CompressedBlobController.cs
@@ -8,6 +8,7 @@ using System.Linq;
 using System.Net;
 using System.Net.Mime;
 using System.Threading.Tasks;
+using EpicGames.AspNet;
 using EpicGames.Horde.Storage;
 using EpicGames.Serialization;
 using Horde.Storage.Implementation;
@@ -19,6 +20,7 @@ using Microsoft.AspNetCore.Authorization;
 using Microsoft.AspNetCore.Mvc;
 using Microsoft.Extensions.Primitives;
 using Serilog;
+using CustomMediaTypeNames = Jupiter.CustomMediaTypeNames;
 
 namespace Horde.Storage.Controllers
 {
diff --git a/Engine/Source/Programs/Horde/HordeStorage/Horde.Storage/Controllers/ObjectController.cs b/Engine/Source/Programs/Horde/HordeStorage/Horde.Storage/Controllers/ObjectController.cs
--- a/Engine/Source/Programs/Horde/HordeStorage/Horde.Storage/Controllers/ObjectController.cs
+++ b/Engine/Source/Programs/Horde/HordeStorage/Horde.Storage/Controllers/ObjectController.cs
@@ -9,6 +9,7 @@ using System.Net;
 using System.Net.Mime;
 using System.Threading.Tasks;
 using async_enumerable_dotnet;
+using EpicGames.AspNet;
 using EpicGames.Horde.Storage;
 using EpicGames.Serialization;
 using Horde.Storage.Implementation;
@@ -19,6 +20,7 @@ using Jupiter.Utils;
 using Microsoft.AspNetCore.Authorization;
 using Microsoft.AspNetCore.Mvc;
 using Serilog;
+using CustomMediaTypeNames = Jupiter.CustomMediaTypeNames;
 
 namespace Horde.Storage.Controllers
 {
diff --git a/Engine/Source/Programs/Horde/HordeStorage/Horde.Storage/Controllers/ReferencesController.cs b/Engine/Source/Programs/Horde/HordeStorage/Horde.Storage/Controllers/ReferencesController.cs
--- a/Engine/Source/Programs/Horde/HordeStorage/Horde.Storage/Controllers/ReferencesController.cs
+++ b/Engine/Source/Programs/Horde/HordeStorage/Horde.Storage/Controllers/ReferencesController.cs
@@ -148,7 +148,17 @@ namespace Horde.Storage.Controllers
                     Response.ContentLength = contentLength;
                     Response.ContentType = contentType;
                     Response.StatusCode = StatusCodes.Status200OK;
-                    await StreamCopyOperation.CopyToAsync(blobContents.Stream, outputStream, count: null, bufferSize: BufferSize, cancel: Response.HttpContext.RequestAborted);
+                    try
+                    {
+                        await StreamCopyOperation.CopyToAsync(blobContents.Stream, outputStream, count: null, bufferSize: BufferSize, cancel: Response.HttpContext.RequestAborted);
+                    }
+                    catch (OperationCanceledException)
+                    {
+                        // do not raise exceptions for cancelled writes
+                        // as we have already started writing a response we can not change the status code
+                        // so we just drop a warning and proceed
+                        _logger.Warning("The operation was canceled while writing the body");
+                    }
                 }
 
                 string responseType = _formatResolver.GetResponseType(Request, format, CustomMediaTypeNames.UnrealCompactBinary);
diff --git a/Engine/Source/Programs/Horde/HordeStorage/Horde.Storage/Controllers/StorageController.cs b/Engine/Source/Programs/Horde/HordeStorage/Horde.Storage/Controllers/StorageController.cs
--- a/Engine/Source/Programs/Horde/HordeStorage/Horde.Storage/Controllers/StorageController.cs
+++ b/Engine/Source/Programs/Horde/HordeStorage/Horde.Storage/Controllers/StorageController.cs
@@ -9,6 +9,7 @@ using System.Linq;
 using System.Net;
 using System.Net.Mime;
 using System.Threading.Tasks;
+using EpicGames.AspNet;
 using EpicGames.Horde.Storage;
 using EpicGames.Serialization;
 using Horde.Storage.Implementation;
diff --git a/Engine/Source/Programs/Horde/HordeStorage/Jupiter.Common/Utils/StreamUtils.cs b/Engine/Source/Programs/Horde/HordeStorage/Jupiter.Common/Utils/StreamUtils.cs
--- a/Engine/Source/Programs/Horde/HordeStorage/Jupiter.Common/Utils/StreamUtils.cs
+++ b/Engine/Source/Programs/Horde/HordeStorage/Jupiter.Common/Utils/StreamUtils.cs
@@ -1,8 +1,8 @@
 ﻿// Copyright Epic Games, Inc. All Rights Reserved.
 
-using System;
 using System.IO;
 using System.Threading.Tasks;
+using EpicGames.AspNet;
 using Microsoft.AspNetCore.Http;
 
 namespace Jupiter.Utils
@@ -19,30 +19,9 @@ namespace Jupiter.Utils
             }
             catch (BadHttpRequestException e)
             {
-                // if a user is sending data to slowly it will be terminated by asp.net and a bad http request exception is thrown
-                // this type of exception can be thrown for multiple reasons so we check for the error message about MinRequestBodyDataRate
-                // that setting can be set to reduce the number of bytes expected to be sent by a user - but that is likely a bad idea
-                // as the most common reason for this happening is a client trying to send more data then what their uplink has bandwidth for.
-                if (e.Message.Contains("MinRequestBodyDataRate", StringComparison.Ordinal))
-                {
-                    throw new ClientSendSlowException(e);
-                }
-
-                // this exception also indicates a send issue, in this case the client sent fewer bytes then what it set as a content length
-                // most likely caused by them shutting down their connection before finishing the entire send
-                if (e.Message.Contains("Unexpected end of request content", StringComparison.Ordinal))
-                {
-                    throw new ClientSendSlowException(e);
-                }
+                ClientSendSlowExceptionUtil.MaybeThrowSlowSendException(e);
                 throw;
             }
         }
     }
-
-    public class ClientSendSlowException : Exception
-    {
-        public ClientSendSlowException(Exception? innerException) : base("Client was sending data to slow", innerException)
-        {
-        }
-    }
 }
diff --git a/Engine/Source/Programs/Shared/EpicGames.AspNet/ClientSendSlowException.cs b/Engine/Source/Programs/Shared/EpicGames.AspNet/ClientSendSlowException.cs
new file mode 100644
--- /dev/null
+++ b/Engine/Source/Programs/Shared/EpicGames.AspNet/ClientSendSlowException.cs
@@ -0,0 +1,38 @@
+// Copyright Epic Games, Inc. All Rights Reserved.
+
+using System;
+using Microsoft.AspNetCore.Http;
+
+namespace EpicGames.AspNet
+{
+	public static class ClientSendSlowExceptionUtil
+	{
+		public static bool MaybeThrowSlowSendException(BadHttpRequestException e)
+		{
+			// if a user is sending data to slowly it will be terminated by asp.net and a bad http request exception is thrown
+			// this type of exception can be thrown for multiple reasons so we check for the error message about MinRequestBodyDataRate
+			// that setting can be set to reduce the number of bytes expected to be sent by a user - but that is likely a bad idea
+			// as the most common reason for this happening is a client trying to send more data then what their uplink has bandwidth for.
+			if (e.Message.Contains("MinRequestBodyDataRate", StringComparison.Ordinal))
+			{
+				throw new ClientSendSlowException(e);
+			}
+
+			// this exception also indicates a send issue, in this case the client sent fewer bytes then what it set as a content length
+			// most likely caused by them shutting down their connection before finishing the entire send
+			if (e.Message.Contains("Unexpected end of request content", StringComparison.Ordinal))
+			{
+				throw new ClientSendSlowException(e);
+			}
+
+			return true;
+		}
+	}
+
+	public class ClientSendSlowException : Exception
+	{
+		public ClientSendSlowException(Exception? innerException) : base("Client was sending data to slow", innerException)
+		{
+		}
+	}
+}
diff --git a/Engine/Source/Programs/Shared/EpicGames.AspNet/CompactBinaryFormatter.cs b/Engine/Source/Programs/Shared/EpicGames.AspNet/CompactBinaryFormatter.cs
--- a/Engine/Source/Programs/Shared/EpicGames.AspNet/CompactBinaryFormatter.cs
+++ b/Engine/Source/Programs/Shared/EpicGames.AspNet/CompactBinaryFormatter.cs
@@ -7,6 +7,7 @@ using System.Threading.Tasks;
 using System.IO;
 using EpicGames.Serialization;
 using EpicGames.Core;
+using Microsoft.AspNetCore.Http;
 using Microsoft.Extensions.DependencyInjection;
 using Microsoft.Extensions.Primitives;
 using Microsoft.Extensions.Logging;
@@ -48,11 +49,17 @@ namespace EpicGames.AspNet
         {
             // Buffer the data into an array
             byte[] data;
-            using (MemoryStream stream = new MemoryStream())
+            try
             {
+                using MemoryStream stream = new MemoryStream();
                 await context.HttpContext.Request.Body.CopyToAsync(stream);
                 data = stream.ToArray();
             }
+            catch (BadHttpRequestException e)
+            {
+                ClientSendSlowExceptionUtil.MaybeThrowSlowSendException(e);
+                throw;
+            }
 
             // Serialize the object
             CbField field;
