Horde Storage - Fixed build error by overloaded ContentId type - using the old type that we were using before.

diff --git a/Engine/Source/Programs/Horde/HordeStorage/Horde.Storage/Controllers/CompressedBlobController.cs b/Engine/Source/Programs/Horde/HordeStorage/Horde.Storage/Controllers/CompressedBlobController.cs
--- a/Engine/Source/Programs/Horde/HordeStorage/Horde.Storage/Controllers/CompressedBlobController.cs
+++ b/Engine/Source/Programs/Horde/HordeStorage/Horde.Storage/Controllers/CompressedBlobController.cs
@@ -20,6 +20,7 @@ using Microsoft.AspNetCore.Authorization;
 using Microsoft.AspNetCore.Mvc;
 using Microsoft.Extensions.Primitives;
 using Serilog;
+using ContentId = Jupiter.Implementation.ContentId;
 using CustomMediaTypeNames = Jupiter.CustomMediaTypeNames;
 
 namespace Horde.Storage.Controllers
diff --git a/Engine/Source/Programs/Horde/HordeStorage/Horde.Storage/Controllers/ContentIdController.cs b/Engine/Source/Programs/Horde/HordeStorage/Horde.Storage/Controllers/ContentIdController.cs
--- a/Engine/Source/Programs/Horde/HordeStorage/Horde.Storage/Controllers/ContentIdController.cs
+++ b/Engine/Source/Programs/Horde/HordeStorage/Horde.Storage/Controllers/ContentIdController.cs
@@ -8,6 +8,7 @@ using Jupiter;
 using Jupiter.Implementation;
 using Microsoft.AspNetCore.Authorization;
 using Microsoft.AspNetCore.Mvc;
+using ContentId = Jupiter.Implementation.ContentId;
 
 namespace Horde.Storage.Controllers
 {
diff --git a/Engine/Source/Programs/Horde/HordeStorage/Horde.Storage/Controllers/ReferencesController.cs b/Engine/Source/Programs/Horde/HordeStorage/Horde.Storage/Controllers/ReferencesController.cs
--- a/Engine/Source/Programs/Horde/HordeStorage/Horde.Storage/Controllers/ReferencesController.cs
+++ b/Engine/Source/Programs/Horde/HordeStorage/Horde.Storage/Controllers/ReferencesController.cs
@@ -32,6 +32,7 @@ using Microsoft.Extensions.DependencyInjection;
 using Newtonsoft.Json;
 using Serilog;
 using ContentHash = Jupiter.Implementation.ContentHash;
+using ContentId = Jupiter.Implementation.ContentId;
 using CustomMediaTypeNames = Jupiter.CustomMediaTypeNames;
 using Log = Serilog.Log;
 
diff --git a/Engine/Source/Programs/Horde/HordeStorage/Horde.Storage/Controllers/StorageController.cs b/Engine/Source/Programs/Horde/HordeStorage/Horde.Storage/Controllers/StorageController.cs
--- a/Engine/Source/Programs/Horde/HordeStorage/Horde.Storage/Controllers/StorageController.cs
+++ b/Engine/Source/Programs/Horde/HordeStorage/Horde.Storage/Controllers/StorageController.cs
@@ -16,7 +16,6 @@ using Horde.Storage.Implementation;
 using Jupiter;
 using Jupiter.Common.Implementation;
 using Jupiter.Implementation;
-using Jupiter.Utils;
 using Microsoft.AspNetCore.Authorization;
 using Microsoft.AspNetCore.Http;
 using Microsoft.AspNetCore.Mvc;
diff --git a/Engine/Source/Programs/Horde/HordeStorage/Horde.Storage/Implementation/Blob/IBlobService.cs b/Engine/Source/Programs/Horde/HordeStorage/Horde.Storage/Implementation/Blob/IBlobService.cs
--- a/Engine/Source/Programs/Horde/HordeStorage/Horde.Storage/Implementation/Blob/IBlobService.cs
+++ b/Engine/Source/Programs/Horde/HordeStorage/Horde.Storage/Implementation/Blob/IBlobService.cs
@@ -23,6 +23,7 @@ using Microsoft.AspNetCore.Http;
 using Microsoft.Extensions.DependencyInjection;
 using Microsoft.Extensions.Options;
 using Serilog;
+using ContentId = Jupiter.Implementation.ContentId;
 using CustomMediaTypeNames = Jupiter.CustomMediaTypeNames;
 
 namespace Horde.Storage.Implementation;
diff --git a/Engine/Source/Programs/Horde/HordeStorage/Horde.Storage/Implementation/ContentId/CacheContentIdStore.cs b/Engine/Source/Programs/Horde/HordeStorage/Horde.Storage/Implementation/ContentId/CacheContentIdStore.cs
--- a/Engine/Source/Programs/Horde/HordeStorage/Horde.Storage/Implementation/ContentId/CacheContentIdStore.cs
+++ b/Engine/Source/Programs/Horde/HordeStorage/Horde.Storage/Implementation/ContentId/CacheContentIdStore.cs
@@ -8,6 +8,7 @@ using EpicGames.Horde.Storage;
 using Horde.Storage.Controllers;
 using Jupiter.Implementation;
 using Microsoft.Extensions.Options;
+using ContentId = Jupiter.Implementation.ContentId;
 
 namespace Horde.Storage.Implementation
 {
diff --git a/Engine/Source/Programs/Horde/HordeStorage/Horde.Storage/Implementation/ContentId/IContentIdStore.cs b/Engine/Source/Programs/Horde/HordeStorage/Horde.Storage/Implementation/ContentId/IContentIdStore.cs
--- a/Engine/Source/Programs/Horde/HordeStorage/Horde.Storage/Implementation/ContentId/IContentIdStore.cs
+++ b/Engine/Source/Programs/Horde/HordeStorage/Horde.Storage/Implementation/ContentId/IContentIdStore.cs
@@ -4,6 +4,7 @@ using EpicGames.Horde.Storage;
 using Jupiter.Implementation;
 using System;
 using System.Threading.Tasks;
+using ContentId = Jupiter.Implementation.ContentId;
 
 namespace Horde.Storage.Implementation
 {
diff --git a/Engine/Source/Programs/Horde/HordeStorage/Horde.Storage/Implementation/ContentId/MemoryCachedContentIdStore.cs b/Engine/Source/Programs/Horde/HordeStorage/Horde.Storage/Implementation/ContentId/MemoryCachedContentIdStore.cs
--- a/Engine/Source/Programs/Horde/HordeStorage/Horde.Storage/Implementation/ContentId/MemoryCachedContentIdStore.cs
+++ b/Engine/Source/Programs/Horde/HordeStorage/Horde.Storage/Implementation/ContentId/MemoryCachedContentIdStore.cs
@@ -7,6 +7,7 @@ using EpicGames.Horde.Storage;
 using Jupiter.Implementation;
 using Microsoft.Extensions.Caching.Memory;
 using Microsoft.Extensions.Options;
+using ContentId = Jupiter.Implementation.ContentId;
 
 namespace Horde.Storage.Implementation
 {
diff --git a/Engine/Source/Programs/Horde/HordeStorage/Horde.Storage/Implementation/ContentId/MemoryContentIdStore.cs b/Engine/Source/Programs/Horde/HordeStorage/Horde.Storage/Implementation/ContentId/MemoryContentIdStore.cs
--- a/Engine/Source/Programs/Horde/HordeStorage/Horde.Storage/Implementation/ContentId/MemoryContentIdStore.cs
+++ b/Engine/Source/Programs/Horde/HordeStorage/Horde.Storage/Implementation/ContentId/MemoryContentIdStore.cs
@@ -5,6 +5,7 @@ using System.Collections.Generic;
 using System.Threading.Tasks;
 using EpicGames.Horde.Storage;
 using Jupiter.Implementation;
+using ContentId = Jupiter.Implementation.ContentId;
 
 namespace Horde.Storage.Implementation
 {
diff --git a/Engine/Source/Programs/Horde/HordeStorage/Horde.Storage/Implementation/ContentId/MongoContentIdStore.cs b/Engine/Source/Programs/Horde/HordeStorage/Horde.Storage/Implementation/ContentId/MongoContentIdStore.cs
--- a/Engine/Source/Programs/Horde/HordeStorage/Horde.Storage/Implementation/ContentId/MongoContentIdStore.cs
+++ b/Engine/Source/Programs/Horde/HordeStorage/Horde.Storage/Implementation/ContentId/MongoContentIdStore.cs
@@ -9,6 +9,7 @@ using Microsoft.Extensions.Options;
 using MongoDB.Bson.Serialization.Attributes;
 using MongoDB.Bson.Serialization.Options;
 using MongoDB.Driver;
+using ContentId = Jupiter.Implementation.ContentId;
 
 namespace Horde.Storage.Implementation
 {
diff --git a/Engine/Source/Programs/Horde/HordeStorage/Horde.Storage/Implementation/ContentId/ScyllaContentIdStore.cs b/Engine/Source/Programs/Horde/HordeStorage/Horde.Storage/Implementation/ContentId/ScyllaContentIdStore.cs
--- a/Engine/Source/Programs/Horde/HordeStorage/Horde.Storage/Implementation/ContentId/ScyllaContentIdStore.cs
+++ b/Engine/Source/Programs/Horde/HordeStorage/Horde.Storage/Implementation/ContentId/ScyllaContentIdStore.cs
@@ -8,6 +8,7 @@ using Cassandra.Mapping;
 using Datadog.Trace;
 using EpicGames.Horde.Storage;
 using Jupiter.Implementation;
+using ContentId = Jupiter.Implementation.ContentId;
 
 namespace Horde.Storage.Implementation
 {
diff --git a/Engine/Source/Programs/Horde/HordeStorage/Horde.Storage/Implementation/IReferenceResolver.cs b/Engine/Source/Programs/Horde/HordeStorage/Horde.Storage/Implementation/IReferenceResolver.cs
--- a/Engine/Source/Programs/Horde/HordeStorage/Horde.Storage/Implementation/IReferenceResolver.cs
+++ b/Engine/Source/Programs/Horde/HordeStorage/Horde.Storage/Implementation/IReferenceResolver.cs
@@ -10,6 +10,7 @@ using EpicGames.Horde.Storage;
 using EpicGames.Serialization;
 using Jupiter.Implementation;
 using Jupiter.Utils;
+using ContentId = Jupiter.Implementation.ContentId;
 
 namespace Horde.Storage.Implementation
 {
diff --git a/Engine/Source/Programs/Horde/HordeStorage/Horde.Storage/Implementation/ObjectService.cs b/Engine/Source/Programs/Horde/HordeStorage/Horde.Storage/Implementation/ObjectService.cs
--- a/Engine/Source/Programs/Horde/HordeStorage/Horde.Storage/Implementation/ObjectService.cs
+++ b/Engine/Source/Programs/Horde/HordeStorage/Horde.Storage/Implementation/ObjectService.cs
@@ -15,6 +15,7 @@ using Jupiter.Utils;
 using Microsoft.AspNetCore.Http;
 using Microsoft.Extensions.DependencyInjection;
 using Serilog;
+using ContentId = Jupiter.Implementation.ContentId;
 
 namespace Horde.Storage.Implementation
 {
diff --git a/Engine/Source/Programs/Horde/HordeStorage/Horde.Storage/Implementation/Objects/CacheReferencesStore.cs b/Engine/Source/Programs/Horde/HordeStorage/Horde.Storage/Implementation/Objects/CacheReferencesStore.cs
--- a/Engine/Source/Programs/Horde/HordeStorage/Horde.Storage/Implementation/Objects/CacheReferencesStore.cs
+++ b/Engine/Source/Programs/Horde/HordeStorage/Horde.Storage/Implementation/Objects/CacheReferencesStore.cs
@@ -15,6 +15,7 @@ using Microsoft.AspNetCore.Mvc;
 using Microsoft.Extensions.DependencyInjection;
 using Microsoft.Extensions.Options;
 using Jupiter;
+using ContentId = Jupiter.Implementation.ContentId;
 
 namespace Horde.Storage.Implementation
 {
diff --git a/Engine/Source/Programs/Horde/HordeStorage/Horde.Storage/Implementation/TransactionLog/ReplicationLogSnapshotBuilder.cs b/Engine/Source/Programs/Horde/HordeStorage/Horde.Storage/Implementation/TransactionLog/ReplicationLogSnapshotBuilder.cs
--- a/Engine/Source/Programs/Horde/HordeStorage/Horde.Storage/Implementation/TransactionLog/ReplicationLogSnapshotBuilder.cs
+++ b/Engine/Source/Programs/Horde/HordeStorage/Horde.Storage/Implementation/TransactionLog/ReplicationLogSnapshotBuilder.cs
@@ -10,6 +10,7 @@ using EpicGames.Horde.Storage;
 using EpicGames.Serialization;
 using Jupiter.Common.Implementation;
 using Jupiter.Implementation;
+using ContentId = Jupiter.Implementation.ContentId;
 
 namespace Horde.Storage.Implementation.TransactionLog
 {
diff --git a/Engine/Source/Programs/Horde/HordeStorage/Horde.Storage/Tests/Functional/Objects/ReferencesTests.cs b/Engine/Source/Programs/Horde/HordeStorage/Horde.Storage/Tests/Functional/Objects/ReferencesTests.cs
--- a/Engine/Source/Programs/Horde/HordeStorage/Horde.Storage/Tests/Functional/Objects/ReferencesTests.cs
+++ b/Engine/Source/Programs/Horde/HordeStorage/Horde.Storage/Tests/Functional/Objects/ReferencesTests.cs
@@ -29,6 +29,7 @@ using Serilog;
 using Logger = Serilog.Core.Logger;
 using EpicGames.Horde.Storage;
 using EpicGames.Serialization;
+using ContentId = Jupiter.Implementation.ContentId;
 
 namespace Horde.Storage.FunctionalTests.References
 {
diff --git a/Engine/Source/Programs/Horde/HordeStorage/Jupiter.Common/BaseStartup.cs b/Engine/Source/Programs/Horde/HordeStorage/Jupiter.Common/BaseStartup.cs
--- a/Engine/Source/Programs/Horde/HordeStorage/Jupiter.Common/BaseStartup.cs
+++ b/Engine/Source/Programs/Horde/HordeStorage/Jupiter.Common/BaseStartup.cs
@@ -5,11 +5,9 @@ using System.Collections.Generic;
 using System.ComponentModel.DataAnnotations;
 using System.IdentityModel.Tokens.Jwt;
 using System.IO;
-using System.Linq;
 using System.Net.Mime;
 using System.Reflection;
 using System.Text;
-using System.Threading;
 using System.Threading.Tasks;
 using Amazon;
 using Datadog.Trace;
