Horde Storage - Improved error handling when a object referenced by the replication log no longer exists.

diff --git a/Engine/Source/Programs/Horde/HordeStorage/Horde.Storage/Implementation/Replication/RefsReplicator.cs b/Engine/Source/Programs/Horde/HordeStorage/Horde.Storage/Implementation/Replication/RefsReplicator.cs
--- a/Engine/Source/Programs/Horde/HordeStorage/Horde.Storage/Implementation/Replication/RefsReplicator.cs
+++ b/Engine/Source/Programs/Horde/HordeStorage/Horde.Storage/Implementation/Replication/RefsReplicator.cs
@@ -510,6 +510,12 @@ namespace Horde.Storage.Implementation
                 _logger.Warning("Failed to resolve references for object {Blob} in {Namespace}. Skipping replication", objectToReplicate, ns);
                 return;
             }
+            if (referencesResponse.StatusCode == HttpStatusCode.NotFound)
+            {
+                // objects that do not exist can not be replicated so we skip them
+                _logger.Warning("Failed to resolve references for object {Blob} in {Namespace}. Got not found with message \"{Message}\". Skipping replication.", objectToReplicate, ns, referencesResponse.ReasonPhrase);
+                return;
+            }
             referencesResponse.EnsureSuccessStatusCode();
 
             ResolvedReferencesResult? refs = JsonConvert.DeserializeObject<ResolvedReferencesResult>(body);
