Horde: Separate config classes into separate files.

diff --git a/Engine/Source/Programs/Horde/Horde.Build/Issues/WorkflowConfig.cs b/Engine/Source/Programs/Horde/Horde.Build/Issues/WorkflowConfig.cs
new file mode 100644
--- /dev/null
+++ b/Engine/Source/Programs/Horde/Horde.Build/Issues/WorkflowConfig.cs
@@ -0,0 +1,122 @@
+// Copyright Epic Games, Inc. All Rights Reserved.
+
+using System;
+using System.Collections.Generic;
+using System.ComponentModel.DataAnnotations;
+using Horde.Build.Jobs.Graphs;
+using Horde.Build.Utilities;
+
+namespace Horde.Build.Issues
+{
+	using WorkflowId = StringId<WorkflowConfig>;
+
+	/// <summary>
+	/// Configuration for an issue workflow
+	/// </summary>
+	public class WorkflowConfig
+	{
+		/// <summary>
+		/// Identifier for this workflow
+		/// </summary>
+		public WorkflowId Id { get; set; } = WorkflowId.Empty;
+
+		/// <summary>
+		/// Times of day at which to send a report
+		/// </summary>
+		public List<TimeSpan> ReportTimes { get; set; } = new List<TimeSpan> { TimeSpan.Zero };
+
+		/// <summary>
+		/// Name of the tab to post summary data to
+		/// </summary>
+		public string? SummaryTab { get; set; }
+
+		/// <summary>
+		/// Channel to post summary information for these templates.
+		/// </summary>
+		public string? ReportChannel { get; set; }
+
+		/// <summary>
+		/// Whether to group issues by template in the report
+		/// </summary>
+		public bool GroupIssuesByTemplate { get; set; } = true;
+
+		/// <summary>
+		/// Channel to post threads for triaging new issues
+		/// </summary>
+		public string? TriageChannel { get; set; }
+
+		/// <summary>
+		/// Prefix for all triage messages
+		/// </summary>
+		public string? TriagePrefix { get; set; } = "*[NEW]* ";
+
+		/// <summary>
+		/// Suffix for all triage messages
+		/// </summary>
+		public string? TriageSuffix { get; set; }
+
+		/// <summary>
+		/// Instructions posted to triage threads
+		/// </summary>
+		public string? TriageInstructions { get; set; }
+
+		/// <summary>
+		/// User id of a Slack user/alias to ping if there is nobody assigned to an issue by default.
+		/// </summary>
+		public string? TriageAlias { get; set; }
+
+		/// <summary>
+		/// Alias to ping if an issue has not been resolved for a certain amount of time
+		/// </summary>
+		public string? EscalateAlias { get; set; }
+
+		/// <summary>
+		/// Times after an issue has been opened to escalate to the alias above, in minutes. Continues to notify on the last interval once reaching the end of the list.
+		/// </summary>
+		public List<int> EscalateTimes { get; set; } = new List<int> { 120 };
+
+		/// <summary>
+		/// Maximum number of people to mention on a triage thread
+		/// </summary>
+		public int MaxMentions { get; set; } = 5;
+
+		/// <summary>
+		/// Whether to mention people on this thread. Useful to disable for testing.
+		/// </summary>
+		public bool AllowMentions { get; set; } = true;
+
+		/// <summary>
+		/// Additional node annotations implicit in this workflow
+		/// </summary>
+		public NodeAnnotations Annotations { get; set; } = new NodeAnnotations();
+
+		/// <summary>
+		/// External issue tracking configuration for this workflow
+		/// </summary>
+		public ExternalIssueConfig? ExternalIssues { get; set; }
+	}
+
+	/// <summary>
+	/// External issue tracking configuration for a workflow
+	/// </summary>
+	public class ExternalIssueConfig
+	{
+		/// <summary>
+		/// Project key in external issue tracker
+		/// </summary>
+		[Required]
+		public string ProjectKey { get; set; } = String.Empty;
+
+		/// <summary>
+		/// Default component id for issues using workflow
+		/// </summary>
+		[Required]
+		public string DefaultComponentId { get; set; } = String.Empty;
+
+		/// <summary>
+		/// Default issue type id for issues using workflow
+		/// </summary>
+		[Required]
+		public string DefaultIssueTypeId { get; set; } = String.Empty;
+	}
+}
diff --git a/Engine/Source/Programs/Horde/Horde.Build/Jobs/Graphs/NodeAnnotations.cs b/Engine/Source/Programs/Horde/Horde.Build/Jobs/Graphs/NodeAnnotations.cs
--- a/Engine/Source/Programs/Horde/Horde.Build/Jobs/Graphs/NodeAnnotations.cs
+++ b/Engine/Source/Programs/Horde/Horde.Build/Jobs/Graphs/NodeAnnotations.cs
@@ -2,6 +2,7 @@
 
 using System;
 using System.Collections.Generic;
+using Horde.Build.Issues;
 using Horde.Build.Server;
 using Horde.Build.Utilities;
 
diff --git a/Engine/Source/Programs/Horde/Horde.Build/Projects/ProjectConfig.cs b/Engine/Source/Programs/Horde/Horde.Build/Projects/ProjectConfig.cs
new file mode 100644
--- /dev/null
+++ b/Engine/Source/Programs/Horde/Horde.Build/Projects/ProjectConfig.cs
@@ -0,0 +1,65 @@
+// Copyright Epic Games, Inc. All Rights Reserved.
+
+using System.Collections.Generic;
+using System.ComponentModel.DataAnnotations;
+using EpicGames.Core;
+using Horde.Build.Acls;
+using Horde.Build.Streams;
+using Horde.Build.Utilities;
+
+namespace Horde.Build.Projects
+{
+	using StreamId = StringId<IStream>;
+
+	/// <summary>
+	/// Stores configuration for a project
+	/// </summary>
+	[JsonSchema("https://unrealengine.com/horde/project")]
+	[JsonSchemaCatalog("Horde Project", "Horde project configuration file", new[] { "*.project.json", "Projects/*.json" })]
+	public class ProjectConfig
+	{
+		/// <summary>
+		/// Name for the new project
+		/// </summary>
+		[Required]
+		public string Name { get; set; } = null!;
+
+		/// <summary>
+		/// Path to the project logo
+		/// </summary>
+		public string? Logo { get; set; }
+
+		/// <summary>
+		/// Categories to include in this project
+		/// </summary>
+		public List<CreateProjectCategoryRequest> Categories { get; set; } = new List<CreateProjectCategoryRequest>();
+
+		/// <summary>
+		/// List of streams
+		/// </summary>
+		public List<StreamConfigRef> Streams { get; set; } = new List<StreamConfigRef>();
+
+		/// <summary>
+		/// Acl entries
+		/// </summary>
+		public UpdateAclRequest? Acl { get; set; }
+	}
+
+	/// <summary>
+	/// Reference to configuration for a stream
+	/// </summary>
+	public class StreamConfigRef
+	{
+		/// <summary>
+		/// Unique id for the stream
+		/// </summary>
+		[Required]
+		public StreamId Id { get; set; } = StreamId.Empty;
+
+		/// <summary>
+		/// Path to the configuration file
+		/// </summary>
+		[Required]
+		public string Path { get; set; } = null!;
+	}
+}
diff --git a/Engine/Source/Programs/Horde/Horde.Build/Server/Config.cs b/Engine/Source/Programs/Horde/Horde.Build/Server/Config.cs
deleted file mode 100644
diff --git a/Engine/Source/Programs/Horde/Horde.Build/Server/GlobalConfig.cs b/Engine/Source/Programs/Horde/Horde.Build/Server/GlobalConfig.cs
new file mode 100644
--- /dev/null
+++ b/Engine/Source/Programs/Horde/Horde.Build/Server/GlobalConfig.cs
@@ -0,0 +1,187 @@
+// Copyright Epic Games, Inc. All Rights Reserved.
+
+using System;
+using System.Collections.Generic;
+using System.ComponentModel.DataAnnotations;
+using System.Diagnostics.CodeAnalysis;
+using System.Linq;
+using EpicGames.Core;
+using EpicGames.Horde.Common;
+using EpicGames.Serialization;
+using Horde.Build.Acls;
+using Horde.Build.Jobs.Graphs;
+using Horde.Build.Projects;
+using Horde.Build.Streams;
+using Horde.Build.Tools;
+using Horde.Build.Utilities;
+
+namespace Horde.Build.Server
+{
+	using ProjectId = StringId<IProject>;
+	using StreamId = StringId<IStream>;
+	using TemplateRefId = StringId<TemplateRef>;
+	
+	/// <summary>
+	/// Global configuration
+	/// </summary>
+	[JsonSchema("https://unrealengine.com/horde/global")]
+	[JsonSchemaCatalog("Horde Globals", "Horde global configuration file", "globals.json")]
+	public class GlobalConfig
+	{
+		/// <summary>
+		/// List of projects
+		/// </summary>
+		public List<ProjectConfigRef> Projects { get; set; } = new List<ProjectConfigRef>();
+
+		/// <summary>
+		/// List of scheduled downtime
+		/// </summary>
+		public List<ScheduledDowntime> Downtime { get; set; } = new List<ScheduledDowntime>();
+
+		/// <summary>
+		/// List of Perforce clusters
+		/// </summary>
+		public List<PerforceCluster> PerforceClusters { get; set; } = new List<PerforceCluster>();
+
+		/// <summary>
+		/// List of costs of a particular agent type
+		/// </summary>
+		public List<AgentRateConfig> Rates { get; set; } = new List<AgentRateConfig>();
+
+		/// <summary>
+		/// List of compute profiles
+		/// </summary>
+		public List<ComputeClusterConfig> Compute { get; set; } = new List<ComputeClusterConfig>();
+
+		/// <summary>
+		/// List of tools hosted by the server
+		/// </summary>
+		public List<ToolOptions> Tools { get; set; } = new List<ToolOptions>();
+
+		/// <summary>
+		/// Maximum number of conforms to run at once
+		/// </summary>
+		public int MaxConformCount { get; set; }
+
+		/// <summary>
+		/// List of storage namespaces
+		/// </summary>
+		public StorageConfig? Storage { get; set; }
+
+		/// <summary>
+		/// Access control list
+		/// </summary>
+		public UpdateAclRequest? Acl { get; set; }
+	}
+
+	/// <summary>
+	/// Profile for executing compute requests
+	/// </summary>
+	public class ComputeClusterConfig
+	{
+		/// <summary>
+		/// Name of the partition
+		/// </summary>
+		public string Id { get; set; } = "default";
+
+		/// <summary>
+		/// Name of the namespace to use
+		/// </summary>
+		public string NamespaceId { get; set; } = "horde.compute";
+
+		/// <summary>
+		/// Name of the input bucket
+		/// </summary>
+		public string RequestBucketId { get; set; } = "requests";
+
+		/// <summary>
+		/// Name of the output bucket
+		/// </summary>
+		public string ResponseBucketId { get; set; } = "responses";
+
+		/// <summary>
+		/// Access control list
+		/// </summary>
+		public UpdateAclRequest? Acl { get; set; }
+	}
+
+	/// <summary>
+	/// References a project configuration
+	/// </summary>
+	public class ProjectConfigRef
+	{
+		/// <summary>
+		/// Unique id for the project
+		/// </summary>
+		[Required]
+		public ProjectId Id { get; set; } = ProjectId.Empty;
+
+		/// <summary>
+		/// Config path for the project
+		/// </summary>
+		[Required]
+		public string Path { get; set; } = null!;
+	}
+
+	/// <summary>
+	/// Configuration for storage system
+	/// </summary>
+	public class StorageConfig
+	{
+		/// <summary>
+		/// List of storage namespaces
+		/// </summary>
+		public List<NamespaceConfig> Namespaces { get; set; } = new List<NamespaceConfig>();
+	}
+
+	/// <summary>
+	/// Configuration for a storage namespace
+	/// </summary>
+	public class NamespaceConfig
+	{
+		/// <summary>
+		/// Identifier for this namespace
+		/// </summary>
+		public string Id { get; set; } = String.Empty;
+
+		/// <summary>
+		/// Buckets within this namespace
+		/// </summary>
+		public List<BucketConfig> Buckets { get; set; } = new List<BucketConfig>();
+
+		/// <summary>
+		/// Access control for this namespace
+		/// </summary>
+		public UpdateAclRequest? Acl { get; set; }
+	}
+
+	/// <summary>
+	/// Configuration for a bucket
+	/// </summary>
+	public class BucketConfig
+	{
+		/// <summary>
+		/// Identifier for the bucket
+		/// </summary>
+		public string Id { get; set; } = String.Empty;
+	}
+
+	/// <summary>
+	/// Describes the monetary cost of agents matching a particular criteris
+	/// </summary>
+	public class AgentRateConfig
+	{
+		/// <summary>
+		/// Condition string
+		/// </summary>
+		[CbField("c")]
+		public Condition? Condition { get; set; }
+
+		/// <summary>
+		/// Rate for this agent
+		/// </summary>
+		[CbField("r")]
+		public double Rate { get; set; }
+	}
+}
+
diff --git a/Engine/Source/Programs/Horde/Horde.Build/Streams/StreamConfig.cs b/Engine/Source/Programs/Horde/Horde.Build/Streams/StreamConfig.cs
new file mode 100644
--- /dev/null
+++ b/Engine/Source/Programs/Horde/Horde.Build/Streams/StreamConfig.cs
@@ -0,0 +1,145 @@
+// Copyright Epic Games, Inc. All Rights Reserved.
+
+using System;
+using System.Collections.Generic;
+using System.ComponentModel.DataAnnotations;
+using System.Diagnostics.CodeAnalysis;
+using System.Linq;
+using EpicGames.Core;
+using Horde.Build.Acls;
+using Horde.Build.Issues;
+using Horde.Build.Utilities;
+
+namespace Horde.Build.Streams
+{
+	using TemplateRefId = StringId<TemplateRef>;
+	using WorkflowId = StringId<WorkflowConfig>;
+
+	/// <summary>
+	/// Config for a stream
+	/// </summary>
+	[JsonSchema("https://unrealengine.com/horde/stream")]
+	[JsonSchemaCatalog("Horde Stream", "Horde stream configuration file", new[] { "*.stream.json", "Streams/*.json" })]
+	public class StreamConfig
+	{
+		/// <summary>
+		/// Name of the stream
+		/// </summary>
+		[Required]
+		public string Name { get; set; } = String.Empty;
+
+		/// <summary>
+		/// The perforce cluster containing the stream
+		/// </summary>
+		public string? ClusterName { get; set; }
+
+		/// <summary>
+		/// Order for this stream
+		/// </summary>
+		public int? Order { get; set; }
+
+		/// <summary>
+		/// Notification channel for all jobs in this stream
+		/// </summary>
+		public string? NotificationChannel { get; set; }
+
+		/// <summary>
+		/// Notification channel filter for this template. Can be Success|Failure|Warnings
+		/// </summary>
+		public string? NotificationChannelFilter { get; set; }
+
+		/// <summary>
+		/// Channel to post issue triage notifications
+		/// </summary>
+		public string? TriageChannel { get; set; }
+
+		/// <summary>
+		/// Default template for running preflights
+		/// </summary>
+		public string? DefaultPreflightTemplate { get; set; }
+
+		/// <summary>
+		/// Default template for running preflights
+		/// </summary>
+		public DefaultPreflightRequest? DefaultPreflight { get; set; }
+
+		/// <summary>
+		/// List of tabs to show for the new stream
+		/// </summary>
+		public List<CreateStreamTabRequest> Tabs { get; set; } = new List<CreateStreamTabRequest>();
+
+		/// <summary>
+		/// Map of agent name to type
+		/// </summary>
+		public Dictionary<string, CreateAgentTypeRequest> AgentTypes { get; set; } = new Dictionary<string, CreateAgentTypeRequest>();
+
+		/// <summary>
+		/// Map of workspace name to type
+		/// </summary>
+		public Dictionary<string, CreateWorkspaceTypeRequest> WorkspaceTypes { get; set; } = new Dictionary<string, CreateWorkspaceTypeRequest>();
+
+		/// <summary>
+		/// List of templates to create
+		/// </summary>
+		public List<TemplateRefConfig> Templates { get; set; } = new List<TemplateRefConfig>();
+
+		/// <summary>
+		/// Custom permissions for this object
+		/// </summary>
+		public UpdateAclRequest? Acl { get; set; }
+
+		/// <summary>
+		/// Pause stream builds until specified date
+		/// </summary>
+		public DateTime? PausedUntil { get; set; }
+
+		/// <summary>
+		/// Reason for pausing builds of the stream
+		/// </summary>
+		public string? PauseComment { get; set; }
+
+		/// <summary>
+		/// How to replicate data from VCS to Horde Storage.
+		/// </summary>
+		public ContentReplicationMode ReplicationMode { get; set; }
+
+		/// <summary>
+		/// Filter for paths to be replicated to storage, as a Perforce wildcard relative to the root of the workspace.
+		/// </summary>
+		public string? ReplicationFilter { get; set; }
+
+		/// <summary>
+		/// Stream to use for replication, if different to the default.
+		/// </summary>
+		public string? ReplicationStream { get; set; }
+
+		/// <summary>
+		/// Workflows for dealing with new issues
+		/// </summary>
+		public List<WorkflowConfig> Workflows { get; set; } = new List<WorkflowConfig>();
+
+		/// <summary>
+		/// Tries to find a template with the given id
+		/// </summary>
+		/// <param name="templateRefId"></param>
+		/// <param name="templateConfig"></param>
+		/// <returns></returns>
+		public bool TryGetTemplate(TemplateRefId templateRefId, [NotNullWhen(true)] out TemplateRefConfig? templateConfig)
+		{
+			templateConfig = Templates.FirstOrDefault(x => x.Id == templateRefId);
+			return templateConfig != null;
+		}
+
+		/// <summary>
+		/// Tries to find a workflow with the given id
+		/// </summary>
+		/// <param name="workflowId"></param>
+		/// <param name="workflowConfig"></param>
+		/// <returns></returns>
+		public bool TryGetWorkflow(WorkflowId workflowId, [NotNullWhen(true)] out WorkflowConfig? workflowConfig)
+		{
+			workflowConfig = Workflows.FirstOrDefault(x => x.Id == workflowId);
+			return workflowConfig != null;
+		}
+	}
+}
diff --git a/Engine/Source/Programs/Horde/Horde.Build/Streams/StreamMessages.cs b/Engine/Source/Programs/Horde/Horde.Build/Streams/StreamMessages.cs
--- a/Engine/Source/Programs/Horde/Horde.Build/Streams/StreamMessages.cs
+++ b/Engine/Source/Programs/Horde/Horde.Build/Streams/StreamMessages.cs
@@ -12,6 +12,7 @@ using Horde.Build.Jobs.Schedules;
 using Horde.Build.Users;
 using Horde.Build.Utilities;
 using HordeCommon;
+using Horde.Build.Issues;
 
 namespace Horde.Build.Streams
 {
