Horde: Upgrade all NuGet packages to latest.

diff --git a/Engine/Source/Programs/Horde/Horde.Agent.Tests/Horde.Agent.Tests.csproj b/Engine/Source/Programs/Horde/Horde.Agent.Tests/Horde.Agent.Tests.csproj
--- a/Engine/Source/Programs/Horde/Horde.Agent.Tests/Horde.Agent.Tests.csproj
+++ b/Engine/Source/Programs/Horde/Horde.Agent.Tests/Horde.Agent.Tests.csproj
@@ -10,12 +10,15 @@
 
   <ItemGroup>
 
-    <PackageReference Include="Microsoft.Extensions.Logging.Abstractions" Version="6.0.0" />
-    <PackageReference Include="Microsoft.NET.Test.Sdk" Version="16.5.0" />
-    <PackageReference Include="Moq" Version="4.14.5" />
-    <PackageReference Include="MSTest.TestAdapter" Version="2.1.0" />
-    <PackageReference Include="MSTest.TestFramework" Version="2.1.0" />
-    <PackageReference Include="coverlet.collector" Version="1.2.0" />
+    <PackageReference Include="Microsoft.Extensions.Logging.Abstractions" Version="6.0.1" />
+    <PackageReference Include="Microsoft.NET.Test.Sdk" Version="17.3.0" />
+    <PackageReference Include="Moq" Version="4.18.2" />
+    <PackageReference Include="MSTest.TestAdapter" Version="2.2.10" />
+    <PackageReference Include="MSTest.TestFramework" Version="2.2.10" />
+    <PackageReference Include="coverlet.collector" Version="3.1.2">
+      <PrivateAssets>all</PrivateAssets>
+      <IncludeAssets>runtime; build; native; contentfiles; analyzers; buildtransitive</IncludeAssets>
+    </PackageReference>
   </ItemGroup>
 
   <ItemGroup>
diff --git a/Engine/Source/Programs/Horde/Horde.Agent/Commands/Service/RunCommand.cs b/Engine/Source/Programs/Horde/Horde.Agent/Commands/Service/RunCommand.cs
--- a/Engine/Source/Programs/Horde/Horde.Agent/Commands/Service/RunCommand.cs
+++ b/Engine/Source/Programs/Horde/Horde.Agent/Commands/Service/RunCommand.cs
@@ -6,6 +6,7 @@ using System.Net.Http;
 using System.Threading.Tasks;
 using Datadog.Trace;
 using Datadog.Trace.Configuration;
+using Datadog.Trace.OpenTracing;
 using EpicGames.Core;
 using EpicGames.Horde.Storage;
 using Horde.Agent.Services;
@@ -20,6 +21,8 @@ using Polly;
 
 namespace Horde.Agent.Modes.Service
 {
+	using ITracer = OpenTracing.ITracer;
+
 	/// <summary>
 	/// 
 	/// </summary>
@@ -138,9 +141,9 @@ namespace Horde.Agent.Modes.Service
 			settings.ServiceVersion = version;
 			settings.LogsInjectionEnabled = true;
 
-			Tracer.Instance = new Tracer(settings);
+			Tracer.Configure(settings);
 
-			ITracer openTracer = Datadog.Trace.OpenTracing.OpenTracingTracerFactory.WrapTracer(Tracer.Instance);
+			ITracer openTracer = OpenTracingTracerFactory.WrapTracer(Tracer.Instance);
 			GlobalTracer.Register(openTracer);
 		}
 	}
diff --git a/Engine/Source/Programs/Horde/Horde.Agent/Commands/Service/UpgradeCommand.cs b/Engine/Source/Programs/Horde/Horde.Agent/Commands/Service/UpgradeCommand.cs
--- a/Engine/Source/Programs/Horde/Horde.Agent/Commands/Service/UpgradeCommand.cs
+++ b/Engine/Source/Programs/Horde/Horde.Agent/Commands/Service/UpgradeCommand.cs
@@ -8,6 +8,7 @@ using System.Linq;
 using System.Management;
 using System.Reflection;
 using System.Runtime.InteropServices;
+using System.Runtime.Versioning;
 using System.ServiceProcess;
 using System.Text;
 using System.Threading.Tasks;
@@ -143,6 +144,7 @@ namespace Horde.Agent.Commands
 			otherProcess.Kill();
 		}
 
+		[SupportedOSPlatform("windows")]
 		void UpgradeWindowsService(ILogger logger, Process otherProcess, HashSet<string> targetFiles, List<Tuple<string, string>> renameFiles)
 		{
 			// Try to get the service associated with the passed-in process id
@@ -196,6 +198,7 @@ namespace Horde.Agent.Commands
 		/// </summary>
 		/// <param name="processId">The process id to search for</param>
 		/// <returns>The service controller corresponding to this process</returns>
+		[SupportedOSPlatform("windows")]
 		static ServiceController? GetServiceForProcess(int processId)
 		{
 			try
diff --git a/Engine/Source/Programs/Horde/Horde.Agent/Horde.Agent.csproj b/Engine/Source/Programs/Horde/Horde.Agent/Horde.Agent.csproj
--- a/Engine/Source/Programs/Horde/Horde.Agent/Horde.Agent.csproj
+++ b/Engine/Source/Programs/Horde/Horde.Agent/Horde.Agent.csproj
@@ -22,32 +22,32 @@
   </ItemGroup>
 
   <ItemGroup>
-    <PackageReference Include="AWSSDK.Core" Version="3.7.3.19" />
-    <PackageReference Include="AWSSDK.EC2" Version="3.7.4" />
-    <PackageReference Include="AWSSDK.SimpleSystemsManagement" Version="3.7.3.19" />
-    <PackageReference Include="Datadog.Trace.OpenTracing" Version="1.18.0" />
-    <PackageReference Include="Google.Protobuf" Version="3.20.1" />
-    <PackageReference Include="Grpc.Net.Client" Version="2.45.0" />
-    <PackageReference Include="Grpc.Tools" Version="2.45.0">
+    <PackageReference Include="AWSSDK.Core" Version="3.7.12.25" />
+    <PackageReference Include="AWSSDK.EC2" Version="3.7.83.1" />
+    <PackageReference Include="AWSSDK.SimpleSystemsManagement" Version="3.7.17.42" />
+    <PackageReference Include="Datadog.Trace.OpenTracing" Version="2.14.0" />
+    <PackageReference Include="Google.Protobuf" Version="3.21.5" />
+    <PackageReference Include="Grpc.Net.Client" Version="2.48.0" />
+    <PackageReference Include="Grpc.Tools" Version="2.48.0">
       <PrivateAssets>all</PrivateAssets>
       <IncludeAssets>runtime; build; native; contentfiles; analyzers; buildtransitive</IncludeAssets>
     </PackageReference>
-    <PackageReference Include="Microsoft.Extensions.Caching.Memory" Version="6.0.0" />
-    <PackageReference Include="Microsoft.Extensions.Hosting" Version="3.1.5" />
-    <PackageReference Include="Microsoft.Extensions.Hosting.Abstractions" Version="3.1.5" />
-    <PackageReference Include="Microsoft.Extensions.Hosting.WindowsServices" Version="3.1.5" />
-    <PackageReference Include="Microsoft.Extensions.Http" Version="3.1.22" />
-    <PackageReference Include="Microsoft.Extensions.Http.Polly" Version="3.1.22" />
-    <PackageReference Include="Microsoft.Extensions.Logging.Abstractions" Version="6.0.0" />
-    <PackageReference Include="Microsoft.Extensions.Options.DataAnnotations" Version="3.1.5" />
+    <PackageReference Include="Microsoft.Extensions.Caching.Memory" Version="6.0.1" />
+    <PackageReference Include="Microsoft.Extensions.Hosting" Version="6.0.1" />
+    <PackageReference Include="Microsoft.Extensions.Hosting.Abstractions" Version="6.0.0" />
+    <PackageReference Include="Microsoft.Extensions.Hosting.WindowsServices" Version="6.0.0" />
+    <PackageReference Include="Microsoft.Extensions.Http" Version="6.0.0" />
+    <PackageReference Include="Microsoft.Extensions.Http.Polly" Version="6.0.8" />
+    <PackageReference Include="Microsoft.Extensions.Logging.Abstractions" Version="6.0.1" />
+    <PackageReference Include="Microsoft.Extensions.Options.DataAnnotations" Version="6.0.0" />
     <PackageReference Include="OpenTracing" Version="0.12.1" />
-    <PackageReference Include="Serilog" Version="2.9.0" />
-    <PackageReference Include="Serilog.Extensions.Hosting" Version="3.1.0" />
-    <PackageReference Include="Serilog.Extensions.Logging" Version="3.0.1" />
-    <PackageReference Include="Serilog.Sinks.Console" Version="3.1.1" />
-    <PackageReference Include="Serilog.Sinks.File" Version="4.1.0" />
-    <PackageReference Include="System.Management" Version="4.7.0" />
-    <PackageReference Include="System.ServiceProcess.ServiceController" Version="4.7.0" />
+    <PackageReference Include="Serilog" Version="2.11.0" />
+    <PackageReference Include="Serilog.Extensions.Hosting" Version="5.0.1" />
+    <PackageReference Include="Serilog.Extensions.Logging" Version="3.1.0" />
+    <PackageReference Include="Serilog.Sinks.Console" Version="4.0.1" />
+    <PackageReference Include="Serilog.Sinks.File" Version="5.0.0" />
+    <PackageReference Include="System.Management" Version="6.0.0" />
+    <PackageReference Include="System.ServiceProcess.ServiceController" Version="6.0.0" />
   </ItemGroup>
 
   <ItemGroup>
diff --git a/Engine/Source/Programs/Horde/Horde.Agent/Services/WorkerService.cs b/Engine/Source/Programs/Horde/Horde.Agent/Services/WorkerService.cs
--- a/Engine/Source/Programs/Horde/Horde.Agent/Services/WorkerService.cs
+++ b/Engine/Source/Programs/Horde/Horde.Agent/Services/WorkerService.cs
@@ -11,6 +11,7 @@ using System.Management;
 using System.Net;
 using System.Reflection;
 using System.Runtime.InteropServices;
+using System.Runtime.Versioning;
 using System.Text;
 using System.Text.RegularExpressions;
 using System.Threading;
@@ -1951,6 +1952,7 @@ namespace Horde.Agent.Services
 			return records;
 		}
 
+		[SupportedOSPlatform("windows")]
 		class WmiProperties
 		{
 			readonly Dictionary<string, object> _properties = new Dictionary<string, object>(StringComparer.Ordinal);
@@ -2010,6 +2012,7 @@ namespace Horde.Agent.Services
 			}
 		}
 
+		[SupportedOSPlatform("windows")]
 		static Dictionary<string, object> GetWmiProperties(ManagementObject row)
 		{
 			Dictionary<string, object> properties = new Dictionary<string, object>(StringComparer.Ordinal);
diff --git a/Engine/Source/Programs/Horde/Horde.Build.Tests/Horde.Build.Tests.csproj b/Engine/Source/Programs/Horde/Horde.Build.Tests/Horde.Build.Tests.csproj
--- a/Engine/Source/Programs/Horde/Horde.Build.Tests/Horde.Build.Tests.csproj
+++ b/Engine/Source/Programs/Horde/Horde.Build.Tests/Horde.Build.Tests.csproj
@@ -8,31 +8,34 @@
   </PropertyGroup>
 
   <ItemGroup>
-    <Compile Remove="Stubs\Collections\JobCollectionStub.cs"/>
-    <Compile Remove="Stubs\Collections\StreamCollectionStub.cs"/>
+    <Compile Remove="Stubs\Collections\JobCollectionStub.cs" />
+    <Compile Remove="Stubs\Collections\StreamCollectionStub.cs" />
   </ItemGroup>
 
   <ItemGroup>
-    <None Include="Stubs\Collections\JobCollectionStub.cs"/>
-    <None Include="Stubs\Collections\StreamCollectionStub.cs"/>
+    <None Include="Stubs\Collections\JobCollectionStub.cs" />
+    <None Include="Stubs\Collections\StreamCollectionStub.cs" />
   </ItemGroup>
 
   <ItemGroup>
-    <PackageReference Include="Microsoft.AspNetCore.Mvc.Testing" Version="3.1.6"/>
-    <PackageReference Include="Microsoft.NET.Test.Sdk" Version="16.5.0"/>
-    <PackageReference Include="Moq" Version="4.14.5"/>
-    <PackageReference Include="MSTest.TestAdapter" Version="2.1.0"/>
-    <PackageReference Include="MSTest.TestFramework" Version="2.1.0"/>
-    <PackageReference Include="Microsoft.Extensions.Caching.Memory" Version="6.0.0"/>
-    <PackageReference Include="coverlet.collector" Version="1.2.0"/>
+    <PackageReference Include="Microsoft.AspNetCore.Mvc.Testing" Version="6.0.8" />
+    <PackageReference Include="Microsoft.NET.Test.Sdk" Version="17.3.0" />
+    <PackageReference Include="Moq" Version="4.18.2" />
+    <PackageReference Include="MSTest.TestAdapter" Version="2.2.10" />
+    <PackageReference Include="MSTest.TestFramework" Version="2.2.10" />
+    <PackageReference Include="Microsoft.Extensions.Caching.Memory" Version="6.0.1" />
+    <PackageReference Include="coverlet.collector" Version="3.1.2">
+      <PrivateAssets>all</PrivateAssets>
+      <IncludeAssets>runtime; build; native; contentfiles; analyzers; buildtransitive</IncludeAssets>
+    </PackageReference>
   </ItemGroup>
 
   <ItemGroup>
-    <ProjectReference Include="..\..\Shared\EpicGames.Serialization\EpicGames.Serialization.csproj"/>
+    <ProjectReference Include="..\..\Shared\EpicGames.Serialization\EpicGames.Serialization.csproj" />
     <ProjectReference Include="..\Horde.Agent\Horde.Agent.csproj">
       <Aliases>HordeAgent</Aliases>
     </ProjectReference>
-    <ProjectReference Include="..\Horde.Build\Horde.Build.csproj"/>
+    <ProjectReference Include="..\Horde.Build\Horde.Build.csproj" />
   </ItemGroup>
 
   <ItemGroup>
@@ -51,7 +54,7 @@
   </ItemGroup>
 
   <ItemGroup>
-    <Folder Include="Resources\"/>
+    <Folder Include="Resources\" />
   </ItemGroup>
 
 </Project>
diff --git a/Engine/Source/Programs/Horde/Horde.Build/Agents/AgentId.cs b/Engine/Source/Programs/Horde/Horde.Build/Agents/AgentId.cs
--- a/Engine/Source/Programs/Horde/Horde.Build/Agents/AgentId.cs
+++ b/Engine/Source/Programs/Horde/Horde.Build/Agents/AgentId.cs
@@ -127,7 +127,7 @@ namespace Horde.Build.Agents
 	public sealed class AgentIdRedisConverter : IRedisConverter<AgentId>
 	{
 		/// <inheritdoc/>
-		public AgentId FromRedisValue(RedisValue value) => new AgentId((string)value);
+		public AgentId FromRedisValue(RedisValue value) => new AgentId((string)value!);
 
 		/// <inheritdoc/>
 		public RedisValue ToRedisValue(AgentId value) => value.ToString();
diff --git a/Engine/Source/Programs/Horde/Horde.Build/Horde.Build.csproj b/Engine/Source/Programs/Horde/Horde.Build/Horde.Build.csproj
--- a/Engine/Source/Programs/Horde/Horde.Build/Horde.Build.csproj
+++ b/Engine/Source/Programs/Horde/Horde.Build/Horde.Build.csproj
@@ -1,4 +1,4 @@
-<Project Sdk="Microsoft.NET.Sdk.Web">
+﻿<Project Sdk="Microsoft.NET.Sdk.Web">
 
   <PropertyGroup>
     <TargetFramework>net6.0</TargetFramework>
@@ -11,45 +11,45 @@
   </PropertyGroup>
 
   <ItemGroup>
-    <PackageReference Include="AWSSDK.AutoScaling" Version="3.7.10.43" />
-    <PackageReference Include="AWSSDK.CloudWatch" Version="3.7.6.3" />
-    <PackageReference Include="AWSSDK.EC2" Version="3.7.75.1" />
+    <PackageReference Include="AWSSDK.AutoScaling" Version="3.7.10.48" />
+    <PackageReference Include="AWSSDK.CloudWatch" Version="3.7.7.3" />
+    <PackageReference Include="AWSSDK.EC2" Version="3.7.83.1" />
     <PackageReference Include="AWSSDK.Extensions.NETCore.Setup" Version="3.7.2" />
-    <PackageReference Include="AWSSDK.S3" Version="3.7.9.19" />
-    <PackageReference Include="AWSSDK.SecurityToken" Version="3.7.1.168" />
-    <PackageReference Include="Datadog.Trace.OpenTracing" Version="2.11.0" />
-    <PackageReference Include="DogStatsD-CSharp-Client" Version="6.0.0" />
-    <PackageReference Include="Grpc.AspNetCore" Version="2.45.0" />
-    <PackageReference Include="Grpc.AspNetCore.Server.Reflection" Version="2.45.0" />
-    <PackageReference Include="Grpc.Net.Client" Version="2.45.0" />
-    <PackageReference Include="Grpc.Tools" Version="2.45.0">
+    <PackageReference Include="AWSSDK.S3" Version="3.7.9.43" />
+    <PackageReference Include="AWSSDK.SecurityToken" Version="3.7.1.189" />
+    <PackageReference Include="Datadog.Trace.OpenTracing" Version="2.14.0" />
+    <PackageReference Include="DogStatsD-CSharp-Client" Version="7.0.1" />
+    <PackageReference Include="Grpc.AspNetCore" Version="2.48.0" />
+    <PackageReference Include="Grpc.AspNetCore.Server.Reflection" Version="2.48.0" />
+    <PackageReference Include="Grpc.Net.Client" Version="2.48.0" />
+    <PackageReference Include="Grpc.Tools" Version="2.48.0">
       <PrivateAssets>all</PrivateAssets>
       <IncludeAssets>runtime; build; native; contentfiles; analyzers; buildtransitive</IncludeAssets>
     </PackageReference>
-    <PackageReference Include="JetBrains.Profiler.SelfApi" Version="2.1.2" />
+    <PackageReference Include="JetBrains.Profiler.SelfApi" Version="2.3.0" />
     <PackageReference Include="Microsoft.AspNetCore.Authentication.Cookies" Version="2.2.0" />
-    <PackageReference Include="Microsoft.AspNetCore.Authentication.JwtBearer" Version="6.0.0" />
-    <PackageReference Include="Microsoft.AspNetCore.Authentication.OpenIdConnect" Version="3.1.23" />
-    <PackageReference Include="Microsoft.AspNetCore.DataProtection.StackExchangeRedis" Version="3.1.7" />
-    <PackageReference Include="Microsoft.AspNetCore.SpaServices.Extensions" Version="3.1.17" />
-    <PackageReference Include="Microsoft.Extensions.Caching.Memory" Version="6.0.0" />
-    <PackageReference Include="Microsoft.Extensions.Http" Version="3.1.22" />
-    <PackageReference Include="MongoDB.Driver" Version="2.11.6" />
+    <PackageReference Include="Microsoft.AspNetCore.Authentication.JwtBearer" Version="6.0.8" />
+    <PackageReference Include="Microsoft.AspNetCore.Authentication.OpenIdConnect" Version="6.0.8" />
+    <PackageReference Include="Microsoft.AspNetCore.DataProtection.StackExchangeRedis" Version="6.0.8" />
+    <PackageReference Include="Microsoft.AspNetCore.SpaServices.Extensions" Version="6.0.8" />
+    <PackageReference Include="Microsoft.Extensions.Caching.Memory" Version="6.0.1" />
+    <PackageReference Include="Microsoft.Extensions.Http" Version="6.0.0" />
+    <PackageReference Include="MongoDB.Driver" Version="2.17.1" />
     <PackageReference Include="OpenTracing" Version="0.12.1" />
     <PackageReference Include="OpenTracing.Contrib.Grpc" Version="0.4.0" />
-    <PackageReference Include="Polly" Version="7.2.1" />
+    <PackageReference Include="Polly" Version="7.2.3" />
     <PackageReference Include="Polly.Extensions.Http" Version="3.0.0" />
-    <PackageReference Include="Serilog.AspNetCore" Version="3.4.0" />
+    <PackageReference Include="Serilog.AspNetCore" Version="6.0.1" />
     <PackageReference Include="Serilog.Enrichers.OpenTracing" Version="0.0.2" />
-    <PackageReference Include="Serilog.Extensions.Hosting" Version="3.1.0" />
-    <PackageReference Include="Serilog.Sinks.Console" Version="3.1.1" />
-    <PackageReference Include="Serilog.Sinks.File" Version="4.1.0" />
+    <PackageReference Include="Serilog.Extensions.Hosting" Version="5.0.1" />
+    <PackageReference Include="Serilog.Sinks.Console" Version="4.0.1" />
+    <PackageReference Include="Serilog.Sinks.File" Version="5.0.0" />
     <PackageReference Include="SerilogAnalyzer" Version="0.15.0" />
-    <PackageReference Include="SharpZipLib" Version="1.2.0" />
-    <PackageReference Include="StackExchange.Redis" Version="2.2.4" />
-    <PackageReference Include="Swashbuckle.AspNetCore" Version="5.6.3" />
-    <PackageReference Include="System.Linq.Async" Version="5.0.0" />
-    <PackageReference Include="TimeZoneConverter" Version="3.2.0" />
+    <PackageReference Include="SharpZipLib" Version="1.3.3" />
+    <PackageReference Include="StackExchange.Redis" Version="2.6.48" />
+    <PackageReference Include="Swashbuckle.AspNetCore" Version="6.4.0" />
+    <PackageReference Include="System.Linq.Async" Version="6.0.1" />
+    <PackageReference Include="TimeZoneConverter" Version="5.0.0" />
   </ItemGroup>
 
   <ItemGroup>
diff --git a/Engine/Source/Programs/Horde/Horde.Build/Notifications/NotificationService.cs b/Engine/Source/Programs/Horde/Horde.Build/Notifications/NotificationService.cs
--- a/Engine/Source/Programs/Horde/Horde.Build/Notifications/NotificationService.cs
+++ b/Engine/Source/Programs/Horde/Horde.Build/Notifications/NotificationService.cs
@@ -328,11 +328,11 @@ namespace Horde.Build.Notifications
 			await redis.KeyDeleteAsync(redisKey);
 
 			List<T> notifications = new List<T>();
-			foreach (byte[] data in rawNotifications)
+			foreach (byte[]? data in rawNotifications)
 			{
 				try
 				{
-					T? notification = JsonSerializer.Deserialize<T>(data);
+					T? notification = JsonSerializer.Deserialize<T>(data!);
 					if (notification == null)
 					{
 						throw new Exception("Unable to deserialize");
diff --git a/Engine/Source/Programs/Horde/Horde.Build/Tools/ToolCollection.cs b/Engine/Source/Programs/Horde/Horde.Build/Tools/ToolCollection.cs
--- a/Engine/Source/Programs/Horde/Horde.Build/Tools/ToolCollection.cs
+++ b/Engine/Source/Programs/Horde/Horde.Build/Tools/ToolCollection.cs
@@ -194,7 +194,7 @@ namespace Horde.Build.Tools
 				RedisValue value = await redis.StringGetAsync(s_indexKey);
 				if (!value.IsNullOrEmpty)
 				{
-					index = CbSerializer.Deserialize<CachedIndex>((byte[])value);
+					index = CbSerializer.Deserialize<CachedIndex>((byte[])value!);
 					if (index.Empty || index.Ids.Count > 0)
 					{
 						break;
diff --git a/Engine/Source/Programs/Horde/Horde.Build/Utilities/VersionedCollection.cs b/Engine/Source/Programs/Horde/Horde.Build/Utilities/VersionedCollection.cs
--- a/Engine/Source/Programs/Horde/Horde.Build/Utilities/VersionedCollection.cs
+++ b/Engine/Source/Programs/Horde/Horde.Build/Utilities/VersionedCollection.cs
@@ -219,7 +219,7 @@ namespace Horde.Build.Utilities
 			{
 				try
 				{
-					return BsonSerializer.Deserialize<VersionedDocument<TId, TLatest>>((byte[])cacheValue);
+					return BsonSerializer.Deserialize<VersionedDocument<TId, TLatest>>((byte[])cacheValue!);
 				}
 				catch
 				{
diff --git a/Engine/Source/Programs/Shared/EpicGames.AspNet/EpicGames.AspNet.csproj b/Engine/Source/Programs/Shared/EpicGames.AspNet/EpicGames.AspNet.csproj
--- a/Engine/Source/Programs/Shared/EpicGames.AspNet/EpicGames.AspNet.csproj
+++ b/Engine/Source/Programs/Shared/EpicGames.AspNet/EpicGames.AspNet.csproj
@@ -33,7 +33,7 @@
   </ItemGroup>
   <ItemGroup>
     <PackageReference Include="Microsoft.CSharp" Version="4.7.0" />
-    <PackageReference Include="Microsoft.Win32.Registry" Version="4.7.0" />
+    <PackageReference Include="Microsoft.Win32.Registry" Version="5.0.0" />
     <PackageReference Include="System.Data.DataSetExtensions" Version="4.5.0" />
 
     
diff --git a/Engine/Source/Programs/Shared/EpicGames.BuildGraph.Tests/EpicGames.BuildGraph.Tests.csproj b/Engine/Source/Programs/Shared/EpicGames.BuildGraph.Tests/EpicGames.BuildGraph.Tests.csproj
--- a/Engine/Source/Programs/Shared/EpicGames.BuildGraph.Tests/EpicGames.BuildGraph.Tests.csproj
+++ b/Engine/Source/Programs/Shared/EpicGames.BuildGraph.Tests/EpicGames.BuildGraph.Tests.csproj
@@ -21,10 +21,13 @@
   </ItemGroup>
 
   <ItemGroup>
-    <PackageReference Include="Microsoft.NET.Test.Sdk" Version="16.5.0" />
-    <PackageReference Include="MSTest.TestAdapter" Version="2.1.0" />
-    <PackageReference Include="MSTest.TestFramework" Version="2.1.0" />
-    <PackageReference Include="coverlet.collector" Version="1.2.0" />
+    <PackageReference Include="Microsoft.NET.Test.Sdk" Version="17.3.0" />
+    <PackageReference Include="MSTest.TestAdapter" Version="2.2.10" />
+    <PackageReference Include="MSTest.TestFramework" Version="2.2.10" />
+    <PackageReference Include="coverlet.collector" Version="3.1.2">
+      <PrivateAssets>all</PrivateAssets>
+      <IncludeAssets>runtime; build; native; contentfiles; analyzers; buildtransitive</IncludeAssets>
+    </PackageReference>
   </ItemGroup>
 
   <ItemGroup>
diff --git a/Engine/Source/Programs/Shared/EpicGames.BuildGraph/BgInterpreter.cs b/Engine/Source/Programs/Shared/EpicGames.BuildGraph/BgInterpreter.cs
--- a/Engine/Source/Programs/Shared/EpicGames.BuildGraph/BgInterpreter.cs
+++ b/Engine/Source/Programs/Shared/EpicGames.BuildGraph/BgInterpreter.cs
@@ -298,16 +298,16 @@ namespace EpicGames.BuildGraph
 						IEnumerable<object> rhs = (IEnumerable<object>)Evaluate(frame);
 						return String.Join(lhs, rhs);
 					}
-				case BgOpcode.StrIsMatch:
+				case BgOpcode.StrMatch:
 					{
 						string input = (string)Evaluate(frame);
-						string pattern = ReadString(frame);
+						string pattern = (string)Evaluate(frame);
 						return Regex.IsMatch(input, pattern);
 					}
 				case BgOpcode.StrReplace:
 					{
 						string input = (string)Evaluate(frame);
-						string pattern = ReadString(frame);
+						string pattern = (string)Evaluate(frame);
 						string replacement = (string)Evaluate(frame);
 						return Regex.Replace(input, pattern, replacement);
 					}
@@ -779,13 +779,13 @@ namespace EpicGames.BuildGraph
 
 						break;
 					}
-				case BgOpcode.StrIsMatch:
+				case BgOpcode.StrMatch:
+					Disassemble(frame, logger);
 					Disassemble(frame, logger);
-					Trace(frame, "pattern", ReadString, logger);
 					break;
 				case BgOpcode.StrReplace:
 					Disassemble(frame, logger);
-					Trace(frame, "pattern", ReadString, logger);
+					Disassemble(frame, logger);
 					Disassemble(frame, logger);
 					break;
 				case BgOpcode.StrOption:
diff --git a/Engine/Source/Programs/Shared/EpicGames.BuildGraph/BgOpcode.cs b/Engine/Source/Programs/Shared/EpicGames.BuildGraph/BgOpcode.cs
--- a/Engine/Source/Programs/Shared/EpicGames.BuildGraph/BgOpcode.cs
+++ b/Engine/Source/Programs/Shared/EpicGames.BuildGraph/BgOpcode.cs
@@ -159,7 +159,7 @@ namespace EpicGames.BuildGraph
 		/// <summary>
 		/// Tests whether a string (first argument) matches a regular expression (second argument)
 		/// </summary>
-		StrIsMatch = 0x27,
+		StrMatch = 0x27,
 
 		/// <summary>
 		/// Returns a string with all ocurrences of the second argument in the first argument replaced with the third argument
diff --git a/Engine/Source/Programs/Shared/EpicGames.BuildGraph/Expressions/BgString.cs b/Engine/Source/Programs/Shared/EpicGames.BuildGraph/Expressions/BgString.cs
--- a/Engine/Source/Programs/Shared/EpicGames.BuildGraph/Expressions/BgString.cs
+++ b/Engine/Source/Programs/Shared/EpicGames.BuildGraph/Expressions/BgString.cs
@@ -199,7 +199,7 @@ namespace EpicGames.BuildGraph.Expressions
 
 		public override void Write(BgBytecodeWriter writer)
 		{
-			writer.WriteOpcode(BgOpcode.StrIsMatch);
+			writer.WriteOpcode(BgOpcode.StrMatch);
 			writer.WriteExpr(Input);
 			writer.WriteExpr(Pattern);
 		}
diff --git a/Engine/Source/Programs/Shared/EpicGames.Core.Tests/EpicGames.Core.Tests.csproj b/Engine/Source/Programs/Shared/EpicGames.Core.Tests/EpicGames.Core.Tests.csproj
--- a/Engine/Source/Programs/Shared/EpicGames.Core.Tests/EpicGames.Core.Tests.csproj
+++ b/Engine/Source/Programs/Shared/EpicGames.Core.Tests/EpicGames.Core.Tests.csproj
@@ -21,10 +21,13 @@
   </ItemGroup>
 
   <ItemGroup>
-    <PackageReference Include="Microsoft.NET.Test.Sdk" Version="16.5.0" />
-    <PackageReference Include="MSTest.TestAdapter" Version="2.1.0" />
-    <PackageReference Include="MSTest.TestFramework" Version="2.1.0" />
-    <PackageReference Include="coverlet.collector" Version="1.2.0" />
+    <PackageReference Include="Microsoft.NET.Test.Sdk" Version="17.3.0" />
+    <PackageReference Include="MSTest.TestAdapter" Version="2.2.10" />
+    <PackageReference Include="MSTest.TestFramework" Version="2.2.10" />
+    <PackageReference Include="coverlet.collector" Version="3.1.2">
+      <PrivateAssets>all</PrivateAssets>
+      <IncludeAssets>runtime; build; native; contentfiles; analyzers; buildtransitive</IncludeAssets>
+    </PackageReference>
   </ItemGroup>
 
   <ItemGroup>
diff --git a/Engine/Source/Programs/Shared/EpicGames.Core/EpicGames.Core.csproj b/Engine/Source/Programs/Shared/EpicGames.Core/EpicGames.Core.csproj
--- a/Engine/Source/Programs/Shared/EpicGames.Core/EpicGames.Core.csproj
+++ b/Engine/Source/Programs/Shared/EpicGames.Core/EpicGames.Core.csproj
@@ -48,11 +48,11 @@
     </When>
   </Choose>
   <ItemGroup>
-    <PackageReference Include="JetBrains.Annotations" Version="2021.2.0" />
-    <PackageReference Include="Microsoft.Extensions.Logging" Version="5.0.0" />
-    <PackageReference Include="Microsoft.Extensions.Logging.Abstractions" Version="5.0.0" />
-    <PackageReference Include="Microsoft.Extensions.ObjectPool" Version="6.0.1" />
+    <PackageReference Include="JetBrains.Annotations" Version="2022.1.0" />
+    <PackageReference Include="Microsoft.Extensions.Logging" Version="6.0.0" />
+    <PackageReference Include="Microsoft.Extensions.Logging.Abstractions" Version="6.0.1" />
+    <PackageReference Include="Microsoft.Extensions.ObjectPool" Version="6.0.8" />
     <PackageReference Include="OpenTracing" Version="0.12.1" />
-    <PackageReference Include="System.Memory" Version="4.5.4" />
+    <PackageReference Include="System.Memory" Version="4.5.5" />
   </ItemGroup>
 </Project>
\ No newline at end of file
diff --git a/Engine/Source/Programs/Shared/EpicGames.Core/ForwardingLogger.cs b/Engine/Source/Programs/Shared/EpicGames.Core/ForwardingLogger.cs
--- a/Engine/Source/Programs/Shared/EpicGames.Core/ForwardingLogger.cs
+++ b/Engine/Source/Programs/Shared/EpicGames.Core/ForwardingLogger.cs
@@ -44,35 +44,18 @@ namespace EpicGames.Core
 		}
 
 		/// <inheritdoc/>
-		public IDisposable? BeginScope<TState>(TState state)
+		public IDisposable BeginScope<TState>(TState state)
 		{
-			IDisposable? disposable = null;
-
-			DisposableList? scope = null;
+			DisposableList scope = new DisposableList();
 			foreach (ILogger logger in Loggers)
 			{
 				IDisposable? next = logger.BeginScope(state);
 				if(next != null)
 				{
-					if (scope != null)
-					{
-						scope.Add(next);
-					}
-					else if (disposable != null)
-					{
-						scope = new DisposableList();
-						scope.Add(disposable);
-						scope.Add(next);
-						disposable = scope;
-					}
-					else
-					{
-						disposable = next;
-					}
+					scope.Add(next);
 				}
 			}
-
-			return disposable;
+			return scope;
 		}
 
 		/// <inheritdoc/>
@@ -89,7 +72,7 @@ namespace EpicGames.Core
 		}
 
 		/// <inheritdoc/>
-		public void Log<TState>(LogLevel logLevel, EventId eventId, TState state, Exception exception, Func<TState, Exception, string> formatter)
+		public void Log<TState>(LogLevel logLevel, EventId eventId, TState state, Exception? exception, Func<TState, Exception?, string> formatter)
 		{
 			foreach (ILogger logger in Loggers)
 			{
diff --git a/Engine/Source/Programs/Shared/EpicGames.Core/ILoggerIndent.cs b/Engine/Source/Programs/Shared/EpicGames.Core/ILoggerIndent.cs
--- a/Engine/Source/Programs/Shared/EpicGames.Core/ILoggerIndent.cs
+++ b/Engine/Source/Programs/Shared/EpicGames.Core/ILoggerIndent.cs
@@ -84,7 +84,7 @@ namespace EpicGames.Core
 			/// <summary>
 			/// Formatter for the inner state
 			/// </summary>
-			readonly Func<TState, Exception, string> _formatter;
+			readonly Func<TState, Exception?, string> _formatter;
 
 			/// <summary>
 			/// Constructor
@@ -92,7 +92,7 @@ namespace EpicGames.Core
 			/// <param name="indent">The indent to apply</param>
 			/// <param name="state">The inner state</param>
 			/// <param name="formatter">Formatter for the inner state</param>
-			public FormattedLogValues(string indent, TState state, Func<TState, Exception, string> formatter)
+			public FormattedLogValues(string indent, TState state, Func<TState, Exception?, string> formatter)
 			{
 				_indent = indent;
 				_state = state;
@@ -131,7 +131,7 @@ namespace EpicGames.Core
 			/// <param name="values">The object instance</param>
 			/// <param name="exception">The exception to format</param>
 			/// <returns>The formatted string</returns>
-			public static string Format(FormattedLogValues<TState> values, Exception exception)
+			public static string Format(FormattedLogValues<TState> values, Exception? exception)
 			{
 				return values._indent + values._formatter(values._state, exception);
 			}
@@ -185,7 +185,7 @@ namespace EpicGames.Core
 		}
 
 		/// <inheritdoc/>
-		public void Log<TState>(LogLevel logLevel, EventId eventId, TState state, Exception exception, Func<TState, Exception, string> formatter)
+		public void Log<TState>(LogLevel logLevel, EventId eventId, TState state, Exception? exception, Func<TState, Exception?, string> formatter)
 		{
 			if (_scopes.Count > 0)
 			{
diff --git a/Engine/Source/Programs/Shared/EpicGames.Core/Log.cs b/Engine/Source/Programs/Shared/EpicGames.Core/Log.cs
--- a/Engine/Source/Programs/Shared/EpicGames.Core/Log.cs
+++ b/Engine/Source/Programs/Shared/EpicGames.Core/Log.cs
@@ -840,7 +840,7 @@ namespace EpicGames.Core
 
 		public bool IsEnabled(LogLevel logLevel) => _inner.IsEnabled(logLevel);
 
-		public void Log<TState>(LogLevel logLevel, EventId eventId, TState state, Exception exception, Func<TState, Exception, string> formatter)
+		public void Log<TState>(LogLevel logLevel, EventId eventId, TState state, Exception? exception, Func<TState, Exception?, string> formatter)
 		{
 			lock (_parser)
 			{
@@ -1155,7 +1155,7 @@ namespace EpicGames.Core
 			return logLevel >= OutputLevel;
 		}
 
-		public void Log<TState>(LogLevel logLevel, EventId eventId, TState state, Exception exception, Func<TState, Exception?, string> formatter)
+		public void Log<TState>(LogLevel logLevel, EventId eventId, TState state, Exception? exception, Func<TState, Exception?, string> formatter)
 		{
 			string[] lines = formatter(state, exception).Split('\n');
 			lock (_syncObject)
diff --git a/Engine/Source/Programs/Shared/EpicGames.Horde.Tests/BundleStoreTests.cs b/Engine/Source/Programs/Shared/EpicGames.Horde.Tests/BundleStoreTests.cs
--- a/Engine/Source/Programs/Shared/EpicGames.Horde.Tests/BundleStoreTests.cs
+++ b/Engine/Source/Programs/Shared/EpicGames.Horde.Tests/BundleStoreTests.cs
@@ -28,7 +28,7 @@ namespace EpicGames.Horde.Tests
 			InMemoryBlobStore blobStore = new InMemoryBlobStore();
 			await TestTreeAsync(blobStore, new BundleOptions { MaxBlobSize = 1024 * 1024 });
 
-			Assert.AreEqual(0, blobStore.Blobs.Count);
+			Assert.AreEqual(1, blobStore.Blobs.Count);
 			Assert.AreEqual(1, blobStore.Refs.Count);
 		}
 
@@ -38,7 +38,7 @@ namespace EpicGames.Horde.Tests
 			InMemoryBlobStore blobStore = new InMemoryBlobStore();
 			await TestTreeAsync(blobStore, new BundleOptions { MaxBlobSize = 1 });
 
-			Assert.AreEqual(4, blobStore.Blobs.Count);
+			Assert.AreEqual(5, blobStore.Blobs.Count);
 			Assert.AreEqual(1, blobStore.Refs.Count);
 		}
 
diff --git a/Engine/Source/Programs/Shared/EpicGames.Horde.Tests/BundleTests.cs b/Engine/Source/Programs/Shared/EpicGames.Horde.Tests/BundleTests.cs
--- a/Engine/Source/Programs/Shared/EpicGames.Horde.Tests/BundleTests.cs
+++ b/Engine/Source/Programs/Shared/EpicGames.Horde.Tests/BundleTests.cs
@@ -117,7 +117,7 @@ namespace EpicGames.Horde.Tests
 				packets.Add(new BundlePacket(20, 40));
 				packets.Add(new BundlePacket(10, 20));
 
-				oldHeader = new BundleHeader(imports, exports, packets);
+				oldHeader = new BundleHeader(BundleCompressionFormat.LZ4, imports, exports, packets);
 			}
 
 			ByteArrayBuilder writer = new ByteArrayBuilder();
@@ -182,7 +182,7 @@ namespace EpicGames.Horde.Tests
 
 			// Should be stored inline
 			Assert.AreEqual(1, blobStore.Refs.Count);
-			Assert.AreEqual(0, blobStore.Blobs.Count);
+			Assert.AreEqual(1, blobStore.Blobs.Count);
 
 			// Check the ref
 			IBlob blob = await blobStore.ReadRefAsync(refName);
@@ -229,7 +229,7 @@ namespace EpicGames.Horde.Tests
 			await treeStore.WriteTreeAsync(refName, root);
 
 			Assert.AreEqual(1, blobStore.Refs.Count);
-			Assert.AreEqual(1, blobStore.Blobs.Count);
+			Assert.AreEqual(2, blobStore.Blobs.Count);
 		}
 
 		[TestMethod]
@@ -255,7 +255,7 @@ namespace EpicGames.Horde.Tests
 				await oldBundle.WriteTreeAsync(refName, root);
 
 				Assert.AreEqual(1, blobStore.Refs.Count);
-				Assert.AreEqual(4, blobStore.Blobs.Count);
+				Assert.AreEqual(5, blobStore.Blobs.Count);
 			}
 
 			using (ITreeStore newBundle = new BundleStore(blobStore, options))
diff --git a/Engine/Source/Programs/Shared/EpicGames.Horde.Tests/EpicGames.Horde.Tests.csproj b/Engine/Source/Programs/Shared/EpicGames.Horde.Tests/EpicGames.Horde.Tests.csproj
--- a/Engine/Source/Programs/Shared/EpicGames.Horde.Tests/EpicGames.Horde.Tests.csproj
+++ b/Engine/Source/Programs/Shared/EpicGames.Horde.Tests/EpicGames.Horde.Tests.csproj
@@ -8,11 +8,14 @@
   </PropertyGroup>
 
   <ItemGroup>
-    <PackageReference Include="Microsoft.NET.Test.Sdk" Version="16.5.0" />
-    <PackageReference Include="Microsoft.Extensions.Caching.Memory" Version="6.0.0" />
-    <PackageReference Include="MSTest.TestAdapter" Version="2.1.0" />
-    <PackageReference Include="MSTest.TestFramework" Version="2.1.0" />
-    <PackageReference Include="coverlet.collector" Version="1.2.0" />
+    <PackageReference Include="Microsoft.NET.Test.Sdk" Version="17.3.0" />
+    <PackageReference Include="Microsoft.Extensions.Caching.Memory" Version="6.0.1" />
+    <PackageReference Include="MSTest.TestAdapter" Version="2.2.10" />
+    <PackageReference Include="MSTest.TestFramework" Version="2.2.10" />
+    <PackageReference Include="coverlet.collector" Version="3.1.2">
+      <PrivateAssets>all</PrivateAssets>
+      <IncludeAssets>runtime; build; native; contentfiles; analyzers; buildtransitive</IncludeAssets>
+    </PackageReference>
   </ItemGroup>
 
   <ItemGroup>
diff --git a/Engine/Source/Programs/Shared/EpicGames.Horde/EpicGames.Horde.csproj b/Engine/Source/Programs/Shared/EpicGames.Horde/EpicGames.Horde.csproj
--- a/Engine/Source/Programs/Shared/EpicGames.Horde/EpicGames.Horde.csproj
+++ b/Engine/Source/Programs/Shared/EpicGames.Horde/EpicGames.Horde.csproj
@@ -36,13 +36,13 @@
     </Compile>
   </ItemGroup>
   <ItemGroup>
-    <PackageReference Include="K4os.Compression.LZ4" Version="1.2.6" />
-    <PackageReference Include="Microsoft.CodeAnalysis.CSharp.Workspaces" Version="3.11.0" />
+    <PackageReference Include="K4os.Compression.LZ4" Version="1.2.16" />
+    <PackageReference Include="Microsoft.CodeAnalysis.CSharp.Workspaces" Version="4.2.0" />
     <PackageReference Include="Microsoft.Extensions.Caching.Abstractions" Version="6.0.0" />
-    <PackageReference Include="Microsoft.Extensions.Caching.Memory" Version="6.0.0" />
+    <PackageReference Include="Microsoft.Extensions.Caching.Memory" Version="6.0.1" />
     <PackageReference Include="Microsoft.Extensions.DependencyInjection" Version="6.0.0" />
-    <PackageReference Include="Microsoft.Extensions.Http" Version="3.1.22" />
-    <PackageReference Include="Microsoft.Extensions.Logging" Version="5.0.0" />
+    <PackageReference Include="Microsoft.Extensions.Http" Version="6.0.0" />
+    <PackageReference Include="Microsoft.Extensions.Logging" Version="6.0.0" />
   </ItemGroup>
   <ItemGroup>
     <ProjectReference Include="..\EpicGames.Core\EpicGames.Core.csproj" />
diff --git a/Engine/Source/Programs/Shared/EpicGames.IoHash/EpicGames.IoHash.csproj b/Engine/Source/Programs/Shared/EpicGames.IoHash/EpicGames.IoHash.csproj
--- a/Engine/Source/Programs/Shared/EpicGames.IoHash/EpicGames.IoHash.csproj
+++ b/Engine/Source/Programs/Shared/EpicGames.IoHash/EpicGames.IoHash.csproj
@@ -30,7 +30,7 @@
   </ItemGroup>
   <ItemGroup>
     <PackageReference Include="Blake3" Version="0.5.1" />
-    <PackageReference Include="System.Memory" Version="4.5.4" />
+    <PackageReference Include="System.Memory" Version="4.5.5" />
   </ItemGroup>
   <ItemGroup>
     <ProjectReference Include="..\EpicGames.Core\EpicGames.Core.csproj" />
diff --git a/Engine/Source/Programs/Shared/EpicGames.Jupiter/EpicGames.Jupiter.csproj b/Engine/Source/Programs/Shared/EpicGames.Jupiter/EpicGames.Jupiter.csproj
--- a/Engine/Source/Programs/Shared/EpicGames.Jupiter/EpicGames.Jupiter.csproj
+++ b/Engine/Source/Programs/Shared/EpicGames.Jupiter/EpicGames.Jupiter.csproj
@@ -27,7 +27,7 @@
     </Compile>
   </ItemGroup>
   <ItemGroup>
-    <PackageReference Include="System.Memory" Version="4.5.4" />
+    <PackageReference Include="System.Memory" Version="4.5.5" />
   </ItemGroup>
   <ItemGroup>
     <ProjectReference Include="..\EpicGames.Core\EpicGames.Core.csproj" />
diff --git a/Engine/Source/Programs/Shared/EpicGames.Perforce/EpicGames.Perforce.csproj b/Engine/Source/Programs/Shared/EpicGames.Perforce/EpicGames.Perforce.csproj
--- a/Engine/Source/Programs/Shared/EpicGames.Perforce/EpicGames.Perforce.csproj
+++ b/Engine/Source/Programs/Shared/EpicGames.Perforce/EpicGames.Perforce.csproj
@@ -60,7 +60,7 @@
 
   <ItemGroup>
     <PackageReference Include="Microsoft.Win32.Registry" Version="5.0.0" />
-    <PackageReference Include="System.Linq.Async" Version="5.0.0" />
+    <PackageReference Include="System.Linq.Async" Version="6.0.1" />
   </ItemGroup>
 
 </Project>
diff --git a/Engine/Source/Programs/Shared/EpicGames.Redis/EpicGames.Redis.csproj b/Engine/Source/Programs/Shared/EpicGames.Redis/EpicGames.Redis.csproj
--- a/Engine/Source/Programs/Shared/EpicGames.Redis/EpicGames.Redis.csproj
+++ b/Engine/Source/Programs/Shared/EpicGames.Redis/EpicGames.Redis.csproj
@@ -28,8 +28,8 @@
     </Compile>
   </ItemGroup>
   <ItemGroup>
-    <PackageReference Include="Microsoft.Extensions.Logging" Version="5.0.0" />
-    <PackageReference Include="StackExchange.Redis" Version="2.2.4" />
+    <PackageReference Include="Microsoft.Extensions.Logging" Version="6.0.0" />
+    <PackageReference Include="StackExchange.Redis" Version="2.6.48" />
   </ItemGroup>
   <ItemGroup>
     <ProjectReference Include="..\EpicGames.Core\EpicGames.Core.csproj" />
diff --git a/Engine/Source/Programs/Shared/EpicGames.Redis/RedisExtensions.cs b/Engine/Source/Programs/Shared/EpicGames.Redis/RedisExtensions.cs
--- a/Engine/Source/Programs/Shared/EpicGames.Redis/RedisExtensions.cs
+++ b/Engine/Source/Programs/Shared/EpicGames.Redis/RedisExtensions.cs
@@ -16,7 +16,7 @@ namespace EpicGames.Redis
 		/// <returns></returns>
 		public static RedisKey AsKey(this RedisValue value)
 		{
-			return (byte[])value;
+			return (byte[])value!;
 		}
 
 		/// <summary>
@@ -26,7 +26,7 @@ namespace EpicGames.Redis
 		/// <returns></returns>
 		public static RedisValue AsValue(this RedisKey key)
 		{
-			return (byte[])key;
+			return (byte[])key!;
 		}
 	}
 }
diff --git a/Engine/Source/Programs/Shared/EpicGames.Redis/RedisSerializer.cs b/Engine/Source/Programs/Shared/EpicGames.Redis/RedisSerializer.cs
--- a/Engine/Source/Programs/Shared/EpicGames.Redis/RedisSerializer.cs
+++ b/Engine/Source/Programs/Shared/EpicGames.Redis/RedisSerializer.cs
@@ -65,7 +65,7 @@ namespace EpicGames.Redis
 		/// <inheritdoc/>
 		public T? FromRedisValue(RedisValue value)
 		{
-			return CbSerializer.Deserialize<T>(new CbField((byte[])value));
+			return CbSerializer.Deserialize<T>(new CbField((byte[])value!));
 		}
 	}
 
@@ -84,7 +84,7 @@ namespace EpicGames.Redis
 			}
 
 			public RedisValue ToRedisValue(T value) => (string?)_typeConverter.ConvertTo(value, typeof(string));
-			public T? FromRedisValue(RedisValue value) => (T?)_typeConverter.ConvertFrom((string)value);
+			public T? FromRedisValue(RedisValue value) => (T?)_typeConverter.ConvertFrom((string)value!);
 		}
 
 		class RedisNativeConverter<T> : IRedisConverter<T>
@@ -120,8 +120,8 @@ namespace EpicGames.Redis
 				CreateNativeConverter(x => (double)x, x => x),
 				CreateNativeConverter(x => (double?)x, x => x),
 				CreateNativeConverter(x => (ReadOnlyMemory<byte>)x, x => x),
-				CreateNativeConverter(x => (byte[])x, x => x),
-				CreateNativeConverter(x => (string)x, x => x),
+				CreateNativeConverter(x => (byte[])x!, x => x),
+				CreateNativeConverter(x => (string)x!, x => x),
 			};
 			return new Dictionary<Type, object>(converters);
 		}
diff --git a/Engine/Source/Programs/Shared/EpicGames.Serialization.Tests/EpicGames.Serialization.Tests.csproj b/Engine/Source/Programs/Shared/EpicGames.Serialization.Tests/EpicGames.Serialization.Tests.csproj
--- a/Engine/Source/Programs/Shared/EpicGames.Serialization.Tests/EpicGames.Serialization.Tests.csproj
+++ b/Engine/Source/Programs/Shared/EpicGames.Serialization.Tests/EpicGames.Serialization.Tests.csproj
@@ -21,10 +21,13 @@
   </ItemGroup>
 
   <ItemGroup>
-    <PackageReference Include="Microsoft.NET.Test.Sdk" Version="16.5.0" />
-    <PackageReference Include="MSTest.TestAdapter" Version="2.1.0" />
-    <PackageReference Include="MSTest.TestFramework" Version="2.1.0" />
-    <PackageReference Include="coverlet.collector" Version="1.2.0" />
+    <PackageReference Include="Microsoft.NET.Test.Sdk" Version="17.3.0" />
+    <PackageReference Include="MSTest.TestAdapter" Version="2.2.10" />
+    <PackageReference Include="MSTest.TestFramework" Version="2.2.10" />
+    <PackageReference Include="coverlet.collector" Version="3.1.2">
+      <PrivateAssets>all</PrivateAssets>
+      <IncludeAssets>runtime; build; native; contentfiles; analyzers; buildtransitive</IncludeAssets>
+    </PackageReference>
   </ItemGroup>
 
   <ItemGroup>
diff --git a/Engine/Source/Programs/Shared/EpicGames.Slack/EpicGames.Slack.csproj b/Engine/Source/Programs/Shared/EpicGames.Slack/EpicGames.Slack.csproj
--- a/Engine/Source/Programs/Shared/EpicGames.Slack/EpicGames.Slack.csproj
+++ b/Engine/Source/Programs/Shared/EpicGames.Slack/EpicGames.Slack.csproj
@@ -1,4 +1,4 @@
-<Project Sdk="Microsoft.NET.Sdk">
+﻿<Project Sdk="Microsoft.NET.Sdk">
   <PropertyGroup>
     <TargetFramework>net6.0</TargetFramework>
     <Configuration Condition=" '$(Configuration)' == '' ">Development</Configuration>
@@ -37,6 +37,6 @@
     <Folder Include="Blocks\" />
   </ItemGroup>
   <ItemGroup>
-    <PackageReference Include="Microsoft.Extensions.Logging.Abstractions" Version="6.0.0" />
+    <PackageReference Include="Microsoft.Extensions.Logging.Abstractions" Version="6.0.1" />
   </ItemGroup>
 </Project>
\ No newline at end of file
