Horde: Support for allocating access tokens for external services for each job step (designed for use with Horde Storage).

diff --git a/Engine/Source/Programs/Horde/Horde.Build/Server/RpcService.cs b/Engine/Source/Programs/Horde/Horde.Build/Server/RpcService.cs
--- a/Engine/Source/Programs/Horde/Horde.Build/Server/RpcService.cs
+++ b/Engine/Source/Programs/Horde/Horde.Build/Server/RpcService.cs
@@ -3,10 +3,14 @@
 using System;
 using System.Collections.Generic;
 using System.Globalization;
+using System.IO;
 using System.Linq;
 using System.Net;
+using System.Net.Http;
 using System.Security.Claims;
 using System.Text;
+using System.Text.Json;
+using System.Text.Json.Serialization;
 using System.Threading;
 using System.Threading.Tasks;
 using EpicGames.Core;
@@ -73,13 +77,14 @@ namespace Horde.Build.Server
 		readonly IJobStepRefCollection _jobStepRefCollection;
 		readonly ITemplateCollection _templateCollection;
 		readonly ConformTaskSource _conformTaskSource;
+		readonly HttpClient _httpClient;
 		readonly IClock _clock;
 		readonly ILogger<RpcService> _logger;
 
 		/// <summary>
 		/// Constructor
 		/// </summary>
-		public RpcService(AclService aclService, AgentService agentService, StreamService streamService, JobService jobService, AgentSoftwareService agentSoftwareService, IArtifactCollection artifactCollection, ILogFileService logFileService, PoolService poolService, LifetimeService lifetimeService, IGraphCollection graphs, ITestDataCollection testData, IJobStepRefCollection jobStepRefCollection, ITemplateCollection templateCollection, ConformTaskSource conformTaskSource, IClock clock, ILogger<RpcService> logger)
+		public RpcService(AclService aclService, AgentService agentService, StreamService streamService, JobService jobService, AgentSoftwareService agentSoftwareService, IArtifactCollection artifactCollection, ILogFileService logFileService, PoolService poolService, LifetimeService lifetimeService, IGraphCollection graphs, ITestDataCollection testData, IJobStepRefCollection jobStepRefCollection, ITemplateCollection templateCollection, ConformTaskSource conformTaskSource, HttpClient httpClient, IClock clock, ILogger<RpcService> logger)
 		{
 			_aclService = aclService;
 			_agentService = agentService;
@@ -95,6 +100,7 @@ namespace Horde.Build.Server
 			_jobStepRefCollection = jobStepRefCollection;
 			_templateCollection = templateCollection;
 			_conformTaskSource = conformTaskSource;
+			_httpClient = httpClient;
 			_clock = clock;
 			_logger = logger;
 		}
@@ -646,6 +652,23 @@ namespace Horde.Build.Server
 					}
 				}
 
+				IStream? stream = await _streamService.GetStreamAsync(job.StreamId);
+				if (stream != null)
+				{
+					foreach (TokenConfig tokenConfig in stream.Config.Tokens)
+					{
+						string? value = await AllocateTokenAsync(tokenConfig, context.CancellationToken);
+						if (value == null)
+						{
+							_logger.LogWarning("Unable to allocate token for job {JobId}:{BatchId}:{StepId} from {Url}", job.Id, batch.Id, step.Id, tokenConfig.Url);
+						}
+						else
+						{
+							response.EnvVars.Add(tokenConfig.EnvVar, value);
+						}
+					}
+				}
+
 				if (node.Properties != null)
 				{
 					response.Properties.Add(node.Properties);
@@ -657,6 +680,35 @@ namespace Horde.Build.Server
 			return null;
 		}
 
+		class GetTokenResponse
+		{
+			[JsonPropertyName("token_type")]
+			public string TokenType { get; set; } = String.Empty;
+
+			[JsonPropertyName("access_token")]
+			public string AccessToken { get; set; } = String.Empty;
+		}
+
+		async Task<string?> AllocateTokenAsync(TokenConfig config, CancellationToken cancellationToken)
+		{
+			List<KeyValuePair<string, string>> content = new List<KeyValuePair<string, string>>();
+			content.Add(KeyValuePair.Create("grant_type", "client_credentials"));
+			content.Add(KeyValuePair.Create("scope", "cache_access"));
+			content.Add(KeyValuePair.Create("client_id", config.ClientId));
+			content.Add(KeyValuePair.Create("client_secret", config.ClientSecret));
+
+			using (HttpRequestMessage request = new HttpRequestMessage(HttpMethod.Post, config.Url))
+			{
+				request.Content = new FormUrlEncodedContent(content);
+				using (HttpResponseMessage response = await _httpClient.SendAsync(request, cancellationToken))
+				{
+					using Stream stream = await response.Content.ReadAsStreamAsync(cancellationToken);
+					GetTokenResponse? token = await JsonSerializer.DeserializeAsync<GetTokenResponse>(stream, cancellationToken: cancellationToken);
+					return token?.AccessToken;
+				}
+			}
+		}
+
 		async Task<IJob> GetJobAsync(JobId jobId)
 		{
 			IJob? job = await _jobService.GetJobAsync(jobId);
diff --git a/Engine/Source/Programs/Horde/Horde.Build/Streams/StreamConfig.cs b/Engine/Source/Programs/Horde/Horde.Build/Streams/StreamConfig.cs
--- a/Engine/Source/Programs/Horde/Horde.Build/Streams/StreamConfig.cs
+++ b/Engine/Source/Programs/Horde/Horde.Build/Streams/StreamConfig.cs
@@ -118,6 +118,11 @@ namespace Horde.Build.Streams
 		/// </summary>
 		public List<WorkflowConfig> Workflows { get; set; } = new List<WorkflowConfig>();
 
+		/// <summary>
+		/// Tokens to create for each job step
+		/// </summary>
+		public List<TokenConfig> Tokens { get; set; } = new List<TokenConfig>(); 
+
 		/// <summary>
 		/// Tries to find a template with the given id
 		/// </summary>
@@ -142,4 +147,34 @@ namespace Horde.Build.Streams
 			return workflowConfig != null;
 		}
 	}
+
+	/// <summary>
+	/// Configuration for allocating access tokens for each job
+	/// </summary>
+	public class TokenConfig
+	{
+		/// <summary>
+		/// URL to request tokens from
+		/// </summary>
+		[Required]
+		public Uri Url { get; set; } = null!;
+
+		/// <summary>
+		/// Client id to use to request a new token
+		/// </summary>
+		[Required]
+		public string ClientId { get; set; } = String.Empty;
+
+		/// <summary>
+		/// Client secret to request a new access token
+		/// </summary>
+		[Required]
+		public string ClientSecret { get; set; } = String.Empty;
+
+		/// <summary>
+		/// Environment variable to set with the access token
+		/// </summary>
+		[Required]
+		public string EnvVar { get; set; } = String.Empty;
+	}
 }
