Horde: Add an HTTP client for RpcService.

diff --git a/Engine/Source/Programs/Horde/Horde.Build/Startup.cs b/Engine/Source/Programs/Horde/Horde.Build/Startup.cs
--- a/Engine/Source/Programs/Horde/Horde.Build/Startup.cs
+++ b/Engine/Source/Programs/Horde/Horde.Build/Startup.cs
@@ -348,6 +348,8 @@ namespace Horde.Build
 			});
 			services.AddGrpcReflection();
 
+			services.AddHttpClient<RpcService>();
+
 			services.AddSingleton<IAgentCollection, AgentCollection>();
 			services.AddSingleton<IAgentSoftwareCollection, AgentSoftwareCollection>();
 			services.AddSingleton<IArtifactCollection, ArtifactCollection>();
