Horde: Fix test setup for RPC tests.

diff --git a/Engine/Source/Programs/Horde/Horde.Build.Tests/TestSetup.cs b/Engine/Source/Programs/Horde/Horde.Build.Tests/TestSetup.cs
--- a/Engine/Source/Programs/Horde/Horde.Build.Tests/TestSetup.cs
+++ b/Engine/Source/Programs/Horde/Horde.Build.Tests/TestSetup.cs
@@ -157,6 +157,8 @@ namespace Horde.Build.Tests
 			services.Configure<ServerSettings>(ConfigureSettings);
 			services.AddSingleton<IConfiguration>(config);
 
+			services.AddHttpClient<RpcService>();
+
 			services.AddLogging(builder => builder.AddConsole());
 			services.AddSingleton<IMemoryCache>(sp => new MemoryCache(new MemoryCacheOptions { }));
 
