Horde: Put DDC token in the credentials list rather than the envVars list.

diff --git a/Engine/Source/Programs/Horde/Horde.Build/Server/RpcService.cs b/Engine/Source/Programs/Horde/Horde.Build/Server/RpcService.cs
--- a/Engine/Source/Programs/Horde/Horde.Build/Server/RpcService.cs
+++ b/Engine/Source/Programs/Horde/Horde.Build/Server/RpcService.cs
@@ -664,7 +664,7 @@ namespace Horde.Build.Server
 						}
 						else
 						{
-							response.EnvVars.Add(tokenConfig.EnvVar, value);
+							credentials.Add(tokenConfig.EnvVar, value);
 						}
 					}
 				}
