Horde: Fix incorrect variable modified for credentials.

diff --git a/Engine/Source/Programs/Horde/Horde.Build/Server/RpcService.cs b/Engine/Source/Programs/Horde/Horde.Build/Server/RpcService.cs
--- a/Engine/Source/Programs/Horde/Horde.Build/Server/RpcService.cs
+++ b/Engine/Source/Programs/Horde/Horde.Build/Server/RpcService.cs
@@ -664,7 +664,7 @@ namespace Horde.Build.Server
 						}
 						else
 						{
-							credentials.Add(tokenConfig.EnvVar, value);
+							response.Credentials.Add(tokenConfig.EnvVar, value);
 						}
 					}
 				}
