Horde Storage - Added middleware to suppress certain exceptions that can be thrown by code not under our own stack - this allows us to make sure they are not logged by the unhandled exception logging (so we can log it as a warning rather then a error).

diff --git a/Engine/Source/Programs/Horde/HordeStorage/Jupiter.Common/BaseStartup.cs b/Engine/Source/Programs/Horde/HordeStorage/Jupiter.Common/BaseStartup.cs
--- a/Engine/Source/Programs/Horde/HordeStorage/Jupiter.Common/BaseStartup.cs
+++ b/Engine/Source/Programs/Horde/HordeStorage/Jupiter.Common/BaseStartup.cs
@@ -59,6 +59,7 @@ namespace Jupiter
             CbConvertersAspNet.AddAspnetConverters();
 
             services.AddServerTiming();
+            services.AddSuppressExceptionsMiddleware();
 
             // aws specific settings
             services.AddOptions<AWSCredentialsSettings>().Bind(Configuration.GetSection("AWSCredentials")).ValidateDataAnnotations();
diff --git a/Engine/Source/Programs/Horde/HordeStorage/Jupiter.Common/SuppressExceptionMiddleware.cs b/Engine/Source/Programs/Horde/HordeStorage/Jupiter.Common/SuppressExceptionMiddleware.cs
new file mode 100644
--- /dev/null
+++ b/Engine/Source/Programs/Horde/HordeStorage/Jupiter.Common/SuppressExceptionMiddleware.cs
@@ -0,0 +1,62 @@
+// Copyright Epic Games, Inc. All Rights Reserved.
+
+using System;
+using System.Linq;
+using System.Threading.Tasks;
+using EpicGames.AspNet;
+using Microsoft.AspNetCore.Http;
+using Microsoft.Extensions.DependencyInjection;
+using Serilog;
+
+namespace Jupiter.Common
+{
+    public class SuppressExceptionMiddleware
+	{
+        private readonly ILogger _logger = Log.ForContext<SuppressExceptionMiddleware>();
+		private readonly RequestDelegate _next;
+
+        public SuppressExceptionMiddleware(RequestDelegate next)
+        {
+            _next = next ?? throw new ArgumentNullException(nameof(next));
+        }
+
+        public async Task Invoke(HttpContext context)
+        {
+            try
+            {
+                await _next(context);
+            }
+            // suppress certain known exceptions that are not considered errors but that we can not catch in our code
+            catch (ClientSendSlowException e)
+            {
+                // always suppress slow sends
+                _logger.Warning(e, "Client was sending data slowly");
+            }
+            catch (OperationCanceledException e)
+            {
+                // if health checks are cancelled its fine, it means it was called while it was running
+                string pathString = context.Request.Path.Value ?? "";
+
+                if (!pathString.StartsWith("/health", StringComparison.OrdinalIgnoreCase))
+                {
+                    // rethrow
+                    throw;
+                }
+                else
+                {
+                    _logger.Warning(e, "Health check was cancelled. This can happen if multiple health checks happen at the same time");
+                }
+            }
+        }
+    }
+
+	public static class SuppressExceptionMiddlewareCollectionExtensions
+	{
+        public static IServiceCollection AddSuppressExceptionsMiddleware(this IServiceCollection services)
+        {
+            services.AddSingleton<SuppressExceptionMiddleware>();
+
+            return services;
+        }
+    }
+}
diff --git a/Engine/Source/Programs/Shared/EpicGames.AspNet/ClientSendSlowException.cs b/Engine/Source/Programs/Shared/EpicGames.AspNet/ClientSendSlowException.cs
--- a/Engine/Source/Programs/Shared/EpicGames.AspNet/ClientSendSlowException.cs
+++ b/Engine/Source/Programs/Shared/EpicGames.AspNet/ClientSendSlowException.cs
@@ -31,7 +31,7 @@ namespace EpicGames.AspNet
 
 	public class ClientSendSlowException : Exception
 	{
-		public ClientSendSlowException(Exception? innerException) : base("Client was sending data to slow", innerException)
+		public ClientSendSlowException(Exception? innerException) : base("Client was sending data to slowly", innerException)
 		{
 		}
 	}
