Horde Storage - Fixed build errors by upgrading Microsoft.Caching and LZ4 libraries to match version used in shared libraries.
Also fixed some static code analysis errors.

diff --git a/Engine/Source/Programs/Horde/HordeStorage/Horde.Storage/Controllers/CompressedBlobController.cs b/Engine/Source/Programs/Horde/HordeStorage/Horde.Storage/Controllers/CompressedBlobController.cs
--- a/Engine/Source/Programs/Horde/HordeStorage/Horde.Storage/Controllers/CompressedBlobController.cs
+++ b/Engine/Source/Programs/Horde/HordeStorage/Horde.Storage/Controllers/CompressedBlobController.cs
@@ -15,7 +15,6 @@ using Horde.Storage.Implementation;
 using Jupiter;
 using Jupiter.Common.Implementation;
 using Jupiter.Implementation;
-using Jupiter.Utils;
 using Microsoft.AspNetCore.Authorization;
 using Microsoft.AspNetCore.Mvc;
 using Microsoft.Extensions.Primitives;
diff --git a/Engine/Source/Programs/Horde/HordeStorage/Horde.Storage/Horde.Storage.Startup.cs b/Engine/Source/Programs/Horde/HordeStorage/Horde.Storage/Horde.Storage.Startup.cs
--- a/Engine/Source/Programs/Horde/HordeStorage/Horde.Storage/Horde.Storage.Startup.cs
+++ b/Engine/Source/Programs/Horde/HordeStorage/Horde.Storage/Horde.Storage.Startup.cs
@@ -5,7 +5,6 @@ using System.Collections.Generic;
 using System.IO;
 using System.Linq;
 using System.Net;
-using System.Net.Http;
 using System.Net.Security;
 using System.Runtime.InteropServices;
 using System.Security.Authentication;
diff --git a/Engine/Source/Programs/Horde/HordeStorage/Horde.Storage/Horde.Storage.csproj b/Engine/Source/Programs/Horde/HordeStorage/Horde.Storage/Horde.Storage.csproj
--- a/Engine/Source/Programs/Horde/HordeStorage/Horde.Storage/Horde.Storage.csproj
+++ b/Engine/Source/Programs/Horde/HordeStorage/Horde.Storage/Horde.Storage.csproj
@@ -82,13 +82,13 @@
     <PackageReference Include="Azure.Storage.Blobs" Version="12.8.0" />
     <PackageReference Include="Destructurama.Attributed" Version="2.0.0" />
     <PackageReference Include="KubernetesClient" Version="4.0.10" />
-    <PackageReference Include="Microsoft.Extensions.Caching.Memory" Version="6.0.0" />
+    <PackageReference Include="Microsoft.Extensions.Caching.Memory" Version="6.0.1" />
     <PackageReference Include="Microsoft.Extensions.Logging.Debug" Version="6.0.0" />
     <PackageReference Include="MongoDB.Driver" Version="2.11.6" />
     <PackageReference Include="Nito.AsyncEx" Version="5.1.0" />
     <PackageReference Include="SerilogAnalyzer" Version="0.15.0" />
     <PackageReference Include="Crc32.NET" Version="1.2.0" />
-    <PackageReference Include="K4os.Compression.LZ4" Version="1.2.6" />
+    <PackageReference Include="K4os.Compression.LZ4" Version="1.2.16" />
 
     <PackageReference Include="CassandraCSharpDriver" Version="3.16.3" />
 
diff --git a/Engine/Source/Programs/Horde/HordeStorage/Jupiter.Common/SuppressExceptionMiddleware.cs b/Engine/Source/Programs/Horde/HordeStorage/Jupiter.Common/SuppressExceptionMiddleware.cs
--- a/Engine/Source/Programs/Horde/HordeStorage/Jupiter.Common/SuppressExceptionMiddleware.cs
+++ b/Engine/Source/Programs/Horde/HordeStorage/Jupiter.Common/SuppressExceptionMiddleware.cs
@@ -1,7 +1,6 @@
 // Copyright Epic Games, Inc. All Rights Reserved.
 
 using System;
-using System.Linq;
 using System.Threading.Tasks;
 using EpicGames.AspNet;
 using Microsoft.AspNetCore.Http;
