Horde Storage - Suppress the correct cancelled exception for health checks and actually register the middleware.

diff --git a/Engine/Source/Programs/Horde/HordeStorage/Jupiter.Common/BaseStartup.cs b/Engine/Source/Programs/Horde/HordeStorage/Jupiter.Common/BaseStartup.cs
--- a/Engine/Source/Programs/Horde/HordeStorage/Jupiter.Common/BaseStartup.cs
+++ b/Engine/Source/Programs/Horde/HordeStorage/Jupiter.Common/BaseStartup.cs
@@ -354,7 +354,8 @@ namespace Jupiter
             //app.UseMiddleware<DatadogTraceMiddleware>("Authorization");
             app.UseAuthorization();
 
-            app.UseMiddleware<ServerTimingMiddleware>();
+            app.UseMiddleware<SuppressExceptionMiddleware>();
+			app.UseMiddleware<ServerTimingMiddleware>();
 
             //app.UseMiddleware<DatadogTraceMiddleware>("Endpoints");
             app.UseEndpoints(endpoints =>
diff --git a/Engine/Source/Programs/Horde/HordeStorage/Jupiter.Common/SuppressExceptionMiddleware.cs b/Engine/Source/Programs/Horde/HordeStorage/Jupiter.Common/SuppressExceptionMiddleware.cs
--- a/Engine/Source/Programs/Horde/HordeStorage/Jupiter.Common/SuppressExceptionMiddleware.cs
+++ b/Engine/Source/Programs/Horde/HordeStorage/Jupiter.Common/SuppressExceptionMiddleware.cs
@@ -31,7 +31,7 @@ namespace Jupiter.Common
                 // always suppress slow sends
                 _logger.Warning(e, "Client was sending data slowly");
             }
-            catch (OperationCanceledException e)
+            catch (TaskCanceledException e)
             {
                 // if health checks are cancelled its fine, it means it was called while it was running
                 string pathString = context.Request.Path.Value ?? "";
