Horde: Misc fixes for bundles.

diff --git a/Engine/Source/Programs/Horde/Horde.Agent/Commands/Bundles/CreateCommand.cs b/Engine/Source/Programs/Horde/Horde.Agent/Commands/Bundles/CreateCommand.cs
--- a/Engine/Source/Programs/Horde/Horde.Agent/Commands/Bundles/CreateCommand.cs
+++ b/Engine/Source/Programs/Horde/Horde.Agent/Commands/Bundles/CreateCommand.cs
@@ -54,7 +54,7 @@ namespace Horde.Agent.Commands.Bundles
 				return Blob.Deserialize(id, bytes);
 			}
 
-			public async Task<IBlob> WriteBlobAsync(ReadOnlySequence<byte> data, IReadOnlyList<BlobId> references, RefName hintRefName = default, CancellationToken cancellationToken = default)
+			public async Task<BlobId> WriteBlobAsync(ReadOnlySequence<byte> data, IReadOnlyList<BlobId> references, RefName hintRefName = default, CancellationToken cancellationToken = default)
 			{
 				BlobId id = BlobId.Create(ServerId.Empty, hintRefName);
 				FileReference file = GetBlobFile(id);
@@ -62,7 +62,7 @@ namespace Horde.Agent.Commands.Bundles
 				_logger.LogInformation("Writing {File}", file);
 				IBlob blob = Blob.FromMemory(id, data.AsSingleSegment(), references);
 				await WriteAsync(file, Blob.Serialize(blob), cancellationToken);
-				return blob;
+				return id;
 			}
 
 			#endregion
@@ -97,11 +97,11 @@ namespace Horde.Agent.Commands.Bundles
 				return new BlobId(text.Trim());
 			}
 
-			public async Task<IBlob> WriteRefAsync(RefName name, ReadOnlySequence<byte> data, IReadOnlyList<BlobId> references, CancellationToken cancellationToken = default)
+			public async Task<BlobId> WriteRefAsync(RefName name, ReadOnlySequence<byte> data, IReadOnlyList<BlobId> references, CancellationToken cancellationToken = default)
 			{
-				IBlob blob = await WriteBlobAsync(data, references, name, cancellationToken);
-				await WriteRefTargetAsync(name, blob.Id, cancellationToken);
-				return blob;
+				BlobId blobId = await WriteBlobAsync(data, references, name, cancellationToken);
+				await WriteRefTargetAsync(name, blobId, cancellationToken);
+				return blobId;
 			}
 
 			public async Task WriteRefTargetAsync(RefName name, BlobId id, CancellationToken cancellationToken = default)
@@ -152,7 +152,8 @@ namespace Horde.Agent.Commands.Bundles
 			DirectoryNode node = new DirectoryNode(DirectoryFlags.None);
 			await node.CopyFromDirectoryAsync(InputDir.ToDirectoryInfo(), new ChunkingOptions(), writer, logger, CancellationToken.None);
 
-			await writer.FlushAsync(node, CancellationToken.None);
+			await writer.WriteNodeAsync(node);
+			await writer.FlushAsync();
 			return 0;
 		}
 	}
diff --git a/Engine/Source/Programs/Horde/Horde.Build/Startup.cs b/Engine/Source/Programs/Horde/Horde.Build/Startup.cs
--- a/Engine/Source/Programs/Horde/Horde.Build/Startup.cs
+++ b/Engine/Source/Programs/Horde/Horde.Build/Startup.cs
@@ -818,6 +818,21 @@ namespace Horde.Build
 			}
 		}
 
+		public sealed class RefNameBsonSerializer : SerializerBase<RefName>
+		{
+			/// <inheritdoc/>
+			public override RefName Deserialize(BsonDeserializationContext context, BsonDeserializationArgs args)
+			{
+				return new RefName(context.Reader.ReadString());
+			}
+
+			/// <inheritdoc/>
+			public override void Serialize(BsonSerializationContext context, BsonSerializationArgs args, RefName value)
+			{
+				context.Writer.WriteString(value.ToString());
+			}
+		}
+
 		static int s_haveConfiguredMongoDb = 0;
 
 		public static void ConfigureMongoDbClient()
@@ -832,6 +847,7 @@ namespace Horde.Build
 
 				// Register the custom serializers
 				BsonSerializer.RegisterSerializer(new RefIdBsonSerializer());
+				BsonSerializer.RegisterSerializer(new RefNameBsonSerializer());
 				BsonSerializer.RegisterSerializer(new ConditionSerializer());
 				BsonSerializer.RegisterSerializationProvider(new BsonSerializationProvider());
 				BsonSerializer.RegisterSerializationProvider(new StringIdSerializationProvider());
diff --git a/Engine/Source/Programs/Shared/EpicGames.Horde/Storage/Bundles/BundleStore.cs b/Engine/Source/Programs/Shared/EpicGames.Horde/Storage/Bundles/BundleStore.cs
--- a/Engine/Source/Programs/Shared/EpicGames.Horde/Storage/Bundles/BundleStore.cs
+++ b/Engine/Source/Programs/Shared/EpicGames.Horde/Storage/Bundles/BundleStore.cs
@@ -612,9 +612,10 @@ namespace EpicGames.Horde.Storage.Bundles
 							output = output.Slice(blockSegment.Length);
 						}
 
-						FlushPacket(format, buffer.Memory, builder, packets);
+						FlushPacket(format, buffer.Memory.Slice(0, blockSize), builder, packets);
 					}
 				}
+				blockSegments.Clear();
 			}
 
 			static void FlushPacket(BundleCompressionFormat format, ReadOnlyMemory<byte> inputData, ByteArrayBuilder builder, List<BundlePacket> packets)
