Horde: Fix tests

diff --git a/Engine/Source/Programs/Horde/Horde.Build.Tests/BlobStoreTests.cs b/Engine/Source/Programs/Horde/Horde.Build.Tests/BlobStoreTests.cs
--- a/Engine/Source/Programs/Horde/Horde.Build.Tests/BlobStoreTests.cs
+++ b/Engine/Source/Programs/Horde/Horde.Build.Tests/BlobStoreTests.cs
@@ -56,10 +56,9 @@ namespace Horde.Build.Tests
 
 			byte[] input = CreateTestData(256, 0);
 
-			IBlob blob = await store.WriteBlobAsync(new ReadOnlySequence<byte>(input), Array.Empty<BlobId>());
-			Assert.IsTrue(blob.Data.Span.SequenceEqual(input));
+			BlobId blobId = await store.WriteBlobAsync(new ReadOnlySequence<byte>(input), Array.Empty<BlobId>());
 
-			blob = await store.ReadBlobAsync(blob.Id);
+			IBlob blob = await store.ReadBlobAsync(blobId);
 			Assert.IsTrue(blob.Data.Span.SequenceEqual(input));
 		}
 
@@ -69,20 +68,20 @@ namespace Horde.Build.Tests
 			IBlobStore store = CreateBlobStore();
 
 			byte[] input1 = CreateTestData(256, 1);
-			IBlob blob1 = await store.WriteBlobAsync(new ReadOnlySequence<byte>(input1), Array.Empty<BlobId>());
-			blob1 = await store.ReadBlobAsync(blob1.Id);
+			BlobId blobId1 = await store.WriteBlobAsync(new ReadOnlySequence<byte>(input1), Array.Empty<BlobId>());
+			IBlob blob1 = await store.ReadBlobAsync(blobId1);
 			Assert.IsTrue(blob1.Data.Span.SequenceEqual(input1));
 			Assert.IsTrue(blob1.References.SequenceEqual(Array.Empty<BlobId>()));
 
 			byte[] input2 = CreateTestData(256, 2);
-			IBlob blob2 = await store.WriteBlobAsync(new ReadOnlySequence<byte>(input2), new BlobId[] { blob1.Id });
-			blob2 = await store.ReadBlobAsync(blob2.Id);
+			BlobId blobId2 = await store.WriteBlobAsync(new ReadOnlySequence<byte>(input2), new BlobId[] { blob1.Id });
+			IBlob blob2 = await store.ReadBlobAsync(blobId2);
 			Assert.IsTrue(blob2.Data.Span.SequenceEqual(input2));
 			Assert.IsTrue(blob2.References.SequenceEqual(new BlobId[] { blob1.Id }));
 
 			byte[] input3 = CreateTestData(256, 3);
-			IBlob blob3 = await store.WriteBlobAsync(new ReadOnlySequence<byte>(input3), new BlobId[] { blob1.Id, blob2.Id, blob1.Id });
-			blob3 = await store.ReadBlobAsync(blob3.Id);
+			BlobId blobId3 = await store.WriteBlobAsync(new ReadOnlySequence<byte>(input3), new BlobId[] { blob1.Id, blob2.Id, blob1.Id });
+			IBlob blob3 = await store.ReadBlobAsync(blobId3);
 			Assert.IsTrue(blob3.Data.Span.SequenceEqual(input3));
 			Assert.IsTrue(blob3.References.SequenceEqual(new BlobId[] { blob1.Id, blob2.Id, blob1.Id }));
 
@@ -96,13 +95,14 @@ namespace Horde.Build.Tests
 
 			RefName refName2 = new RefName("hello2");
 
-			IBlob refTarget = await store.WriteRefAsync(refName2, new ReadOnlySequence<byte>(input3), new BlobId[] { blob1.Id, blob2.Id });
-			Assert.IsTrue(refTarget.Data.Span.SequenceEqual(input3));
-			Assert.IsTrue(refTarget.References.SequenceEqual(new BlobId[] { blob1.Id, blob2.Id }));
+			BlobId refTargetId2 = await store.WriteRefAsync(refName2, new ReadOnlySequence<byte>(input3), new BlobId[] { blob1.Id, blob2.Id });
+			IBlob refTarget2 = await store.ReadBlobAsync(refTargetId2);
+			Assert.IsTrue(refTarget2.Data.Span.SequenceEqual(input3));
+			Assert.IsTrue(refTarget2.References.SequenceEqual(new BlobId[] { blob1.Id, blob2.Id }));
 
-			refTarget = await store.ReadRefAsync(refName2);
-			Assert.IsTrue(refTarget.Data.Span.SequenceEqual(input3));
-			Assert.IsTrue(refTarget.References.SequenceEqual(new BlobId[] { blob1.Id, blob2.Id }));
+			refTarget2 = await store.ReadRefAsync(refName2);
+			Assert.IsTrue(refTarget2.Data.Span.SequenceEqual(input3));
+			Assert.IsTrue(refTarget2.References.SequenceEqual(new BlobId[] { blob1.Id, blob2.Id }));
 		}
 	}
 }
