Horde: Fix some deprecation warnings building agent tests.

diff --git a/Engine/Source/Programs/Horde/Horde.Agent.Tests/AwsLambdaListenerCommandTest.cs b/Engine/Source/Programs/Horde/Horde.Agent.Tests/AwsLambdaListenerCommandTest.cs
--- a/Engine/Source/Programs/Horde/Horde.Agent.Tests/AwsLambdaListenerCommandTest.cs
+++ b/Engine/Source/Programs/Horde/Horde.Agent.Tests/AwsLambdaListenerCommandTest.cs
@@ -24,7 +24,7 @@ namespace Horde.Agent.Tests
 		{
 			using ILoggerFactory loggerFactory = LoggerFactory.Create(builder =>
 			{
-				builder.AddConsole(consoleLoggerOptions => consoleLoggerOptions.TimestampFormat = "[HH:mm:ss] ");
+				builder.AddSimpleConsole(consoleLoggerOptions => consoleLoggerOptions.TimestampFormat = "[HH:mm:ss] ");
 			});
 
 			_workerLogger = loggerFactory.CreateLogger<WorkerService>();
diff --git a/Engine/Source/Programs/Horde/Horde.Agent.Tests/WorkerServiceTest.cs b/Engine/Source/Programs/Horde/Horde.Agent.Tests/WorkerServiceTest.cs
--- a/Engine/Source/Programs/Horde/Horde.Agent.Tests/WorkerServiceTest.cs
+++ b/Engine/Source/Programs/Horde/Horde.Agent.Tests/WorkerServiceTest.cs
@@ -41,7 +41,7 @@ namespace Horde.Agent.Tests
 		{
 			using ILoggerFactory loggerFactory = LoggerFactory.Create(builder =>
 			{
-				builder.AddConsole(consoleLoggerOptions => consoleLoggerOptions.TimestampFormat = "[HH:mm:ss] ");
+				builder.AddSimpleConsole(consoleLoggerOptions => consoleLoggerOptions.TimestampFormat = "[HH:mm:ss] ");
 			});
 
 			_workerLogger = loggerFactory.CreateLogger<WorkerService>();
