Horde: Additional logging for Slack escalations.

diff --git a/Engine/Source/Programs/Horde/Horde.Build/Notifications/Sinks/SlackNotificationSink.cs b/Engine/Source/Programs/Horde/Horde.Build/Notifications/Sinks/SlackNotificationSink.cs
--- a/Engine/Source/Programs/Horde/Horde.Build/Notifications/Sinks/SlackNotificationSink.cs
+++ b/Engine/Source/Programs/Horde/Horde.Build/Notifications/Sinks/SlackNotificationSink.cs
@@ -2061,20 +2061,27 @@ namespace Horde.Build.Notifications.Sinks
 		async ValueTask EscalateAsync(CancellationToken cancellationToken)
 		{
 			DateTime utcNow = DateTime.UtcNow;
+			double time = (utcNow - DateTime.UnixEpoch).TotalSeconds;
 
-			int[] issueIds = await _escalateIssues.RangeByScoreAsync(0, (utcNow - DateTime.UnixEpoch).TotalSeconds);
-			foreach (int issueId in issueIds)
+			int[] issueIds = await _escalateIssues.RangeByScoreAsync(0, time);
+			if (issueIds.Length > 0)
 			{
-				cancellationToken.ThrowIfCancellationRequested();
-
-				double? nextTime = await EscalateSingleIssueAsync(issueId, utcNow);
-				if (nextTime == null)
+				_logger.LogInformation("Escalating issues for {Time} ({TimeSecs})", utcNow, time);
+				foreach (int issueId in issueIds)
 				{
-					await _escalateIssues.RemoveAsync(issueId);
-				}
-				else
-				{
-					await _escalateIssues.AddAsync(issueId, nextTime.Value);
+					cancellationToken.ThrowIfCancellationRequested();
+
+					double? nextTime = await EscalateSingleIssueAsync(issueId, utcNow);
+					if (nextTime == null)
+					{
+						_logger.LogInformation("Cancelling escalation for issue {IssueId}", issueId);
+						await _escalateIssues.RemoveAsync(issueId);
+					}
+					else
+					{
+						_logger.LogInformation("Next escalation for issue {IssueId} is at timestamp {Time}", issueId, nextTime.Value);
+						await _escalateIssues.AddAsync(issueId, nextTime.Value);
+					}
 				}
 			}
 		}
