Horde: Remove unused methods from PerforceService.

diff --git a/Engine/Source/Programs/Horde/Horde.Build.Tests/Stubs/Services/PerforceServiceStub.cs b/Engine/Source/Programs/Horde/Horde.Build.Tests/Stubs/Services/PerforceServiceStub.cs
--- a/Engine/Source/Programs/Horde/Horde.Build.Tests/Stubs/Services/PerforceServiceStub.cs
+++ b/Engine/Source/Programs/Horde/Horde.Build.Tests/Stubs/Services/PerforceServiceStub.cs
@@ -113,11 +113,6 @@ namespace Horde.Build.Tests.Stubs.Services
 			throw new NotImplementedException();
 		}
 
-		public Task<PerforceUserInfo?> GetUserInfoAsync(string clusterName, string userName)
-		{
-			return Task.FromResult<PerforceUserInfo?>(new PerforceUserInfo { Login = userName, FullName = userName, Email = $"[REDACTED_EMAIL]" });
-		}
-
 		public Task<List<ChangeDetails>> GetChangeDetailsAsync(string clusterName, string streamName, IReadOnlyList<int> changeNumbers, string? impersonateUser)
 		{
 			List<ChangeDetails> results = new List<ChangeDetails>();
@@ -190,11 +185,6 @@ namespace Horde.Build.Tests.Stubs.Services
 			return Task.CompletedTask;
 		}
 
-		public Task<IStreamView> GetStreamViewAsync(string clusterName, string streamName)
-		{
-			throw new NotImplementedException();
-		}
-
 		public Task<List<ChangeSummary>> GetChangesAsync(string clusterName, int? minChange, int? maxChange, int maxResults)
 		{
 			throw new NotImplementedException();
diff --git a/Engine/Source/Programs/Horde/Horde.Build/Perforce/IPerforceService.cs b/Engine/Source/Programs/Horde/Horde.Build/Perforce/IPerforceService.cs
--- a/Engine/Source/Programs/Horde/Horde.Build/Perforce/IPerforceService.cs
+++ b/Engine/Source/Programs/Horde/Horde.Build/Perforce/IPerforceService.cs
@@ -12,27 +12,6 @@ using Horde.Build.Users;
 
 namespace Horde.Build.Perforce
 {
-	/// <summary>
-	/// Information about a Perforce user
-	/// </summary>
-	public class PerforceUserInfo
-	{
-		/// <summary>
-		/// Login for the user
-		/// </summary>
-		public string? Login { get; set; }
-
-		/// <summary>
-		/// Full name of the user
-		/// </summary>
-		public string? FullName { get; set; }
-
-		/// <summary>
-		/// User's email address
-		/// </summary>
-		public string? Email { get; set; }
-	}
-
 	/// <summary>
 	/// Result from checking a shelved change status
 	/// </summary>
@@ -64,20 +43,6 @@ namespace Horde.Build.Perforce
 		WrongStream,
 	}
 
-	/// <summary>
-	/// Interface for a Perforce stream view
-	/// </summary>
-	public interface IStreamView : IDisposable
-	{
-		/// <summary>
-		/// Maps a depot path into the stream
-		/// </summary>
-		/// <param name="depotPath">The depot path</param>
-		/// <param name="streamPath">The resulting stream path</param>
-		/// <returns></returns>
-		bool TryGetStreamPath(string depotPath, [NotNullWhen(true)] out string? streamPath);
-	}
-
 	/// <summary>
 	/// Wrapper around Perforce functionality. Can use a local p4.exe client for development purposes, or a separate HordePerforceBridge instance over REST for deployments.
 	/// </summary>
@@ -98,14 +63,6 @@ namespace Horde.Build.Perforce
 		/// <returns></returns>
 		public Task<IPerforceConnection?> GetServiceUserConnection(string? clusterName);
 
-		/// <summary>
-		/// Gets the definition of a stream
-		/// </summary>
-		/// <param name="clusterName">Name of the Perforce cluster</param>
-		/// <param name="streamName">The stream name</param>
-		/// <returns></returns>
-		public Task<IStreamView> GetStreamViewAsync(string clusterName, string streamName);
-
 		/// <summary>
 		/// Create a new changelist by submitting the given file
 		/// </summary>
@@ -125,14 +82,6 @@ namespace Horde.Build.Perforce
 		/// <returns>Code change for the latest change</returns>
 		public Task<int> GetCodeChangeAsync(string clusterName, string streamName, int change);
 
-		/// <summary>
-		/// Gets information about the given user
-		/// </summary>
-		/// <param name="clusterName">Name of the Perforce cluster</param>
-		/// <param name="userName">The user name</param>
-		/// <returns>User information</returns>
-		public Task<PerforceUserInfo?> GetUserInfoAsync(string clusterName, string userName);
-
 		/// <summary>
 		/// Finds changes submitted to a depot Gets the latest change for a particular stream
 		/// </summary>
diff --git a/Engine/Source/Programs/Horde/Horde.Build/Perforce/PerforceService.cs b/Engine/Source/Programs/Horde/Horde.Build/Perforce/PerforceService.cs
--- a/Engine/Source/Programs/Horde/Horde.Build/Perforce/PerforceService.cs
+++ b/Engine/Source/Programs/Horde/Horde.Build/Perforce/PerforceService.cs
@@ -136,8 +136,8 @@ namespace Horde.Build.Perforce
 				user = await _userCollection.FindUserByLoginAsync(userName);
 				if (user == null)
 				{
-					PerforceUserInfo? userInfo = await GetUserInfoAsync(clusterName, userName);
-					user = await _userCollection.FindOrAddUserByLoginAsync(userName, userInfo?.FullName, userInfo?.Email);
+					P4.User? userInfo = await GetUserInfoAsync(clusterName, userName);
+					user = await _userCollection.FindOrAddUserByLoginAsync(userName, userInfo?.FullName, userInfo?.EmailAddress);
 				}
 
 				using (ICacheEntry entry = _userCache.CreateEntry((clusterName, userName)))
@@ -698,74 +698,6 @@ namespace Horde.Build.Perforce
 			P4.LogFile.LogMessage(logLevel, category, message);
 		}
 
-		class StreamView : IStreamView
-		{
-			readonly P4.P4MapApi _mapApi;
-
-			public StreamView(P4.P4MapApi mapApi)
-			{
-				_mapApi = mapApi;
-			}
-
-			public void Dispose()
-			{
-				_mapApi.Dispose();
-			}
-
-			public bool TryGetStreamPath(string depotPath, out string streamPath)
-			{
-				streamPath = _mapApi.Translate(depotPath, P4.P4MapApi.Direction.LeftRight);
-				return streamPath != null;
-			}
-		}
-
-		/// <inheritdoc/>
-		public async Task<IStreamView> GetStreamViewAsync(string clusterName, string streamName)
-		{
-			using IScope scope = GlobalTracer.Instance.BuildSpan("PerforceService.GetStreamViewAsync").StartActive();
-			scope.Span.SetTag("ClusterName", clusterName);
-			scope.Span.SetTag("StreamName", streamName);
-			
-			PerforceCluster cluster = await GetClusterAsync(clusterName);
-			using (P4.Repository repository = await GetConnection(cluster, noClient: true))
-			{
-				P4.Stream stream = repository.GetStream(streamName, new P4.StreamCmdOptions(P4.StreamCmdFlags.View, null, null));
-
-				P4.P4MapApi? mapApi = null;
-				try
-				{
-					mapApi = new P4.P4MapApi(repository.Server.Metadata.UnicodeEnabled);
-					foreach (P4.MapEntry entry in stream.View)
-					{
-						P4.P4MapApi.Type mapType;
-						switch (entry.Type)
-						{
-							case P4.MapType.Include:
-								mapType = P4.P4MapApi.Type.Include;
-								break;
-							case P4.MapType.Exclude:
-								mapType = P4.P4MapApi.Type.Exclude;
-								break;
-							case P4.MapType.Overlay:
-								mapType = P4.P4MapApi.Type.Overlay;
-								break;
-							default:
-								throw new Exception($"Invalid map type: {entry.Type}");
-						}
-						mapApi.Insert(entry.Left.ToString(), entry.Right.ToString(), mapType);
-					}
-
-					StreamView view = new StreamView(mapApi);
-					mapApi = null;
-					return view;
-				}
-				finally
-				{
-					mapApi?.Dispose();
-				}
-			}
-		}
-
 		static int GetSyncRevision(string path, P4.FileAction headAction, int headRev)
 		{
 			switch (headAction)
@@ -1533,7 +1465,7 @@ namespace Horde.Build.Perforce
 		}
 
 		/// <inheritdoc/>		
-		public async Task<PerforceUserInfo?> GetUserInfoAsync(string clusterName, string userName)
+		async Task<P4.User?> GetUserInfoAsync(string clusterName, string userName)
 		{
 			using IScope scope = GlobalTracer.Instance.BuildSpan("PerforceService.GetUserInfoAsync").StartActive();
 			scope.Span.SetTag("ClusterName", clusterName);
@@ -1543,13 +1475,7 @@ namespace Horde.Build.Perforce
 			using (P4.Repository repository = await GetConnection(cluster, noClient: true))
 			{
 				P4.User user = repository.GetUser(userName, new P4.UserCmdOptions(P4.UserCmdFlags.Output));
-
-				if (user == null)
-				{
-					return null;
-				}
-
-				return new PerforceUserInfo { Login = userName, FullName = user.FullName, Email = user.EmailAddress };
+				return user;
 			}
 		}
 
