Horde: Adding device telemetry to dashboard

diff --git a/Engine/Source/Programs/Horde/HordeDashboard/package.json b/Engine/Source/Programs/Horde/HordeDashboard/package.json
--- a/Engine/Source/Programs/Horde/HordeDashboard/package.json
+++ b/Engine/Source/Programs/Horde/HordeDashboard/package.json
@@ -10,6 +10,7 @@
     "@fluentui/react": "^8.36.5",
     "@fluentui/react-charting": "^5.2.0",
     "crypto-js": "^4.0.0",
+    "d3": "^7.6.1",
     "formik": "^2.1.1",
     "idb": "^7.0.2",
     "markdown-to-jsx": "^6.10.3",
@@ -49,6 +50,7 @@
   },
   "devDependencies": {
     "@types/crypto-js": "^4.0.1",
+    "@types/d3": "^7.4.0",
     "@types/markdown-to-jsx": "6.9.1",
     "@types/react": "^17.0.30",
     "@types/react-custom-scrollbars": "^4.0.6",
diff --git a/Engine/Source/Programs/Horde/HordeDashboard/src/backend/Api.ts b/Engine/Source/Programs/Horde/HordeDashboard/src/backend/Api.ts
--- a/Engine/Source/Programs/Horde/HordeDashboard/src/backend/Api.ts
+++ b/Engine/Source/Programs/Horde/HordeDashboard/src/backend/Api.ts
@@ -4169,3 +4169,118 @@ export type GetNoticeResponse = {
 	createdByUser?: GetThinUserInfoResponse;
 
 }
+
+// Device telemetry
+
+export type DevicePoolTelemetryQuery = {
+	minCreateTime?: string;
+	maxCreateTime?: string;
+	index?: number;
+	count?: number;
+};
+
+export type DeviceTelemetryQuery = {
+	deviceIds?: string[];
+	poolId?: string;
+	platformId?: string;
+	minCreateTime?: string;
+	maxCreateTime?: string;
+	index?: number;
+	count?: number;
+};
+
+export type GetDevicePoolReservationTelemetryResponse = {
+	
+	/** Device id for reservation */
+	deviceId: string;
+
+	/** Job id associated with reservation */
+	jobId?: string; 
+
+	/** The step id of reservation */
+	stepId?: string; 
+
+	/** The name of the job holding reservation */
+	jobName?: string; 
+
+	/** The name of the step holding reservation */
+	stepName?: string;
+
+}
+
+/** Device pool telemetry respponse */
+export type GetDevicePlatformTelemetryResponse = {
+
+	/** The corresponding platform id */
+	platformId: string;
+
+	/** Available device ids of this platform */
+	available?: string[];
+
+	/** Device ids in maintenance state */
+	maintenance?: string[];
+
+	/** Device ids in problem state */
+	problem?: string[];
+
+	/** Device ids in disabled state */
+	disabled?: string[];
+
+	/** Stream id => reserved devices of this platform */
+	reserved?: Record<string, GetDevicePoolReservationTelemetryResponse[]>;
+}
+
+
+/** Device telemetry respponse */
+export type GetDevicePoolTelemetryResponse = {
+
+	/** The UTC time the telemetry data was created */
+	createTimeUtc: Date | string;
+
+
+	/** Individual pool id -> telemetry data points */
+	telemetry: Record<string, GetDevicePlatformTelemetryResponse[]>;
+}
+
+/// Device telemetry respponse
+export type GetTelemetryInfoResponse = {
+
+	/// The UTC time the telemetry data was created	
+	createTimeUtc: Date | string;
+	
+	/// The stream id which utilized device	
+	streamId?: string;
+	
+	/// The job id which utilized device	
+	jobId?: string;
+	
+	/// The job's step id	
+	stepId?: string;
+
+	/// The job name which utilized device	
+	jobName?: string;
+	
+	/// The job's step name
+	stepName?: string;
+
+	/// If this telemetry has a reservation, the start time of the reservation	
+	reservationStartUtc?: Date | string;
+
+	/// If this telemetry has a reservation, the finish time of the reservation
+	reservationFinishUtc?: Date | string;
+	
+	/// If this telemetry marks a detected device issue, the time of the issue	
+	problemTimeUtc?: Date | string;
+}
+
+
+/// Device telemetry respponse
+export type GetDeviceTelemetryResponse = {
+
+	/// The device id for the telemetry data
+	deviceId: string;
+
+	/// Individual telemetry data points	
+	telemetry: GetTelemetryInfoResponse[];
+
+}
\ No newline at end of file
diff --git a/Engine/Source/Programs/Horde/HordeDashboard/src/backend/Backend.ts b/Engine/Source/Programs/Horde/HordeDashboard/src/backend/Backend.ts
--- a/Engine/Source/Programs/Horde/HordeDashboard/src/backend/Backend.ts
+++ b/Engine/Source/Programs/Horde/HordeDashboard/src/backend/Backend.ts
@@ -3,7 +3,7 @@
 import { configure } from "mobx";
 import { isNumber } from 'util';
 import templateCache from '../backend/TemplateCache';
-import { AgentData, AgentQuery, ArtifactData, AuditLogEntry, AuditLogQuery, BatchUpdatePoolRequest, ChangeSummaryData, CreateDeviceRequest, CreateDeviceResponse, CreateExternalIssueRequest, CreateExternalIssueResponse, CreateJobRequest, CreateJobResponse, CreateNoticeRequest, CreatePoolRequest, CreateSoftwareResponse, CreateSubscriptionRequest, CreateSubscriptionResponse, DashboardPreference, EventData, FindIssueResponse, FindJobTimingsResponse, GetAgentSoftwareChannelResponse, GetArtifactZipRequest, GetDashboardConfigResponse, GetDevicePlatformResponse, GetDevicePoolResponse, GetDeviceReservationResponse, GetDeviceResponse, GetExternalIssueProjectResponse, GetExternalIssueResponse, GetGraphResponse, GetIssueStreamResponse, GetJobsTabResponse, GetJobStepRefResponse, GetJobStepTraceResponse, GetJobTimingResponse, GetLogEventResponse, GetNoticeResponse, GetNotificationResponse, GetPerforceServerStatusResponse, GetPoolResponse, GetServerInfoResponse, GetServerSettingsResponse, GetSoftwareResponse, GetSubscriptionResponse, GetUserResponse, GetUtilizationTelemetryResponse, GlobalConfig, IssueData, IssueQuery, IssueQueryV2, JobData, JobQuery, JobsTabColumnType, JobStepOutcome, JobStreamQuery, JobTimingsQuery, LeaseData, LogData, LogLineData, PoolData, ProjectData, ScheduleData, ScheduleQuery, SearchLogFileResponse, ServerUpdateResponse, SessionData, StreamData, TabType, TemplateData, TestData, UpdateAgentRequest, UpdateDeviceRequest, UpdateGlobalConfigRequest, UpdateIssueRequest, UpdateJobRequest, UpdateLeaseRequest, UpdateNoticeRequest, UpdateNotificationsRequest, UpdatePoolRequest, UpdateServerSettingsRequest, UpdateStepRequest, UpdateStepResponse, UpdateTemplateRefRequest, UpdateUserRequest, UsersQuery } from './Api';
+import { AgentData, AgentQuery, ArtifactData, AuditLogEntry, AuditLogQuery, BatchUpdatePoolRequest, ChangeSummaryData, CreateDeviceRequest, CreateDeviceResponse, CreateExternalIssueRequest, CreateExternalIssueResponse, CreateJobRequest, CreateJobResponse, CreateNoticeRequest, CreatePoolRequest, CreateSoftwareResponse, CreateSubscriptionRequest, CreateSubscriptionResponse, DashboardPreference, DevicePoolTelemetryQuery, DeviceTelemetryQuery, EventData, FindIssueResponse, FindJobTimingsResponse, GetAgentSoftwareChannelResponse, GetArtifactZipRequest, GetDashboardConfigResponse, GetDevicePlatformResponse, GetDevicePoolResponse, GetDevicePoolTelemetryResponse, GetDeviceReservationResponse, GetDeviceResponse, GetDeviceTelemetryResponse, GetExternalIssueProjectResponse, GetExternalIssueResponse, GetGraphResponse, GetIssueStreamResponse, GetJobsTabResponse, GetJobStepRefResponse, GetJobStepTraceResponse, GetJobTimingResponse, GetLogEventResponse, GetNoticeResponse, GetNotificationResponse, GetPerforceServerStatusResponse, GetPoolResponse, GetServerInfoResponse, GetServerSettingsResponse, GetSoftwareResponse, GetSubscriptionResponse, GetUserResponse, GetUtilizationTelemetryResponse, GlobalConfig, IssueData, IssueQuery, IssueQueryV2, JobData, JobQuery, JobsTabColumnType, JobStepOutcome, JobStreamQuery, JobTimingsQuery, LeaseData, LogData, LogLineData, PoolData, ProjectData, ScheduleData, ScheduleQuery, SearchLogFileResponse, ServerUpdateResponse, SessionData, StreamData, TabType, TemplateData, TestData, UpdateAgentRequest, UpdateDeviceRequest, UpdateGlobalConfigRequest, UpdateIssueRequest, UpdateJobRequest, UpdateLeaseRequest, UpdateNoticeRequest, UpdateNotificationsRequest, UpdatePoolRequest, UpdateServerSettingsRequest, UpdateStepRequest, UpdateStepResponse, UpdateTemplateRefRequest, UpdateUserRequest, UsersQuery } from './Api';
 import dashboard from './Dashboard';
 import { ChallengeStatus, Fetch } from './Fetch';
 import graphCache, { GraphQuery } from './GraphCache';
@@ -1334,6 +1334,28 @@ export class Backend {
         });
     }
 
+    getDevicePoolTelemetry(query?: DevicePoolTelemetryQuery): Promise<GetDevicePoolTelemetryResponse[]> {
+
+        return new Promise<GetDevicePoolTelemetryResponse[]>((resolve, reject) => {
+            this.backend.get(`/api/v2/devices/pools/telemetry`, { params: query }).then((value) => {
+                resolve(value.data as GetDevicePoolTelemetryResponse[]);
+            }).catch(reason => {
+                reject(reason);
+            });
+        });
+    }
+
+    getDeviceTelemetry(query?: DeviceTelemetryQuery): Promise<GetDeviceTelemetryResponse[]> {
+
+        return new Promise<GetDeviceTelemetryResponse[]>((resolve, reject) => {
+            this.backend.get(`/api/v2/devices/telemetry`, { params: query }).then((value) => {
+                resolve(value.data as GetDeviceTelemetryResponse[]);
+            }).catch(reason => {
+                reject(reason);
+            });
+        });
+    }
+
     addDevice(request: CreateDeviceRequest): Promise<CreateDeviceResponse> {
         return new Promise<CreateDeviceResponse>((resolve, reject) => {
             this.backend.post(`/api/v2/devices`, request).then((value) => {
diff --git a/Engine/Source/Programs/Horde/HordeDashboard/src/base/utilities/timeUtils.ts b/Engine/Source/Programs/Horde/HordeDashboard/src/base/utilities/timeUtils.ts
--- a/Engine/Source/Programs/Horde/HordeDashboard/src/base/utilities/timeUtils.ts
+++ b/Engine/Source/Programs/Horde/HordeDashboard/src/base/utilities/timeUtils.ts
@@ -15,7 +15,7 @@ type HordeTime = {
     server: string;
 }
 
-export const msecToElapsed = (millisec: number, includeMinutes: boolean = true): string => {
+export const msecToElapsed = (millisec: number, includeMinutes: boolean = true, includeSeconds: boolean = true): string => {
 
     let duration = "";
     const d = moment.duration(millisec);
@@ -34,7 +34,7 @@ export const msecToElapsed = (millisec: number, includeMinutes: boolean = true):
             duration += `${d.minutes()}m `;
         }
 
-        if (d.seconds()) {
+        if (d.seconds() && (!duration || includeSeconds)) {
             duration += `${d.seconds()}s `;
         }
     }
diff --git a/Engine/Source/Programs/Horde/HordeDashboard/src/components/DevicePoolTelemetry.tsx b/Engine/Source/Programs/Horde/HordeDashboard/src/components/DevicePoolTelemetry.tsx
new file mode 100644
--- /dev/null
+++ b/Engine/Source/Programs/Horde/HordeDashboard/src/components/DevicePoolTelemetry.tsx
@@ -0,0 +1,1228 @@
+import { DefaultButton, DetailsList, DetailsListLayoutMode, Dropdown, FocusZone, FocusZoneDirection, IconButton, IDropdownOption, Modal, ScrollablePane, ScrollbarVisibility, SelectionMode, Spinner, SpinnerSize, Stack, Text } from "@fluentui/react";
+import * as d3 from "d3";
+import { action, observable } from "mobx";
+import { observer } from "mobx-react-lite";
+import moment from "moment";
+import { useState } from "react";
+import backend from "../backend";
+import { DevicePoolTelemetryQuery, DevicePoolType, DeviceTelemetryQuery, GetDevicePlatformResponse, GetDevicePlatformTelemetryResponse, GetDevicePoolResponse, GetDevicePoolTelemetryResponse, GetDeviceResponse, GetDeviceTelemetryResponse, GetTelemetryInfoResponse } from "../backend/Api";
+import dashboard, { StatusColor } from "../backend/Dashboard";
+import { projectStore } from "../backend/ProjectStore";
+import { getHumanTime, msecToElapsed } from "../base/utilities/timeUtils";
+import { hordeClasses, modeColors } from "../styles/Styles";
+
+
+type PoolId = string;
+
+type SelectionType = d3.Selection<SVGGElement, unknown, null, undefined>;
+type ScalarTime = d3.ScaleTime<number, number>
+type DivSelectionType = d3.Selection<HTMLDivElement, unknown, null, undefined>;
+
+type PoolTelemetryData = {
+   date: Date;
+   dateString: string;
+   pools: Record<PoolId, GetDevicePlatformTelemetryResponse[]>;
+};
+
+type PlatformTelemetry = GetDevicePlatformTelemetryResponse & {
+   date: Date;
+   dateString: string;
+};
+
+type DeviceTelemetry = GetTelemetryInfoResponse & {
+   createTime: Date;
+   reservationStart?: Date;
+   reservationFinish?: Date;
+   problemTime?: Date;
+}
+
+
+type TelemetryData = {
+   pool: GetDevicePoolResponse;
+   platform: GetDevicePlatformResponse;
+   streamIds: string[];
+   // stream id => full stream name
+   streamNames: Record<string, string>;
+   // stream full stream name => stream id
+   streamNamesReverse: Record<string, string>;
+   // stream name => valid selector
+   streamSelectors: Record<string, string>;
+   // selector => stream name
+   streamSelectorsReverse: Record<string, string>;
+   deviceIds: string[];
+   devices: Record<string, GetDeviceResponse>;
+   telemetry: PlatformTelemetry[];
+}
+
+type StepMetrics = {
+   numKits: number;
+   numProblems: number;
+   duration: string;
+}
+
+type ProblemStep = {
+   streamId?: string;
+   jobId?: string;
+   jobName?: string;
+   stepId?: string;
+   stepName?: string;
+}
+
+class PoolTelemetryHandler {
+
+   constructor() {
+      this.set();
+   }
+
+   getProblemSteps(deviceIds: string[], time: Date): Record<string, ProblemStep> {
+      const problems: Record<string, ProblemStep> = {};
+
+      deviceIds.forEach(id => {
+         const telemetry = this.deviceTelemetry.get(id);
+         if (!telemetry) {
+            return;
+         }
+
+         let best: DeviceTelemetry | undefined;
+         let bestMS = Number.MAX_VALUE;
+
+         telemetry.forEach(t => {
+
+            if (!t.problemTime) {
+               return;
+            }
+
+            const ctime = Math.abs(t.problemTime.getTime() - time.getTime());
+
+            if (ctime > 60 * 1000 * 30) {
+               return;
+            }
+
+            if (bestMS > ctime) {
+               best = t;
+               bestMS = ctime;
+            }
+
+         });
+
+         if (!best) {
+            return;
+         }
+
+         problems[id] = {
+            streamId: best.streamId,
+            jobId: best.jobId,
+            jobName: best.jobName,
+            stepId: best.stepId,
+            stepName: best.stepName
+         };
+
+      });
+
+      return problems;
+   }
+
+   getSteps(platformId: string) {
+
+      const stepData: Map<string, { streamId: string, stepName: string, reservations: number, problems: number, durationMS: number }> = new Map();
+
+      this.deviceTelemetry.forEach((telemetry, deviceId) => {
+
+         const device = this.devices.get(deviceId);
+         if (!device || device.platformId !== platformId) {
+            return;
+         }
+
+         telemetry.forEach(t => {
+
+            let stepName = t.stepName ?? "Unknown Step";
+            let streamId = t.streamId ?? "Unknown Stream";
+
+            /*
+            let streamId = t.streamId ?? "Unknown Stream";
+            let jobName = t.jobName ?? "Unknown Job";
+            
+
+            if (jobName.indexOf("- Kicked By") !== -1) {
+               jobName = jobName.split("- Kicked By")[0];
+            }
+            */
+
+            const key = stepName + streamId;
+
+            if (!stepData.has(key)) {
+               stepData.set(key, { streamId: streamId, stepName: stepName, reservations: 0, problems: 0, durationMS: 0 });
+            }
+
+            const step = stepData.get(key)!;
+
+            if (t.problemTimeUtc) {
+               step.problems++;
+            } else {
+               step.reservations++;
+               if (t.reservationStart && t.reservationFinish) {
+                  step.durationMS += (t.reservationFinish.getTime() - t.reservationStart.getTime());
+               }
+            }
+         });
+      });
+
+      return stepData;
+   }
+
+   getProblemDevices() {
+
+      const platforms: Map<string, { deviceId: string, deviceName: string, problems: number, latestProblem?: string }[]> = new Map();
+
+      this.deviceTelemetry.forEach((telemetry, deviceId) => {
+
+         const device = this.devices.get(deviceId);
+         if (!device) {
+            return;
+         }
+
+         const platform = this.platforms.get(device.platformId);
+         if (!platform) {
+            return;
+         }
+
+         let count = 0;
+         let latest: Date | undefined;
+         telemetry.forEach(t => { if (t.problemTimeUtc) count++; if (!latest || latest < t.problemTime!) latest = t.problemTime });
+
+         if (!count) {
+            return;
+         }
+
+         if (!platforms.has(platform.id)) {
+            platforms.set(platform.id, []);
+         }
+
+         platforms.get(platform.id)!.push({ deviceId: deviceId, deviceName: device.name, problems: count, latestProblem: latest ? getHumanTime(latest) : "" });
+
+      });
+
+      platforms.forEach((devices, key) => {
+         devices = devices.sort((a, b) => b.problems - a.problems);
+         platforms.set(key, devices);
+      });
+
+      return platforms;
+
+   }
+
+   getStepMetrics(jobId: string, stepId: string): StepMetrics | undefined {
+
+      const telemetry: DeviceTelemetry[] = [];
+
+      const devices = new Set<string>();
+
+      this.deviceTelemetry.forEach((t, id) => {
+
+         const stepTelemetry = t.filter(td => { return (td.jobId === jobId) && (td.stepId === stepId) });
+
+         if (stepTelemetry.length) {
+            devices.add(id);
+            telemetry.push(...stepTelemetry);
+         }
+
+      });
+
+      if (!telemetry.length) {
+         return undefined;
+      }
+
+      let minDate: Date | undefined;
+      let maxDate: Date | undefined;
+
+      telemetry.forEach(t => {
+
+         if (t.reservationStart) {
+            if (!minDate || minDate > t.reservationStart) {
+               minDate = t.reservationStart;
+            }
+         }
+
+         if (t.reservationFinish) {
+            if (!maxDate || maxDate < t.reservationFinish) {
+               maxDate = t.reservationFinish;
+            }
+         }
+      });
+
+      let duration = "";
+      if (minDate && maxDate) {
+         duration = msecToElapsed(maxDate.getTime() - minDate.getTime(), true, false);
+      }
+
+      let numProblems = 0;
+
+      telemetry.forEach(t => {
+         if (t.problemTime) {
+            numProblems++;
+         }
+      });
+
+      return {
+         numKits: devices.size,
+         numProblems: numProblems,
+         duration: duration
+      }
+   }
+
+   async set() {
+
+      const requests: any = [];
+
+      // get 2 weeks worth
+      const date = moment().subtract(14, 'days');
+
+      const poolQuery: DevicePoolTelemetryQuery = { minCreateTime: date.toISOString(), count: 65536 };
+      const deviceQuery: DeviceTelemetryQuery = { minCreateTime: date.toISOString(), count: 65536 };
+
+      requests.push(backend.getDevicePlatforms());
+      requests.push(backend.getDevicePools());
+      requests.push(backend.getDevices());
+      requests.push(backend.getDevicePoolTelemetry(poolQuery));
+      requests.push(backend.getDeviceTelemetry(deviceQuery));
+
+      const responses = await Promise.all(requests);
+
+      const platforms = responses[0] as GetDevicePlatformResponse[];
+      const pools = responses[1] as GetDevicePoolResponse[];
+      const devices = responses[2] as GetDeviceResponse[];
+      const telemetry = responses[3] as GetDevicePoolTelemetryResponse[];
+      const deviceTelemetry = responses[4] as GetDeviceTelemetryResponse[];
+
+      deviceTelemetry.forEach(d => {
+         const telemetry: DeviceTelemetry[] = [];
+         d.telemetry.forEach(t => {
+
+            const t2: DeviceTelemetry = { ...t, createTime: new Date(t.createTimeUtc) };
+
+            if (t.reservationStartUtc) {
+               t2.reservationStart = new Date(t.reservationStartUtc);
+            }
+
+            if (t.problemTimeUtc) {
+               t2.problemTime = new Date(t.problemTimeUtc);
+            }
+
+            if (t.reservationFinishUtc) {
+               t2.reservationFinish = new Date(t.reservationFinishUtc);
+            }
+
+            telemetry.push(t2);
+         });
+
+         this.deviceTelemetry.set(d.deviceId, telemetry);
+      });
+
+      for (const platform of platforms) {
+         this.platforms.set(platform.id, platform);
+      }
+
+      for (const pool of pools) {
+         if (pool.poolType !== DevicePoolType.Automation) {
+            continue;
+         }
+         this.pools.set(pool.id, pool);
+      }
+
+      for (const device of devices) {
+         this.devices.set(device.id, device);
+      }
+
+      // patch up telemetry data 
+      telemetry.forEach(t => {
+         for (const poolId in t.telemetry) {
+            const values = t.telemetry[poolId];
+            values.forEach(v => {
+
+               if (!v.reserved) {
+                  return;
+               }
+
+               for (const rkey in v.reserved) {
+                  const rd = v.reserved[rkey];
+                  rd.forEach(rdv => {
+
+                     if (!rdv.jobName) {
+                        rdv.jobName = `Unknown Job in stream ${rkey}, ${rdv.jobId}`;
+                     }
+                     if (!rdv.stepName) {
+                        rdv.stepName = `Unknown Step in stream ${rkey}`;
+                     }
+
+                     if (rdv.jobName.indexOf("- Kicked By") !== -1) {
+                        rdv.jobName = rdv.jobName.split("- Kicked By")[0];
+                     }
+                  })
+               }
+            });
+         };
+      });
+
+      this.data = telemetry.map(t => {
+         const date = new Date(t.createTimeUtc);
+         return { date: date, dateString: date.toDateString(), pools: t.telemetry }
+      }).sort((a, b) => a.date.getTime() - b.date.getTime());
+
+      this.setUpdated();
+
+   }
+
+   @action
+   setUpdated() {
+      this.updated++;
+   }
+
+   getData(poolId: string, platformId: string): TelemetryData | undefined {
+
+      if (!this.data) {
+         return undefined;
+      }
+
+      const pool = this.pools.get(poolId);
+      const platform = this.platforms.get(platformId);
+
+      if (!pool || !platform) {
+         return undefined;
+      }
+
+      const telemetry: PlatformTelemetry[] = [];
+      const streamIds = new Set<string>();
+      const deviceIds = new Set<string>();
+
+      const streamNames: Record<string, string> = {};
+      const streamNamesReverse: Record<string, string> = {};
+      const streamSelectors: Record<string, string> = {};
+      const streamSelectorsReverse: Record<string, string> = {};
+
+      this.data.forEach(d => {
+
+         const poolData = d.pools[poolId];
+
+         let platformData: GetDevicePlatformTelemetryResponse | undefined;
+
+         if (poolData) {
+            platformData = poolData.find(pd => pd.platformId === platformId);
+         }
+
+         if (!platformData) {
+            return;
+         }
+
+         const tdata = { ...platformData, date: d.date, dateString: d.dateString } as any;
+
+         tdata["Available"] = platformData.available?.length ?? 0;
+         tdata["Disabled"] = platformData.disabled?.length ?? 0;
+         tdata["Maintenance"] = platformData.maintenance?.length ?? 0;
+         tdata["Problem"] = platformData.problem?.length ?? 0;
+
+         platformData.available?.forEach(d => { deviceIds.add(d); });
+         platformData.maintenance?.forEach(d => { deviceIds.add(d); });
+         platformData.disabled?.forEach(d => { deviceIds.add(d); });
+         platformData.problem?.forEach(d => { deviceIds.add(d); });
+
+         if (platformData.reserved) {
+            for (const streamId in platformData.reserved) {
+               if (!streamIds.has(streamId)) {
+
+                  streamIds.add(streamId);
+
+                  streamNames[streamId] = streamId;
+                  const stream = projectStore.streamById(streamId);
+                  if (stream) {
+                     streamNames[streamId] = stream.fullname ?? streamId;
+                  }
+
+                  const name = streamNames[streamId];
+                  streamNamesReverse[name] = streamId;
+
+                  streamSelectors[name] = name.replaceAll(".", "_").replaceAll("/", "_");
+                  streamSelectorsReverse[streamSelectors[name]] = name;
+               }
+
+               // make stream data accessible by id and fullname
+               tdata[streamId] = platformData.reserved[streamId].length;
+               tdata[streamNames[streamId]] = platformData.reserved[streamId].length;
+
+               platformData.reserved[streamId].forEach(d => { deviceIds.add(d.deviceId); })
+            }
+         }
+
+         telemetry.push(tdata);
+      });
+
+      const deviceIdArray = Array.from(deviceIds).sort((a, b) => a.localeCompare(b));
+      const deviceResponses = deviceIdArray.map(id => this.devices.get(id)).filter(d => !!d);
+      const devices: Record<string, GetDeviceResponse> = {};
+
+      deviceResponses.forEach(d => devices[d!.id] = d!);
+
+      const streamIdArray = Array.from(streamIds).sort((a, b) => a.localeCompare(b))
+
+      for (const streamId of streamIdArray) {
+         telemetry.forEach((t: any) => {
+            if (!t[streamId]) {
+               t[streamId] = 0;
+               t[streamNames[streamId]] = 0;
+            }
+         });
+      }
+
+      return {
+         pool: pool,
+         platform: platform,
+         streamIds: streamIdArray,
+         streamNames: streamNames,
+         streamNamesReverse: streamNamesReverse,
+         streamSelectors: streamSelectors,
+         streamSelectorsReverse: streamSelectorsReverse,
+         telemetry: telemetry,
+         deviceIds: deviceIdArray,
+         devices: devices
+      }
+   }
+
+   get loaded(): boolean { return this.data ? true : false; }
+
+   @observable
+   updated: number = 0;
+
+   private data?: PoolTelemetryData[];
+
+   platforms = new Map<string, GetDevicePlatformResponse>();
+   pools = new Map<string, GetDevicePoolResponse>();
+   devices = new Map<string, GetDeviceResponse>();
+
+   deviceTelemetry = new Map<string, DeviceTelemetry[]>();
+
+}
+
+class PoolTelemetryGraph {
+   constructor(handler: PoolTelemetryHandler) {
+      this.handler = handler;
+      this.margin = { top: 32, right: 28, bottom: 40, left: 48 };
+   }
+
+   set(poolId: string, platformId: string) {
+
+      this.poolId = poolId;
+      this.platformId = platformId;
+      this.streamIds.clear();
+      this.data = undefined;
+
+      const handler = this.handler;
+      this.data = handler.getData(poolId, platformId);
+
+   }
+
+   draw(container: HTMLDivElement, poolId: string, platformId: string, width: number, height: number) {
+
+      if (this.hasRendered && !this.forceRender) {
+         return;
+      }
+
+      this.container = container;
+      this.width = width;
+      this.height = height;
+      const margin = this.margin;
+      this.clear();
+
+      this.set(poolId, platformId);
+
+      const data = this.data;
+
+      if (!data) {
+         return;
+      }
+
+      this.hasRendered = true;
+      this.forceRender = false;
+
+      const xPlot = this.xPlot = this.xPlotMethod();
+
+      const totalDevices = data.deviceIds.length;
+
+      const streams = Array.from(this.streamIds).sort((a, b) => a.localeCompare(b));
+
+      const schemaColor = d3.scaleOrdinal()
+         .domain(streams)
+         .range( dashboard.darktheme ? d3.schemeDark2 : d3.schemeSet2);
+
+      const scolors = dashboard.getStatusColors();
+      const colors: Record<string, string> = {
+         "Available": scolors.get(StatusColor.Success)!,
+         "Problem": scolors.get(StatusColor.Failure)!,
+         "Maintenance": dashboard.darktheme ? "#413F3D" : "#A19F9D",
+         "Disabled": "#715F5D"!
+      };
+
+      const y = d3.scaleLinear()
+         .domain([0, totalDevices])
+         .range([height - margin.top - margin.bottom, 0]);
+
+      const svg = d3.select(container)
+         .append("svg")
+         .attr("width", width)
+         .attr("height", height)
+         .attr("viewBox", [0, 0, width, height] as any)
+         .style("background-color", modeColors.crumbs)
+
+      // define a clipping region
+      svg.append("clipPath")
+         .attr("id", this.clipId)
+         .append("rect")
+         .attr("x", margin.left)
+         .attr("y", margin.top)
+         .attr("width", width - margin.left - margin.right - 240)
+         .attr("height", height - margin.top - margin.bottom);
+
+      const updateChart = (event: any) => {
+      
+         if (this.freezeTooltip || (event && !event.sourceEvent)) {
+            return;
+         }
+
+         const extent = event?.selection
+         this.extent = extent;
+
+         if (!extent) {
+            xPlot!.domain(d3.extent(this.data!.telemetry, function (d: any) { return d.date; }) as any)
+         } else {
+            xPlot!.domain([xPlot!.invert(extent[0]), xPlot!.invert(extent[1])])
+            areaChart.select(".brush").call((brush as any).move, null)
+         }
+
+         this.xAxis!.transition().duration(500).call(this.xAxisMethod as any, xPlot)
+         areaChart
+            .selectAll("path")
+            .transition().duration(500)
+            .attr("d", area as any)
+      }
+
+      // left axis
+      svg.append("g")
+         .attr("transform", `translate(${this.margin.left},${margin.bottom - 8})`)
+         .call(d3.axisLeft(y).ticks(totalDevices))
+
+      // bottom axis      
+      this.xAxis = svg.append("g")
+         .call(this.xAxisMethod as any, xPlot)
+
+      const area = d3.area()
+         .x((d: any) => xPlot!(d.data.date))
+         .y0((d) => y(d[0]) + 32)
+         .y1((d) => y(d[1]) + 32)
+
+      const keys = ["Disabled", "Maintenance", "Problem"];
+
+      data.streamIds.forEach(sid => keys.push(data.streamNames[sid]));
+
+      const stackedData = d3.stack()
+         .keys(keys)
+         (data.telemetry as any)
+
+      const highlight = (e: any, d: any, fromArea?: boolean) => {
+         this.currentTarget = fromArea ? d : undefined;
+         d3.selectAll(".telemetryArea").style("opacity", .2)
+         d3.selectAll("." + (data.streamSelectors[d] ?? d)).style("opacity", 1)
+      }
+
+      const noHighlight = (e: any, d: any) => {
+         this.currentTarget = undefined;
+         d3.selectAll(".telemetryArea").style("opacity", 1)
+      }
+
+      const brush = d3.brushX()
+         .extent([[0, 0], [width - 270, height]])
+         .on("end", updateChart)
+
+      const areaChart = svg.append('g')
+         .attr("clip-path", `url(#${this.clipId})`);
+
+      areaChart
+         .append("g")
+         .attr("class", "brush")
+         .call(brush);
+
+      areaChart
+         .selectAll("telemetrylayers")
+         .data(stackedData)
+         .enter()
+         .append("path")
+         .attr("class", function (d) { return "telemetryArea " + (data.streamSelectors[d.key] ?? d.key); })
+         .style("shapeRendering", "geometricPrecision")
+         .style("fill", function (d) { return colors[d.key] ?? schemaColor(d.key) })
+         .attr("d", area as any)
+         .on("mouseover", (e, d) => { highlight(e, e.target.classList[1], true) })
+         .on("mouseleave", (e, d) => { noHighlight(e, d) });
+
+      // legend
+      const size = 15
+      const posx = 1160;
+      const posy = margin.top;
+      svg.selectAll("legendrect")
+         .data(keys)
+         .enter()
+         .append("rect")
+         .attr("class", function (d) { return "telemetryArea " + (data.streamSelectors[d] ?? d) })
+         .attr("x", posx)
+         .attr("y", function (d, i) { return posy + 10 + i * (size + 5 + 4) })
+         .attr("width", size)
+         .attr("height", size)
+         .style("fill", function (d: any) { return colors[d] ?? schemaColor(d) })
+         .on("mouseover", highlight)
+         .on("mouseleave", noHighlight)
+
+      svg.selectAll("legendlabels")
+         .data(keys)
+         .enter()
+         .append("text")
+         .attr("x", posx + size * 1.2 + 8)
+         .attr("y", function (d, i) { return posy + 10 + i * (size + 5 + 4) + (size / 2) })
+         .style("fill", function (d: any) { return modeColors.text })
+         .text(function (d) { return d })
+         .attr("text-anchor", "left")
+         .attr("class", function (d) { return "telemetryArea " + (data.streamSelectors[d] ?? d) })
+         .style("alignment-baseline", "middle")
+         .style("font-family", "Horde Open Sans Regular")
+         .style("font-size", "13px")
+         .on("mouseover", highlight)
+         .on("mouseleave", noHighlight)
+
+      // tooltip
+      this.tooltip = d3.select(container)
+         .append("div")
+         .style("opacity", 0)
+         .style("background-color", modeColors.background)
+         .style("border", "solid")
+         .style("border-width", "1px")
+         .style("border-radius", "3px")
+         .style("border-color", dashboard.darktheme ? "#413F3D" : "#2D3F5F")
+         .style("padding", "8px")
+         .style("position", "absolute")
+      //.style("pointer-events", "none")
+
+      svg.on("mousemove", (event) => this.handleMouseMove(event))
+      svg.on("wheel", (event) => {
+         if (event.wheelDelta < 0) {
+            updateChart(undefined);
+         }
+      })
+      // dblclick for double click events
+      svg.on("click", (event) => {
+
+         if (this.freezeTooltip) {
+            this.freezeTooltip = false;
+            this.updateTooltip(false);
+            return;
+         }
+
+         if (this.currentTarget) {
+            this.freezeTooltip = true;
+         }
+
+      });
+
+
+   }
+
+   private clear() {
+      d3.selectAll("#pool_graph_container > *").remove();
+   }
+
+   private closestData(x: number): PlatformTelemetry | undefined {
+
+      const xPlot = this.xPlot;
+      if (!xPlot) {
+         return;
+      }
+
+      const data = this.data!.telemetry;
+
+      let closest = data.reduce((best, value, i) => {
+
+         let absx = Math.abs(xPlot(value.date) - x)
+         if (absx < best.value) {
+            return { index: i, value: absx };
+         }
+         else {
+            return best;
+         }
+      }, { index: 0, value: Number.MAX_SAFE_INTEGER });
+
+      return data[closest.index];
+   }
+
+   private updateTooltip(show: boolean, x?: number, y?: number, html?: string) {
+      if (!this.tooltip || this.freezeTooltip) {
+         return;
+      }
+
+      x = x ?? 0;
+      y = y ?? 0;
+
+      this.tooltip
+         .html(html ?? "")
+         .style("left", (x + 40) + "px")
+         .style("top", y + "px")
+         .style("font-family", "Horde Open Sans Regular")
+         .style("font-size", "10px")
+         .style("line-height", "16px")
+         .style("shapeRendering", "crispEdges")
+         .style("stroke", "none")
+         .style("opacity", show ? 1 : 0);
+   }
+
+   private handleMouseMove(event: any) {
+      const target = this.currentTarget;
+      if (!target) {
+         this.updateTooltip(false);
+         return;
+      }
+
+      const data = this.data!;
+
+      const streamId = data.streamIds.find(sid => {
+         const streamName = data.streamSelectorsReverse[target!];
+         return data.streamNamesReverse[streamName] === sid;
+      });
+
+      const closest = this.closestData(d3.pointer(event)[0]);
+      if (!closest) {
+         this.updateTooltip(false);
+         return undefined;
+      }
+
+      const displayTime = moment(closest.date);
+      const format = dashboard.display24HourClock ? "HH:mm:ss z" : "LT z";
+
+      let displayTimeStr = displayTime.format('MMM Do YYYY') + ` - ${displayTime.format(format)}`;
+
+      let targetName = "???";
+      if (target) {
+         targetName = data.streamSelectorsReverse[target];
+         if (!targetName) {
+            targetName = target;
+         }
+      }
+
+      let text = `<div style="padding-left:4px;padding-right:12px;padding-bottom:12px">`;
+
+      text += `<div style="padding-bottom:8px;padding-top:2px"><span style="font-family:Horde Open Sans Semibold;font-size:11px">${targetName}</span></div>`;
+
+      if (streamId && closest.reserved && closest.reserved[streamId]) {
+         closest.reserved[streamId].forEach(v => {
+            const d = data.devices[v.deviceId]?.name ?? v.deviceId;
+            text += `<div style="padding-bottom:8px">`;
+            const href = `/job/${v.jobId}?step=${v.stepId}`;
+            text += `<a href="${href}" target="_blank"><span style="font-family:Horde Open Sans Semibold;font-size:11px;padding-left:8px">${d}: ${v.stepName}</span></a><br/>`
+            if (v.jobId && v.stepId) {
+               const metrics = this.handler.getStepMetrics(v.jobId, v.stepId);
+               if (metrics && (metrics.duration || metrics.numKits > 1 || metrics.numProblems > 0)) {
+
+                  text += `<div style="padding-left: 16px">`;
+
+                  if (metrics.numKits > 1) {
+                     text += `<span style="">Devices:</span> ${metrics.numKits} <br/>`;
+                  }
+
+                  if (metrics.numProblems > 0) {
+                     text += `<span style="font-family:Horde Open Sans Semibold;color:#FF0000">Problems: ${metrics.numProblems}</span> <br/>`;
+                  }
+
+                  if (metrics.duration) {
+                     text += `<span style="">Duration:</span> ${metrics.duration} <br/>`;
+                  }
+
+                  text += `</div>`;
+               }
+
+               text += `</div>`;
+            }
+         });
+      }
+      else if (target === "Maintenance" && closest.maintenance) {
+         closest.maintenance.forEach(d => {
+            d = data.devices[d]?.name ?? d;
+            text += `<span style="font-family:Horde Open Sans Semibold;font-size:11px;padding-left:8px">${d}</span><br/>`
+         });
+         text += `<br/>`;
+      }
+      else if (target === "Problem" && closest.problem?.length) {
+
+         const problems = this.handler.getProblemSteps(closest.problem, closest.date);
+
+         closest.problem.forEach(d => {
+            const problem = problems[d];
+            d = data.devices[d]?.name ?? d;
+
+            text += `<span style="font-family:Horde Open Sans Semibold;font-size:11px;padding-left:8px">${d}</span><br/>`
+
+            if (problem && problem.streamId && problem.jobName) {
+
+               const name = problem.stepName ?? problem.jobName;
+               let desc = `${data.streamNames[problem.streamId]} - ${name}`;
+
+               const href = `/job/${problem.jobId}?step=${problem.stepId}`;
+
+               text += `<a href="${href}" target="_blank"><div style="padding-bottom:8px;padding-left:8px"><span style="">${desc}</span></div></a>`
+
+            }
+
+         });
+      }
+      else if (target === "Disabled" && closest.disabled) {
+         closest.disabled.forEach(d => {
+            d = data.devices[d]?.name ?? d;
+            text += `<span style="font-family:Horde Open Sans Semibold;font-size:11px;padding-left:8px">${d}</span><br/>`
+         });
+
+         text += `<br/>`;
+      }
+
+      text += `${displayTimeStr}<br/>`;
+
+      text += "</div>";
+
+      this.updateTooltip(true, d3.pointer(event)[0], d3.pointer(event)[1], text);
+
+   }
+
+   private get xAxisMethod() {
+      return (g: SelectionType, x: ScalarTime) => {
+
+         const data = this.data!;
+
+         const xPlot = this.xPlot;
+
+         const tickValues: { x: number, date: Date }[] = [];
+         const d = x.domain();
+
+         if (d.length === 2) {
+
+            const start = d[0];
+            const end = d[1];
+            const tickSet: Set<string> = new Set();
+
+            data.telemetry.forEach(d => {
+               if (d.date < start || d.date > end) {
+                  return;
+               }
+               const v = d.dateString;
+               if (tickSet.has(v)) {
+                  return;
+               }
+               tickSet.add(v);
+               const date = new Date(v);
+               const x = xPlot!(date);
+               if (true) { //x >= 0 ) {
+                  tickValues.push({ x: x, date: date });
+               }
+            });
+         }
+
+         //const min = d[0];
+         //const max = d[1];
+         //const span = moment.duration(moment(max).diff(moment(min))).asDays();
+         if (true) {//span <= 6) {
+
+            const tickSet: Set<string> = new Set();
+
+            g.attr("transform", `translate(0,${this.height - this.margin.bottom})`)
+               .style("font-family", "Horde Open Sans Regular")
+               .style("font-size", "11px")
+               .call(d3.axisBottom(x)
+                  .ticks(8)
+                  .tickFormat(d => {
+
+                     const v = (d as Date).toDateString();
+                     if (!tickSet.has(v)) {
+                        tickSet.add(v)
+                        return getHumanTime(d as Date)
+                     }
+
+                     return (d as Date).toLocaleString(undefined, { hour: 'numeric', minute: "numeric", hour12: true })
+                  })
+                  .tickSizeOuter(0))
+
+            return;
+
+         }
+
+      }
+
+
+
+   }
+
+   private xPlot?: ScalarTime;
+
+   private xPlotMethod() {
+
+      const data = this.data!.telemetry;
+
+      return d3.scaleTime()
+         .domain(d3.extent(data, (d => d.date)) as any)
+         .range([this.margin.left, this.width - this.margin.right - 240])
+   }
+
+   handler: PoolTelemetryHandler;
+
+   container?: HTMLDivElement;
+   private margin: { top: number, right: number, bottom: number, left: number }
+   private width = 0;
+   private height = 0;
+
+   private data?: TelemetryData;
+   private poolId = "";
+   private platformId = "";
+   private streamIds = new Set<string>();
+
+   private freezeTooltip = false;
+
+   private xAxis?: SelectionType;
+   private tooltip?: DivSelectionType;
+   private readonly clipId = "pool_history_clip_path";
+
+   private currentTarget?: string;
+
+   private hasRendered = false;
+   forceRender = false;
+
+   extent?: any;
+}
+
+const SummaryModal: React.FC<{ handler: PoolTelemetryHandler, onClose: () => void }> = ({ handler, onClose }) => {
+
+   const platforms = Array.from(handler.platforms.values()).sort((a, b) => a.name.localeCompare(b.name));
+
+   const problems = handler.getProblemDevices();
+
+   const renderPlatform = (platform: GetDevicePlatformResponse) => {
+
+      const psteps = handler.getSteps(platform.id);
+      let platformProblems = problems.get(platform.id) ?? [];
+      platformProblems = platformProblems.filter(p => p.problems > 8);
+
+      const problemColumns = [
+         { key: 'column1', name: 'Problem Devices', fieldName: 'deviceName', minWidth: 150, maxWidth: 150 },
+         { key: 'column2', name: 'Issues', fieldName: 'problems', minWidth: 120, maxWidth: 120 },
+         { key: 'column3', name: 'Latest Issue', fieldName: 'latestProblem', minWidth: 120, maxWidth: 120 },
+      ];
+
+      const steps: Map<string, { stepName: string, reservations: number, problems: number, durationMS: number }> = new Map();
+
+      psteps.forEach(s => {
+         if (!steps.has(s.stepName)) {
+            steps.set(s.stepName, { stepName: s.stepName, reservations: s.reservations, problems: s.problems, durationMS: s.durationMS })
+         } else {
+            const step = steps.get(s.stepName)!;
+            step.reservations += s.reservations;
+            step.problems += s.problems;
+            step.durationMS += s.durationMS;
+         }
+      });
+
+      const stepItems = Array.from(steps.values()).sort((a, b) => b.durationMS - a.durationMS).slice(0, 10).map(s => {
+         return { stepName: s.stepName, reservations: s.reservations, problems: s.problems, duration: msecToElapsed(s.durationMS, false, false) }
+      });
+
+      const stepColumns = [
+         { key: 'column1', name: 'Top Consumers', fieldName: 'stepName', minWidth: 320, maxWidth: 320 },
+         { key: 'column2', name: 'Duration', fieldName: 'duration', minWidth: 100, maxWidth: 100 },
+         { key: 'column3', name: 'Reservations', fieldName: 'reservations', minWidth: 100, maxWidth: 100 },
+         { key: 'column4', name: 'Problems', fieldName: 'problems', minWidth: 100, maxWidth: 100 },
+      ];
+
+      return <Stack style={{ paddingLeft: 12, paddingRight: 32 }}>
+         <Stack>
+            <Stack>
+               <Text variant="medium" styles={{ root: { fontWeight: "unset", fontFamily: "Horde Open Sans SemiBold" } }}>{platform.name}</Text>
+            </Stack>
+            <Stack tokens={{ childrenGap: 12 }}>
+               {!!platformProblems?.length &&
+                  <Stack style={{ paddingLeft: 12, paddingTop: 12, width: 480 }}>
+                     <DetailsList
+                        isHeaderVisible={true}
+                        items={platformProblems}
+                        columns={problemColumns}
+                        layoutMode={DetailsListLayoutMode.justified}
+                        compact={true}
+                        selectionMode={SelectionMode.none}
+                     />
+                  </Stack>}
+               {!!stepItems?.length &&
+                  <Stack style={{ paddingLeft: 12, paddingTop: 12, width: 800 }} >
+                     <DetailsList
+                        isHeaderVisible={true}
+                        items={stepItems}
+                        columns={stepColumns}
+                        layoutMode={DetailsListLayoutMode.justified}
+                        compact={true}
+                        selectionMode={SelectionMode.none}
+                     />
+                  </Stack>}
+
+            </Stack>
+         </Stack>
+      </Stack>
+   };
+
+   return <Modal className={hordeClasses.modal} isOpen={true} isBlocking={true} topOffsetFixed={true} styles={{ main: { padding: 8, width: 940, backgroundColor: dashboard.darktheme ? modeColors.content : modeColors.background, hasBeenOpened: false, top: "72px", position: "absolute", height: "820px" } }} onDismiss={() => { onClose() }}>
+
+      <Stack>
+
+         <Stack horizontal verticalAlign="center" style={{ paddingTop: 12 }}>
+            <Stack style={{ paddingLeft: 16 }}>
+               <Text variant="mediumPlus" styles={{ root: { fontWeight: "unset", fontFamily: "Horde Open Sans SemiBold" } }}>Device Platform Summary</Text>
+            </Stack>
+            <Stack grow />
+            <Stack horizontalAlign="end" style={{ paddingRight: 12 }}>
+               <IconButton
+                  iconProps={{ iconName: 'Cancel' }}
+                  ariaLabel="Close popup modal"
+                  onClick={() => { onClose() }}
+               />
+            </Stack>
+         </Stack>
+         <Stack style={{ paddingTop: 12, paddingRight: 8 }}>
+            <FocusZone direction={FocusZoneDirection.vertical}>
+               <div style={{ position: 'relative', height: "720px" }} data-is-scrollable>
+                  <ScrollablePane scrollbarVisibility={ScrollbarVisibility.always} onScroll={() => { }}>
+
+                     <Stack style={{ paddingLeft: 24, paddingTop: 12 }} tokens={{ childrenGap: 24 }}>
+                        {platforms.map(p => renderPlatform(p))}
+                     </Stack>
+                  </ScrollablePane>
+               </div>
+            </FocusZone>
+         </Stack>
+      </Stack>
+   </Modal>
+
+}
+
+export const DevicePoolGraph: React.FC = observer(() => {
+
+   const [container, setContainer] = useState<HTMLDivElement | null>(null);
+   const [showSummary, setShowSummary] = useState(false);
+   const [state, setState] = useState<{ handler?: PoolTelemetryHandler, graph?: PoolTelemetryGraph, pool?: GetDevicePoolResponse, platform?: GetDevicePlatformResponse }>({});
+
+   let handler = state.handler;
+
+   if (!handler) {
+      handler = new PoolTelemetryHandler();
+      setState({ handler: handler, graph: new PoolTelemetryGraph(handler) });
+      return null;
+   }
+
+   // subscribe
+   if (handler.updated) { }
+
+   if (!handler.loaded) {
+      return <Stack horizontalAlign="center" tokens={{ childrenGap: 32 }}>
+         <Text variant="mediumPlus" styles={{ root: { fontWeight: "unset", fontFamily: "Horde Open Sans SemiBold" } }}>Loading Device Pool Telemetry</Text>
+         <Spinner size={SpinnerSize.large} />
+      </Stack>
+   }
+
+   const pools = Array.from(handler.pools.values()).sort((a, b) => a.name.localeCompare(b.name));
+   const platforms = Array.from(handler.platforms.values()).sort((a, b) => a.name.localeCompare(b.name));
+
+   if (!state.pool || !state.platform) {
+      if (!pools.length) {
+         return <Stack>No Pools</Stack>
+      }
+      if (!platforms.length) {
+         return <Stack>No Platforms</Stack>
+      }
+
+
+      const pool = pools[0];
+      const platform = platforms[0];
+
+      setState({ ...state, pool: pool, platform: platform })
+      return null;
+
+   }
+
+   if (container && state.pool && state.platform) {
+      state.graph?.draw(container, state.pool.id, state.platform.id, 1400, 600);
+   }
+
+   const poolOptions: IDropdownOption[] = pools.map(p => { return { key: `pool_key_${p.id}`, text: p.name, data: p } });
+   const platformOptions: IDropdownOption[] = platforms.map(p => { return { key: `platform_key_${p.id}`, text: p.name, data: p } });
+
+   return <Stack>
+      {showSummary && !!handler && <SummaryModal handler={handler} onClose={() => setShowSummary(false)} />}
+      <Stack style={{ position: "relative" }}>
+         <Stack horizontal verticalAlign="start" style={{ width: 1412, paddingBottom: 18 }} tokens={{ childrenGap: 32 }}>
+            <Stack style={{ paddingLeft: 4 }}>
+               <Text variant="mediumPlus" styles={{ root: { fontWeight: "unset", fontFamily: "Horde Open Sans SemiBold" } }}>Device Pool Telemetry</Text>
+            </Stack>
+
+            <Stack grow />
+            <DefaultButton text="Summary" onClick={() => setShowSummary(true)} />
+            <Dropdown
+               defaultSelectedKey={`pool_key_${state.pool!.id}`}
+               options={poolOptions}
+               onChange={(event, option) => {
+                  if (state.graph) {
+                     state.graph.forceRender = true;
+                  }
+                  setState({ ...state, pool: option!.data, platform: platforms[0] });
+               }}
+               style={{ width: 120 }}
+            />
+            <Dropdown
+               defaultSelectedKey={`platform_key_${state.platform!.id}`}
+               options={platformOptions}
+               onChange={(event, option) => {
+                  if (state.graph) {
+                     state.graph.forceRender = true;
+                  }
+                  setState({ ...state, platform: option!.data, })
+               }}
+               style={{ width: 120 }}
+            />
+         </Stack>
+
+         <Stack style={{ paddingLeft: 12 }}>
+            <div id="pool_graph_container" className="horde-no-darktheme" style={{ shapeRendering: "crispEdges", userSelect: "none" }} ref={(ref: HTMLDivElement) => setContainer(ref)} onMouseEnter={() => { }} onMouseLeave={() => { }} />
+         </Stack>
+
+      </Stack>
+   </Stack>
+});
+
+export const DevicePoolTelemetryModal: React.FC<{ onClose?: () => void }> = ({ onClose }) => {
+
+
+   return <Stack>
+      <Modal className={hordeClasses.modal} isOpen={true} isBlocking={true} topOffsetFixed={true} styles={{ main: { padding: 8, width: 1520, backgroundColor: modeColors.background, hasBeenOpened: false, top: "64px", position: "absolute", height: "760px" } }} onDismiss={() => { if (onClose) { onClose() } }}>
+         <Stack>
+            <Stack horizontal>
+               <Stack grow />
+               <Stack horizontalAlign="end" style={{ paddingTop: 8, paddingRight: 12 }}>
+                  <IconButton
+                     iconProps={{ iconName: 'Cancel' }}
+                     ariaLabel="Close popup modal"
+                     onClick={() => { if (onClose) onClose() }}
+                  />
+               </Stack>
+            </Stack>
+            <Stack style={{ paddingLeft: 24 }}>
+               <DevicePoolGraph />
+            </Stack>
+         </Stack>
+      </Modal>
+   </Stack>
+
+}
+
+
+
diff --git a/Engine/Source/Programs/Horde/HordeDashboard/src/components/DeviceView.tsx b/Engine/Source/Programs/Horde/HordeDashboard/src/components/DeviceView.tsx
--- a/Engine/Source/Programs/Horde/HordeDashboard/src/components/DeviceView.tsx
+++ b/Engine/Source/Programs/Horde/HordeDashboard/src/components/DeviceView.tsx
@@ -1,4 +1,4 @@
-import { DetailsList, DetailsListLayoutMode, DetailsRow, Dropdown, FocusZone, FocusZoneDirection, getTheme, IColumn, Icon, IconButton, IDetailsListProps, IGroup, Label, Link as FluentLink, mergeStyleSets, Modal, Pivot, PivotItem, PrimaryButton, ScrollablePane, ScrollbarVisibility, SelectionMode, Spinner, SpinnerSize, Stack, Text } from "@fluentui/react";
+import { DefaultButton, DetailsList, DetailsListLayoutMode, DetailsRow, Dropdown, FocusZone, FocusZoneDirection, getTheme, IColumn, Icon, IconButton, IDetailsListProps, IGroup, Label, Link as FluentLink, mergeStyleSets, Modal, Pivot, PivotItem, PrimaryButton, ScrollablePane, ScrollbarVisibility, SelectionMode, Spinner, SpinnerSize, Stack, Text } from "@fluentui/react";
 import { observer } from "mobx-react-lite";
 import moment from "moment";
 import React, { useEffect, useState } from "react";
@@ -12,6 +12,7 @@ import { hordeClasses, modeColors } from "../styles/Styles";
 import { Breadcrumbs } from "./Breadcrumbs";
 import { DeviceEditor, DeviceHandler, DeviceStatus } from "./DeviceEditor";
 import { DeviceInfoModal } from "./DeviceInfoView";
+import { DevicePoolTelemetryModal } from "./DevicePoolTelemetry";
 import { TopNav } from "./TopNav";
 
 const handler = new DeviceHandler();
@@ -52,7 +53,7 @@ const dropDownStyle: any = () => {
             }
          }
       },
-      dropdownItemHeader: { fontSize: 12, color: modeColors.text},
+      dropdownItemHeader: { fontSize: 12, color: modeColors.text },
       dropdownOptionText: { fontSize: 12 },
       dropdownItem: {
          minHeight: 28, lineHeight: 28
@@ -103,6 +104,7 @@ const DevicePanel: React.FC = observer(() => {
    const [editState, setEditState] = useState<{ shown?: boolean, infoShown?: boolean, editNote?: boolean, device?: DeviceItem | undefined }>({});
    const [pivotState, setPivotState] = useState<{ key: string, poolFilter: Set<string> }>({ key: pivotKeyShared, poolFilter: new Set() });
    const [checkoutState, setCheckoutState] = useState<{ checkoutId?: string, checkinId?: string, showConfirm?: "in" | "out" | "error" }>({});
+   const [telemState, setTelemState] = useState(false);
 
    useEffect(() => {
 
@@ -480,6 +482,7 @@ const DevicePanel: React.FC = observer(() => {
    const checkedOut = handler.getUserDeviceCheckouts(dashboard.userId);
 
    return (<Stack>
+      {!!telemState && <DevicePoolTelemetryModal onClose={() => setTelemState(false)} />}
       {(!!checkoutState.checkinId || !!checkoutState.checkoutId) && <InfoModal />}
       {!!checkoutState.showConfirm && checkedOut.length > 0 && <CheckoutConfirmModal check={checkoutState.showConfirm} devices={checkedOut} onClose={() => setCheckoutState({})} />}
       {editState.infoShown && <DeviceInfoModal handler={handler} deviceIn={editState.device?.device} onEdit={(device) => { setEditState({ infoShown: true, shown: true, device: { device: device } }) }} onClose={() => { setEditState({}) }} />}
@@ -487,7 +490,7 @@ const DevicePanel: React.FC = observer(() => {
       <Stack styles={{ root: { paddingLeft: 12, paddingRight: 12, width: "100%" } }} >
          <Stack>
 
-            <Stack horizontal>
+            <Stack horizontal verticalAlign="center" style={{ paddingBottom: 12 }} tokens={{ childrenGap: 32 }}>
                <Stack>
                   <Pivot className={hordeClasses.pivot}
                      selectedKey={pivotState.key}
@@ -499,7 +502,7 @@ const DevicePanel: React.FC = observer(() => {
                      {pivotItems}
                   </Pivot>
                </Stack>
-               {poolItems.length > 1 && <Stack style={{ paddingTop: 8, paddingLeft: 48 }}>
+               {poolItems.length > 1 && <Stack style={{ paddingLeft: 48 }}>
                   <Dropdown
                      placeholder="Filter Pools"
                      style={{ width: 200 }}
@@ -524,6 +527,10 @@ const DevicePanel: React.FC = observer(() => {
                   />
                </Stack>}
 
+               {automationTab && <Stack>
+                  <DefaultButton text="Pool Telemetry" onClick={() => { setTelemState(true) }} />
+               </Stack>}
+
                <Stack grow />
                <Stack>
                   <PrimaryButton text="Add Device"
@@ -543,7 +550,7 @@ const DevicePanel: React.FC = observer(() => {
                <div className={customStyles.details} style={{ height: "calc(100vh - 270px)", position: 'relative' }} data-is-scrollable>
                   <ScrollablePane scrollbarVisibility={ScrollbarVisibility.always} onScroll={() => { }}>
                      <DetailsList
-                        styles={{ root: { overflowX: "hidden" } }}
+                        styles={{ root: { overflowX: "hidden", width: 1676 } }}
                         items={devices}
                         groups={groups}
                         columns={columns}
@@ -657,7 +664,7 @@ export const CheckoutConfirmModal: React.FC<{ check: "in" | "out" | "error", dev
 
 
 export const DeviceView: React.FC = () => {
-
+   
    const windowSize = useWindowSize();
    const vw = Math.max(document.documentElement.clientWidth, window.innerWidth || 0);
 
@@ -665,11 +672,11 @@ export const DeviceView: React.FC = () => {
       <TopNav />
       <Breadcrumbs items={[{ text: 'Devices' }]} />
       <Stack horizontal>
-         <div key={`windowsize_streamview_${windowSize.width}_${windowSize.height}`} style={{ width: vw / 2 - 900, flexShrink: 0, backgroundColor: 'rgb(250, 249, 249)' }} />
-         <Stack tokens={{ childrenGap: 0 }} styles={{ root: { backgroundColor: 'rgb(250, 249, 249)', width: "100%" } }}>
+         <div key={`windowsize_streamview_${windowSize.width}_${windowSize.height}`} style={{ width: vw / 2 - 900, flexShrink: 0, backgroundColor: modeColors.background }} />
+         <Stack tokens={{ childrenGap: 0 }} styles={{ root: { backgroundColor: modeColors.background, width: "100%" } }}>
             <Stack style={{ maxWidth: 1800, paddingTop: 6, marginLeft: 4, height: 'calc(100vh - 8px)' }}>
                <Stack horizontal className={hordeClasses.raised}>
-                  <Stack style={{ width: "100%", height: 'calc(100vh - 228px)' }} tokens={{ childrenGap: 18 }}>
+                  <Stack style={{ width: "100%", height: 'calc(100vh - 228px)'}} tokens={{ childrenGap: 18 }}>
                      <DevicePanel />
                   </Stack>
                </Stack>
@@ -677,6 +684,7 @@ export const DeviceView: React.FC = () => {
          </Stack>
       </Stack>
    </Stack>
+
 };
 
 
