Horde: Include p4 server name in clients created by Horde, to allow using edge servers.

diff --git a/Engine/Source/Programs/Horde/Horde.Build/Perforce/PerforceService.cs b/Engine/Source/Programs/Horde/Horde.Build/Perforce/PerforceService.cs
--- a/Engine/Source/Programs/Horde/Horde.Build/Perforce/PerforceService.cs
+++ b/Engine/Source/Programs/Horde/Horde.Build/Perforce/PerforceService.cs
@@ -520,9 +520,9 @@ namespace Horde.Build.Perforce
 
 		}
 
-		static string GetClientName(string? serviceUserName, string stream, bool readOnly, bool createChange, string? username = null)
+		static string GetClientName(string? serviceUserName, string serverName, string stream, bool readOnly, bool createChange, string? username = null)
 		{
-			string clientName = $"zzt-horde-p4bridge-{Dns.GetHostName()}-{serviceUserName ?? "default"}-{stream.Replace("/", "+", StringComparison.OrdinalIgnoreCase)}";
+			string clientName = $"zzt-horde-p4bridge-{serverName}-{Dns.GetHostName()}-{serviceUserName ?? "default"}-{stream.Replace("/", "+", StringComparison.OrdinalIgnoreCase)}";
 
 			if (!readOnly)
 			{
@@ -598,7 +598,7 @@ namespace Horde.Build.Perforce
 					throw new Exception("Stream required for client");
 
 				}
-				string clientName = GetClientName(serviceUserName, stream, readOnly, createChange, username);
+				string clientName = GetClientName(serviceUserName, repository.Server.Metadata.Name, stream, readOnly, createChange, username);
 
 				IList<P4.Client>? clients = repository.GetClients(new P4.ClientsCmdOptions(P4.ClientsCmdFlags.None, username, clientName, 1, stream));
 
