Horde Storage - Updated documentation to better explain what horde storage, how to run it and provide more examples.

diff --git a/Engine/Source/Programs/Horde/HordeStorage/Helm/Horde.Storage/README.md b/Engine/Source/Programs/Horde/HordeStorage/Helm/Horde.Storage/README.md
--- a/Engine/Source/Programs/Horde/HordeStorage/Helm/Horde.Storage/README.md
+++ b/Engine/Source/Programs/Horde/HordeStorage/Helm/Horde.Storage/README.md
@@ -4,25 +4,70 @@ Horde.Storage - A Distributed storage service for use with Unreal Engine
 
 # Components 
 
-* Horde.Storage - Main service - The service implementing most of the logic
-* Callisto - Transaction - A C# service that manages and serializes the ref transactions from the ref server. Essentially a custom db for Horde.Storage focusing on a log structured approch to the transaction rather then the traditional db aspects of storage. This is a legacy component that is being faced out.
-* Jupiter.Common - Lib - Library with common shared infrastructure that is reused across all or some C# services
-
+* Horde.Storage - The service implementing most of the logic
 
 # Chart Values
 
+We recommend using a base value file for your general settings and then having a secondary values file for per-region settings. See example-values-base.yaml and example-values-base-A.yaml for a example of such a setup.
+
 ## Globals
 
 | Name | Description | Possible Values |
 | ------------- |:-------------:| -----:|
 | cloudProvider | Which cloud provider this will be deployed on, impacts which services we interact with. | AWS or Azure |
 | awsRegion | Only used for AWS cloud, which region you expect the services to be in | AWS Region format |
-| awsRoleArn | The arn of a AWS role to assume | A valid ARN | 
+| awsRole | How to assume a IAM role in AWS. AssumeRoleWebIdentity assumes a kubernetes service account with a mapped IAM role. AssumeRole will explicitly run a assume role call on the arn. | AssumeRoleWebIdentity / AssumeRole / Basic |
+| awsRoleArn | The arn of a AWS role to assume if using AssumeRole for awsRole | A valid ARN | 
 | siteName | The name of this particular site within your cluster of Jupiter nodes. Should be unique amongst all your jupiter sites and is used to understand were a record has been replicated | Any globally unique value |
-| authMethod | Which method is used to auth | Okta / JWTBearer / Disabled |
-| oktaDomain | The okta domain used to auth against | See your okta setup |
-| oktaAuthorizationServerId | The authorization server id in okta | See your okta setup for which values are valid |
-| jwtAudience | The audience expected in your JWT token, used both with Okta and JWTBearer | The JWT audience
-| jwtAuthority | The address to your JWT Authority, only used for JWT Bearer verification | A url to your JWT authority |
 
 
+## Resolving secrets
+
+Horde Storage supports resolving secrets for some configuration fields. The fields are:
+* ServiceCredentials.OAuthClientId
+* ServiceCredentials.OAuthClientSecret
+* azureBlobStore.ConnectionString
+* Scylla.ConnectionString
+
+A secret can be just the raw text string if you want to (not recommended as this puts the secret in clear text in your kubernetes objects and in logs as well as the values files) instead we support using cloud specific formats for resolving the secrets.
+
+### Resolving secret in AWS
+
+The expected format for the secret string is:
+`aws!<arn>|<secret-value>`
+
+A example could look a bit like this:
+`aws!arn:aws:secretsmanager:region-A:<aws-account-number>:secret:hordestorage-oktaaccount|client_secret`
+
+### Resolving secret in Azure
+
+The expected format for the secret string is:
+`akv!<vault-name>|<secret-name>`
+
+A example could look a bit like this:
+`akv!horde-storage-keyvault|client_secret`
+
+
+## Authentication setup
+See the Authentication section in readme.md under the Horde Storage root for a example of how to setup authentication with namespaces
+
+
+## Scylla configuration
+
+Setting up scylla is documented in the readme.md in the Horde Storage root.
+
+Scylla will require a connection stringof where to connect to. Notice that scylla will find the layout of the cluster itself so you only need to point it to one node and it will find the best nodes to query. Example: `Contact Points=your-scylla-dns.your-domain.com`
+
+Each region also needs to setup some configuration of the local table (we keep a table that is not replicated within the cluster for data that is unique to each region):
+
+```
+    LocalKeyspaceReplicationStrategy:
+      class : "NetworkTopologyStrategy"
+      region-A : 2
+      region-B: 0
+    LocalDatacenterName: "region-ABC"
+    LocalKeyspaceSuffix: "abc"
+```
+The local keyspace suffix is appened to the keyspace (database) name to make it globally unique.
+The local datacenter name is used to make sure we attempt to query the local data center before hitting remote nodes and should be the name of the DC in scylla (usually the same name as the region name from your cloud provider) 
+LocalKeyspaceReplicationStrategy should always be set to `NetworkTopologyStrategy` class and then you need to enumerate all the regions that exists - setting them to 0 except for the region which you are creating the configuration for (that should be set to 2).
diff --git a/Engine/Source/Programs/Horde/HordeStorage/README.md b/Engine/Source/Programs/Horde/HordeStorage/README.md
--- a/Engine/Source/Programs/Horde/HordeStorage/README.md
+++ b/Engine/Source/Programs/Horde/HordeStorage/README.md
@@ -1,24 +1,55 @@
-# Horde.Storage
-
-Horde.Storage - A Distributed storage service
+# Introduction
+
+Horde Storage - A Distributed storage service primarily used as a Cloud based DDC cache.
+
+Horde Storage stores Compact Binary objects in a Content Addressable store that can be replicated across the world. Compact binary objects are self describing key value objects (think json) that has good binary serialization and a extensive type system. One of the key types are the Attachment types which lets you have a object that references a large blob without containing it. These attached payloads are content addressed (meaning the hash of the payload is used as the identifier of it), content addressing is a common practise in distributed storage systems most notiably in git. Content addressing allows us to quickly generate a immutable identifier of a payload to quickly determine if its already available in a cache and to determine if we need to replicate this somewhere else.
+
+Horde Storage can signficantly help teams speed up their cook processes in the Unreal Engine for cases were they do not have access to a primied local (or fileshare) cache, specifically for cases when users are remote. Its a fairly complicated software to setup and operate so we do not think this is a great fit for every team that uses Unreal, but if you are geographically distributed and have engineers with cloud experience then this can be very helpful.
+
+# Table of contents
+- [Introduction](#introduction)
+- [Table of contents](#table-of-contents)
+- [Directories](#directories)
+- [Dependencies](#dependencies)
+- [Other useful things](#other-useful-things)
+- [Functional Test requirements](#functional-test-requirements)
+- [Running locally](#running-locally)
+- [Deployment](#deployment)
+  - [Scylla](#scylla)
+  - [AWS](#aws)
+  - [On premise](#on-premise)
+  - [Azure](#azure)
+  - [Testing your deployment](#testing-your-deployment)
+- [Monitoring](#monitoring)
+  - [Health Checks](#health-checks)
+- [Authentication](#authentication)
+  - [IdentityProvider setup](#identityprovider-setup)
+  - [Namespace access](#namespace-access)
+- [Networking setup](#networking-setup)
+  - [Public Port](#public-port)
+  - [Private port](#private-port)
+  - [Internal Port](#internal-port)
+- [Common operations](#common-operations)
+  - [Add new region](#add-new-region)
+  - [Blob replication setup](#blob-replication-setup)
 
 # Directories 
 
 * Horde.Storage - Most of Horde.Storage functionality.
-* Callisto - Transaction - A legacy C# service that manages and serializes the references inserted into Horde.Storage. Essentially a custom db for Horde.Storage focusing on a log structured approch to the transaction rather then the traditional db aspects of storage.
-* Jupiter.Common - Lib - Library with common shared infrastructure that is reused across Horde.Storage and Callisto. 
+* Jupiter.Common - Lib - Legacy library with functionality that could be used outside of Horde Storage. Will eventually get moved into Horde.Storage or one of our Shared libraries.
 * Helm - Helm charters for all components
-* Benchmark - Useful templates for doing simple http benchmarking user SuperBenchmark (sb)
+* Benchmark - Useful templates for doing simple http benchmarking user SuperBenchmark (sb) as well as a docker container that can be used to run benchmarks using vegeta.
 * Composes - a bunch of different docker compose files for easier local deployment and testing
 
 # Dependencies
 
-* DotNet Core 5 (and Visual Studio 2019 or VS Code)
+* DotNet Core 6 (and Visual Studio 2022 or VS Code)
 * Docker
+* Scylla
+* Blob storage (S3, Azure Blob Store or a local filesystem)
 
 # Other useful things
 * MongoDB
-* Scylla
 * Minio
 * Docker Compose
 
@@ -37,21 +68,56 @@ The AWS compose file can be replaced with `docker-compose-azure.yml` if you want
 Docker compose setups disable authentication for to make it quick to get started, generally we recommend that you hook Horde.Storage up to a OIDC provider before deploying this.
 
 Horde.Storage hosts Swagger documentation at `/docs` (so `http://localhost/docs` when running locally). This page lets you pick the service you want to call and lists the API for it.
-Note that some endpoints using binary protocol for effieceny and the format of those is not documented in swagger, we always provide a REST compatible api for every use case.
+Note that some endpoints using binary protocol for effieceny and the format of those is not documented in swagger, we always provide a REST api for every use case.
 
 # Deployment
-Horde.Storage is currently only run in deployment on AWS, but the requirements on storage and db are very generic and also abstracted. 
-As such Horde.Storage supports running against local storage and Mongo for a small scale fully on premise deployment. Simlarly we have basic (untested) support for Azure services.
+Horde.Storage is currently only run in production on AWS, but the requirements on storage and db are very generic and also abstracted. 
+We have basic (untested) support for Azure services.
+
+We provide helm values (under `/Helm`) that we use for epic internal deployments to kubernetes, but kubernetes is not a requirement.
+
+Docker images are published to the Epic Games github organization `[REDACTED_PROFILE_URL]/EpicGames/packages/container/package/horde-storage`
+
+## Scylla
+You will need to setup a scylla cluster for Horde Storage to talk to.
+Horde Storage supports running with Scylla Open Source (which is free) or their paid offerings. The paid offerings can help reduce the amount of effort you need to put into managing the cluster as Scylla manager will help you with maintainance tasks.
 
-We provide helm values (under `/Helm`) that we use for epic internal deployments. For prod deployment use the `base` values as well as the value file for the region you are deploying to.
+Some useful links from Scylla on how to setup a multi region cluster:
+https://docs.scylladb.com/stable/operating-scylla/procedures/cluster-management/create-cluster-multidc.html
+
+Download link to open source version:
+https://www.scylladb.com/download/#open-source
+
+Scylla provides machine images for use in cloud environments.
 
 ## AWS
-This is the most tested deployment form as this is how we operate it at Epic. We manage the AWS deployments thru a Terraform script, the key requirement is setting up the replicated Dynamodb tables in the regions we want to run Horde.Storage in.
-The helm chart we install into each regions kubernetes cluster is provided in this repo.
+This is the most tested deployment form as this is how we operate it at Epic. The helm chart we install into each regions kubernetes cluster is provided in this repo.
 
 ## On premise
-Horde.Storage can be deployed onprem without using any cloud resources but it does require a kubernetes cluster to run in for production cases. You will also need to setup a Mongo database for it to run against.
-Note that if you intend to deployment to multiple on prem sites you will need to manually setup replication of the database used.
+Horde.Storage can be deployed onprem without using any cloud resources. You can either setup a Mongo database for this (if you only intend to run this in a single region) or Scylla if you inted to run it multi region but still on premise. If you are starting with one region but might expand later we recommend using Scylla - as that allows you to just scale out while Mongo would require drop all your existing state.
+
+## Azure
+To deploy on AWS you will just need to set Azure as your cloud provider and specify the `Azure.ConnectionString` setting with a connection string to your Azure Blob Storage.
+
+## Testing your deployment
+Once you have a deployment up and running you can connect to the machine and run curl commands to verify its working as it should.
+
+First you can attempt to use the health checks, these should just return the string `Health`
+```
+curl http://localhost/health/live
+```
+
+Next you can attempt to add and fetch content into a namespace
+This will insert a test string (`test`) into the `test-namespace`, you may need to use a different namespace depending on your setup. This also assumes you have authentication disabled. This should return a 200 status code with a empty "needs" list.
+```
+curl http://localhost/api/v1/refs/test-namespace/default/00000000000000000000000000000000000000aa -X PUT --data 'test' -H 'content-type: application/octet-stream' -H 'X-Jupiter-IoHash: 4878CA0425C739FA427F7EDA20FE845F6B2E46BA' -i
+```
+
+After that you can attempt to retrive this object, this should print the 'test' string and a 200 staus code.
+
+```
+curl http://localhost/api/v1/refs/test-namespace/default/00000000000000000000000000000000000000aa.raw -i
+```
 
 # Monitoring
 We use Datadog to monitor our services, as such Horde.Storage is instrumented to work well with that service. But all logs are delivered as structured logs to stdout, so any monitoring service that understands structured logs should be able to monitor it quite well.
@@ -141,13 +207,80 @@ namespace:
 If you specify multiple claims in the claims array these are ANDed together thus requiring all the claims. A claim statement like `A=B` requires claim `A` to have value `B` (or contain value `B` in the case of a array).
 You can also specify the `*` claim which grants access for any valid token no matter which claims it has, this is mostly for debug / testing scenarios and shouldnt be used for production data.
 
-# Development
-## Debugging profiles
-We have two seperate Debug profiles `Local deployment` and `IIS Express`. `ISS Express` configures the service to with in memory mocks of dependent services, allowing for quick iteration but limited functionality.
-`Local Deployment` starts the service on a dedicated port, with configuration to reach other locally running services on their dedicated port. Can be useful to debug and run Horde.Storage without any virtualization. Note that it can be a bit of work to setup all dependencies locally for Local Deployment. If you just want to test things, or step thru the flow we recommend using the docker compose setup and attaching to that.
+# Networking setup
+Horde Storage derives alot of its performance from using your normal internet connection and not relying on a VPN tunnel. As its we strongly recommend that you expose horde-storage on a public internet endpoint. From that follows recommendations like using https and setting up your authentication (see `Authentication`) to prevent anyone from accessing this data.
+
+Horde Storage also provides multiple ports that you can use to control access level of the api.
+## Public Port
+This is the port that you should expose to public internet and what most users should be using to access the service. This port does not expose some of the more sensitive apis (enumerating all contents for instance). This is exposed on port `80` and as `http` within kubernetes. If you are only operating a single region this is the only port you need to expose.
 
-## Making a new release
+## Private port
+This can also be called the Corp port - its intended to be used if you have a route exposed to your intranet. The purpose of this is to have certain very sensitive namespaces that are only exposed to users on your intranet. Use the `IsPublicNamespace` (set to false) in the namespace policy to enable this. We do *NOT* recommend using this for DDC as that prevents users WFH (Working From Home) to access the namespace without a VPN (which is typically to slow for DDC use cases).
 
-Update appversion in version.yaml
-Update changelog.md changing the unreleased section into the version number of the new release.
-Push the version as a git tag - convert this into a release on Github.
\ No newline at end of file
+This is typically exposed on port `8008` and as `corp-http` within kubernetes.
+
+        "PublicApiPorts": [ 80, 8081 ],
+        "CorpApiPorts": [ 8008, 8082 ],
+        "InternalApiPorts": [ 8080, 8083 ]
+
+## Internal Port
+The internal port is only needed to be reachable by other horde-storage instances. This exposes everything that the private port does but also certain apis that are deemed sensitive (enumerating content via the replication log primarily).
+This is exposed on port `8080` and as `internal-http` within kubernetes.
+Its recommended to keep this ingress only to other horde-storage instances via a private VPC or using some kind of ip-range allow list or similar.
+
+Note that this port is primarily used for the speculative blob replication (see `Blob replication setup`). 
+
+# Common operations
+
+## Add new region
+
+The new region will need to contain:
+* S3 storage
+* Compute (Kubernetes cluster or VMs)
+* Scylla deployments
+
+Horde storage itself does not require a lot of configuration for adding a new region - you should just update your cluster settings on all nodes to include the dns of the new region.
+You should also make sure to set your `LocalKeyspaceReplicationStrategy` for the new region.
+
+Scylla does require a bit of effort to scale out to the new region, see their documentation on scaling out:
+https://docs.scylladb.com/stable/operating-scylla/procedures/cluster-management/add-dc-to-existing-dc.html
+Specifically the important part is updating keyspaces
+
+```
+ALTER KEYSPACE jupiter WITH replication = { 'class' : 'NetworkTopologyStrategy', '<exiting_dc>' : 3, <new_dc> : 3};
+ALTER KEYSPACE system_auth WITH replication = { 'class' : 'NetworkTopologyStrategy', '<exiting_dc>' : 3, <new_dc> : 3};
+ALTER KEYSPACE system_distributed WITH replication = { 'class' : 'NetworkTopologyStrategy', '<exiting_dc>' : 3, <new_dc> : 3};
+ALTER KEYSPACE system_traces WITH replication = { 'class' : 'NetworkTopologyStrategy', '<exiting_dc>' : 3, <new_dc> : 3};
+```
+We use a replication factor of 3 everywhere so just add the name of the new region (DC).
+
+You will also need to alter the keyspaces of each of the `local` keyspaces to set the replication factor to 0 for the new region (see `LocalKeyspaceReplicationStrategy`)
+```
+ALTER KEYSPACE jupiter_local_regionA WITH replication = { 'class' : 'NetworkTopologyStrategy', 'regionA' : 2, 'regionB' : 0}
+ALTER KEYSPACE jupiter_local_regionB WITH replication = { 'class' : 'NetworkTopologyStrategy', 'regionA' : 0, 'regionB' : 2}
+```
+This makes sure that the local keyspace is only written to the local region - while this isnt cruical this data will only ever be requested within that region and as such this saves you on a lot of bandwidth and storage within the Scylla cluster.
+
+You will also likely want to update your replicators in your horde storage worker configuration to replicate from this new region. See `Blob replication` setup.
+
+## Blob replication setup
+
+Horde Storage has two methods of replication. On-demand replication and speculative replication.
+On-demand replication will copy a blob from region A to region B as requests happen in region B that is missing the required blob. This type of replication is opt in per namespace by setting the `OnDemandReplication` to true in the namespace policy. We do *NOT* recommend setting this for DDC namespaces as it causes response times to be very variable. For DDC its better to accept the cache miss and rebuild the content in that case, but generally rely on the speculative replication to transfer blobs so that they are available everywhere without the added latency.
+
+Speculative replication uses a journal kept in each region as refs are added to know which content to replicate. This will follow along as changes happen in a namespace and copy all blobs that are being referenced by these new refs. This will end up copying all content, which may never actually be used or needed in the local region but has the benefit of most often having a local blob available once a ref is being resolved. For DDC, where response times are quite important to keep low, we recommended relying on the speculative replication.
+To set it up you will add a section to your worker configuration like this (see `example-values-ABC.yaml`)
+```
+worker:
+  config:
+    Replication:
+      Enabled: true
+      Replicators: 
+      - ReplicatorName: DEF-to-ABC-test-namespace
+        Namespace: test-namespace
+        ConnectionString: http://url-to-region-DEF.com
+```
+The replicator name can be any string that uniquly identifies this replicator (used to store the state of where that replicator has gotten to as well as being used in the logging)
+`Namespace` is the namespace to replicate.
+`ConnectionString` is the url to use to connect to the other regions horde-storage deployment.
+This needs to be exposed using the internal ports (see `Networking setup`). You will also need to have credentials setup for horde storage to use (in the `ServiceCredentials` section) and those credentials will need to have `ReadTransactionLog` access)
