Horde: Create a new log file only when lease is accepted

diff --git a/Engine/Source/Programs/Horde/Horde.Build/Jobs/JobTaskSource.cs b/Engine/Source/Programs/Horde/Horde.Build/Jobs/JobTaskSource.cs
--- a/Engine/Source/Programs/Horde/Horde.Build/Jobs/JobTaskSource.cs
+++ b/Engine/Source/Programs/Horde/Horde.Build/Jobs/JobTaskSource.cs
@@ -765,8 +765,10 @@ namespace Horde.Build.Jobs
 			// The next time to try assigning to another agent
 			DateTime backOffTime = DateTime.UtcNow + TimeSpan.FromMinutes(1.0);
 
+			// Allocate a log ID but hold off creating the actual log file until the lease has been accepted
+			LogId logId = LogId.GenerateNewId();
+			
 			// Try to update the job with this agent id
-			LogId logId = (await _logFileService.CreateLogFileAsync(job.Id, agent.SessionId, LogType.Json)).Id;
 			IJob? newJob = await _jobs.TryAssignLeaseAsync(item._job, item._batchIdx, item._poolId, agent.Id, agent.SessionId!.Value, leaseId, logId);
 			if (newJob != null)
 			{
@@ -798,6 +800,7 @@ namespace Horde.Build.Jobs
 					if (waiter.LeaseSource.TrySetResult(lease))
 					{
 						_logger.LogDebug("Assigned lease {LeaseId} to agent {AgentId}", leaseId, agent.Id);
+						await _logFileService.CreateLogFileAsync(job.Id, agent.SessionId, LogType.Json, logId);
 						return lease;
 					}
 				}
diff --git a/Engine/Source/Programs/Horde/Horde.Build/Logs/ILogFileCollection.cs b/Engine/Source/Programs/Horde/Horde.Build/Logs/ILogFileCollection.cs
--- a/Engine/Source/Programs/Horde/Horde.Build/Logs/ILogFileCollection.cs
+++ b/Engine/Source/Programs/Horde/Horde.Build/Logs/ILogFileCollection.cs
@@ -57,8 +57,9 @@ namespace Horde.Build.Logs
 		/// <param name="jobId">Unique id of the job that owns this log file</param>
 		/// <param name="sessionId">Agent session allowed to update the log</param>
 		/// <param name="type">Type of events to be stored in the log</param>
+		/// <param name="logId">ID of the log file (optional)</param>
 		/// <returns>The new log file document</returns>
-		Task<ILogFile> CreateLogFileAsync(JobId jobId, SessionId? sessionId, LogType type);
+		Task<ILogFile> CreateLogFileAsync(JobId jobId, SessionId? sessionId, LogType type, LogId? logId = null);
 
 		/// <summary>
 		/// Adds a new chunk
diff --git a/Engine/Source/Programs/Horde/Horde.Build/Logs/LogFileCollection.cs b/Engine/Source/Programs/Horde/Horde.Build/Logs/LogFileCollection.cs
--- a/Engine/Source/Programs/Horde/Horde.Build/Logs/LogFileCollection.cs
+++ b/Engine/Source/Programs/Horde/Horde.Build/Logs/LogFileCollection.cs
@@ -79,9 +79,9 @@ namespace Horde.Build.Logs
 			{
 			}
 
-			public LogFileDocument(JobId jobId, SessionId? sessionId, LogType type)
+			public LogFileDocument(JobId jobId, SessionId? sessionId, LogType type, LogId? logId)
 			{
-				Id = LogId.GenerateNewId();
+				Id = logId ?? LogId.GenerateNewId();
 				JobId = jobId;
 				SessionId = sessionId;
 				Type = type;
@@ -117,9 +117,9 @@ namespace Horde.Build.Logs
 		}
 
 		/// <inheritdoc/>
-		public async Task<ILogFile> CreateLogFileAsync(JobId jobId, SessionId? sessionId, LogType type)
+		public async Task<ILogFile> CreateLogFileAsync(JobId jobId, SessionId? sessionId, LogType type, LogId? logId)
 		{
-			LogFileDocument newLogFile = new LogFileDocument(jobId, sessionId, type);
+			LogFileDocument newLogFile = new (jobId, sessionId, type, logId);
 			await _logFiles.InsertOneAsync(newLogFile);
 			return newLogFile;
 		}
diff --git a/Engine/Source/Programs/Horde/Horde.Build/Logs/LogFileService.cs b/Engine/Source/Programs/Horde/Horde.Build/Logs/LogFileService.cs
--- a/Engine/Source/Programs/Horde/Horde.Build/Logs/LogFileService.cs
+++ b/Engine/Source/Programs/Horde/Horde.Build/Logs/LogFileService.cs
@@ -60,8 +60,9 @@ namespace Horde.Build.Logs
 		/// <param name="jobId">Unique id of the job that owns this log file</param>
 		/// <param name="sessionId">Agent session allowed to update the log</param>
 		/// <param name="type">Type of events to be stored in the log</param>
+		/// <param name="logId">ID of the log file (optional)</param>
 		/// <returns>The new log file document</returns>
-		Task<ILogFile> CreateLogFileAsync(JobId jobId, SessionId? sessionId, LogType type);
+		Task<ILogFile> CreateLogFileAsync(JobId jobId, SessionId? sessionId, LogType type, LogId? logId = null);
 
 		/// <summary>
 		/// Gets a logfile by ID
@@ -431,9 +432,9 @@ namespace Horde.Build.Logs
 		}
 
 		/// <inheritdoc/>
-		public Task<ILogFile> CreateLogFileAsync(JobId jobId, SessionId? sessionId, LogType type)
+		public Task<ILogFile> CreateLogFileAsync(JobId jobId, SessionId? sessionId, LogType type, LogId? logId = null)
 		{
-			return _logFiles.CreateLogFileAsync(jobId, sessionId, type);
+			return _logFiles.CreateLogFileAsync(jobId, sessionId, type, logId);
 		}
 
 		/// <inheritdoc/>
