Horde: Fix exception serializing report state with streams containing periods.

diff --git a/Engine/Source/Programs/Horde/Horde.Build/Issues/IssueReportService.cs b/Engine/Source/Programs/Horde/Horde.Build/Issues/IssueReportService.cs
--- a/Engine/Source/Programs/Horde/Horde.Build/Issues/IssueReportService.cs
+++ b/Engine/Source/Programs/Horde/Horde.Build/Issues/IssueReportService.cs
@@ -16,6 +16,8 @@ using Horde.Build.Utilities;
 using HordeCommon;
 using Microsoft.Extensions.Hosting;
 using Microsoft.Extensions.Logging;
+using MongoDB.Bson.Serialization.Attributes;
+using MongoDB.Bson.Serialization.Options;
 
 namespace Horde.Build.Issues
 {
@@ -67,7 +69,8 @@ namespace Horde.Build.Issues
 	[SingletonDocument("6268871c211d05611b3e4fd8")]
 	class IssueReportState : SingletonBase
 	{
-		public Dictionary<string, DateTime> KeyToLastReportTime { get; set; } = new Dictionary<string, DateTime>();
+		[BsonDictionaryOptions(DictionaryRepresentation.ArrayOfDocuments)]
+		public Dictionary<string, DateTime> ReportTimes { get; set; } = new Dictionary<string, DateTime>();
 	}
 
 	/// <summary>
@@ -118,7 +121,7 @@ namespace Horde.Build.Issues
 		async ValueTask TickAsync(CancellationToken cancellationToken)
 		{
 			IssueReportState state = await _state.GetAsync();
-			HashSet<string> invalidKeys = new HashSet<string>(state.KeyToLastReportTime.Keys, StringComparer.Ordinal);
+			HashSet<string> invalidKeys = new HashSet<string>(state.ReportTimes.Keys, StringComparer.Ordinal);
 
 			DateTime currentTime = _clock.UtcNow;
 
@@ -145,9 +148,9 @@ namespace Horde.Build.Issues
 						invalidKeys.Remove(key);
 						
 						DateTime lastReportTime;
-						if (!state.KeyToLastReportTime.TryGetValue(key, out lastReportTime))
+						if (!state.ReportTimes.TryGetValue(key, out lastReportTime))
 						{
-							state = await _state.UpdateAsync(s => s.KeyToLastReportTime[key] = currentTime);
+							state = await _state.UpdateAsync(s => s.ReportTimes[key] = currentTime);
 							continue;
 						}
 
@@ -208,11 +211,11 @@ namespace Horde.Build.Issues
 				{
 					foreach (string updateKey in updateKeys)
 					{
-						state.KeyToLastReportTime[updateKey] = currentTime;
+						state.ReportTimes[updateKey] = currentTime;
 					}
 					foreach (string invalidKey in invalidKeys)
 					{
-						state.KeyToLastReportTime.Remove(invalidKey);
+						state.ReportTimes.Remove(invalidKey);
 					}
 				}
 				state = await _state.UpdateAsync(UpdateKeys);
