Horde: Add additional logging for Slack message updates.

diff --git a/Engine/Source/Programs/Horde/Horde.Build/Notifications/Sinks/SlackNotificationSink.cs b/Engine/Source/Programs/Horde/Horde.Build/Notifications/Sinks/SlackNotificationSink.cs
--- a/Engine/Source/Programs/Horde/Horde.Build/Notifications/Sinks/SlackNotificationSink.cs
+++ b/Engine/Source/Programs/Horde/Horde.Build/Notifications/Sinks/SlackNotificationSink.cs
@@ -293,6 +293,15 @@ namespace Horde.Build.Notifications.Sinks
 			}
 
 			MessageStateDocument state = await _messageStates.FindOneAndUpdateAsync(filter, update, new FindOneAndUpdateOptions<MessageStateDocument> { IsUpsert = true, ReturnDocument = ReturnDocument.After });
+			if (state.Id == newId)
+			{
+				_logger.LogInformation("Posted message {Ts} (recipient: {Recipient}, user: {UserId}), event: {EventId}, digest: {Digest})", state.Ts, recipient, userId ?? UserId.Empty, eventId, digest);
+			}
+			else
+			{
+				_logger.LogInformation("Updated message {Ts} (recipient: {Recipient}, user: {UserId}), event: {EventId}, digest: {Digest})", state.Ts, recipient, userId ?? UserId.Empty, eventId, digest);
+			}
+
 			return (state, state.Id == newId);
 		}
 
@@ -1259,7 +1268,7 @@ namespace Horde.Build.Notifications.Sinks
 
 			if (IsRecipientAllowed(recipient, "issue update"))
 			{
-				await SendOrUpdateMessageAsync(recipient, GetIssueEventId(issue), userId, attachment);
+				await SendOrUpdateMessageAsync(recipient, GetIssueEventId(issue, recipient), userId, attachment);
 			}
 		}
 
@@ -1268,6 +1277,11 @@ namespace Horde.Build.Notifications.Sinks
 			return $"issue_{issue.Id}";
 		}
 
+		static string GetIssueEventId(IIssue issue, string recipient)
+		{
+			return $"issue_{issue.Id}_for_{recipient}";
+		}
+
 		async Task<string> FormatNameAsync(UserId userId)
 		{
 			IUser? user = await _userCollection.GetUserAsync(userId);
