Horde: Fix triage alias not being tagged if an issue doesn't alert suspects due to exceeding max count.

diff --git a/Engine/Source/Programs/Horde/Horde.Build/Notifications/Sinks/SlackNotificationSink.cs b/Engine/Source/Programs/Horde/Horde.Build/Notifications/Sinks/SlackNotificationSink.cs
--- a/Engine/Source/Programs/Horde/Horde.Build/Notifications/Sinks/SlackNotificationSink.cs
+++ b/Engine/Source/Programs/Horde/Horde.Build/Notifications/Sinks/SlackNotificationSink.cs
@@ -918,6 +918,7 @@ namespace Horde.Build.Notifications.Sinks
 					await SendOrUpdateMessageAsync(triageChannel, state.Ts, $"{eventId}_buttons", null, message);
 				}
 
+				bool notifyTriageAlias = false;
 				if (isNew)
 				{
 					// If it has an owner, show that
@@ -938,7 +939,6 @@ namespace Horde.Build.Notifications.Sinks
 					else
 					{
 						IGrouping<UserId, IIssueSuspect>[] suspectGroups = suspects.GroupBy(x => x.AuthorId).ToArray();
-
 						if (suspectGroups.Length > 0 && suspectGroups.Length <= workflow.MaxMentions)
 						{
 							List<string> suspectList = new List<string>();
@@ -953,6 +953,10 @@ namespace Horde.Build.Notifications.Sinks
 							string suspectMessage = $"Possibly {StringUtils.FormatList(suspectList, "or")}.";
 							await _slackClient.PostMessageAsync(triageChannel, state.Ts, suspectMessage);
 						}
+						else
+						{
+							notifyTriageAlias = true;
+						}
 					}
 
 					if (_environment.IsProduction() && workflow.AllowMentions)
@@ -970,7 +974,7 @@ namespace Horde.Build.Notifications.Sinks
 					}
 				}
 
-				if (workflow.TriageAlias != null && issue.OwnerId == null && suspects.All(x => x.DeclinedAt != null))
+				if (workflow.TriageAlias != null && issue.OwnerId == null && (suspects.All(x => x.DeclinedAt != null) || notifyTriageAlias))
 				{
 					string triageMessage = $"(cc {FormatUserOrGroupMention(workflow.TriageAlias)} for triage).";
 					await SendOrUpdateMessageAsync(triageChannel, state.Ts, eventId + "_triage", null, triageMessage);
