Horde: Improve logging for Slack messages.

diff --git a/Engine/Source/Programs/Horde/Horde.Build/Notifications/Sinks/SlackNotificationSink.cs b/Engine/Source/Programs/Horde/Horde.Build/Notifications/Sinks/SlackNotificationSink.cs
--- a/Engine/Source/Programs/Horde/Horde.Build/Notifications/Sinks/SlackNotificationSink.cs
+++ b/Engine/Source/Programs/Horde/Horde.Build/Notifications/Sinks/SlackNotificationSink.cs
@@ -295,11 +295,11 @@ namespace Horde.Build.Notifications.Sinks
 			MessageStateDocument state = await _messageStates.FindOneAndUpdateAsync(filter, update, new FindOneAndUpdateOptions<MessageStateDocument> { IsUpsert = true, ReturnDocument = ReturnDocument.After });
 			if (state.Id == newId)
 			{
-				_logger.LogInformation("Posted message {Ts} (recipient: {Recipient}, user: {UserId}), event: {EventId}, digest: {Digest})", state.Ts, recipient, userId ?? UserId.Empty, eventId, digest);
+				_logger.LogInformation("Posted message {MessageId} (ts: {Ts}, recipient: {Recipient}, user: {UserId}, event: {EventId}, digest: {Digest})", state.Id, state.Ts, recipient, userId ?? UserId.Empty, eventId, digest);
 			}
 			else
 			{
-				_logger.LogInformation("Updated message {Ts} (recipient: {Recipient}, user: {UserId}), event: {EventId}, digest: {Digest})", state.Ts, recipient, userId ?? UserId.Empty, eventId, digest);
+				_logger.LogInformation("Updated message {MessageId} (ts: {Ts}, recipient: {Recipient}, user: {UserId}, event: {EventId}, digest: {Digest})", state.Id, state.Ts, recipient, userId ?? UserId.Empty, eventId, digest);
 			}
 
 			return (state, state.Id == newId);
@@ -324,6 +324,7 @@ namespace Horde.Build.Notifications.Sinks
 				update = update.Set(x => x.Permalink, permalink);
 			}
 			await _messageStates.FindOneAndUpdateAsync(filter, update);
+			_logger.LogInformation("Updated message {MessageId} (ts: {Ts}, recipient: {Recipient}, permalink: {Permalink})", messageId, ts, channel, permalink ?? "(n/a)");
 		}
 
 		#endregion
