Slack: Include threadTs in log messages.

diff --git a/Engine/Source/Programs/Horde/Horde.Build/Notifications/Sinks/SlackNotificationSink.cs b/Engine/Source/Programs/Horde/Horde.Build/Notifications/Sinks/SlackNotificationSink.cs
--- a/Engine/Source/Programs/Horde/Horde.Build/Notifications/Sinks/SlackNotificationSink.cs
+++ b/Engine/Source/Programs/Horde/Horde.Build/Notifications/Sinks/SlackNotificationSink.cs
@@ -2048,7 +2048,7 @@ namespace Horde.Build.Notifications.Sinks
 			(MessageStateDocument state, bool isNew) = await AddOrUpdateMessageStateAsync(recipient, eventId, userId, requestDigest);
 			if (isNew)
 			{
-				_logger.LogInformation("Sending new slack message to {SlackUser} (msg: {MessageId})", recipient, state.Id);
+				_logger.LogInformation("Sending new slack message to {SlackUser} (msg: {MessageId}, threadTs: {ThreadTs})", recipient, state.Id, threadTs ?? "n/a");
 
 				state.Channel = recipient;
 				if (threadTs == null)
@@ -2064,7 +2064,8 @@ namespace Horde.Build.Notifications.Sinks
 			}
 			else if (!String.IsNullOrEmpty(state.Ts))
 			{
-				_logger.LogInformation("Updating existing slack message {MessageId} for user {SlackUser} ({Channel}, {MessageTs})", state.Id, recipient, state.Channel, state.Ts);
+				_logger.LogInformation("Updating existing slack message {MessageId} for user {SlackUser} (channel: {Channel}, ts: {Ts}, threadTs: {ThreadTs})", state.Id, recipient, state.Channel, state.Ts, threadTs ?? "n/a");
+
 				if (threadTs == null)
 				{
 					await _slackClient.UpdateMessageAsync(state.Channel, state.Ts, message);
