Horde: Do not remove leases that have no steps to execute. Should prevent errors when there is a race with leases being assigned to agents.

diff --git a/Engine/Source/Programs/Horde/Horde.Build/Jobs/JobCollection.cs b/Engine/Source/Programs/Horde/Horde.Build/Jobs/JobCollection.cs
--- a/Engine/Source/Programs/Horde/Horde.Build/Jobs/JobCollection.cs
+++ b/Engine/Source/Programs/Horde/Horde.Build/Jobs/JobCollection.cs
@@ -1334,7 +1334,7 @@ namespace Horde.Build.Jobs
 			}
 
 			// Remove any batches which are now empty
-			job.Batches.RemoveAll(x => x.Steps.Count == 0 && x.Error == JobStepBatchError.None);
+			job.Batches.RemoveAll(x => x.Steps.Count == 0 && x.LeaseId == null && x.Error == JobStepBatchError.None);
 
 			// Find all the targets in this job
 			HashSet<string> targets = new HashSet<string>(StringComparer.OrdinalIgnoreCase);
