Horde: Refactor Perforce service to use EpicGames.Perforce rather than P4API.NET. New implementation supports connection pooling and asynchronous operation.

diff --git a/Engine/Source/Programs/Horde/BuildHorde.xml b/Engine/Source/Programs/Horde/BuildHorde.xml
--- a/Engine/Source/Programs/Horde/BuildHorde.xml
+++ b/Engine/Source/Programs/Horde/BuildHorde.xml
@@ -22,7 +22,6 @@
 				Engine/Source/Programs/Horde/...
 				Engine/Source/Programs/AutomationTool/AutomationUtils/Matchers/...
 				Engine/Source/Programs/UnrealBuildTool/Matchers/...
-				Engine/Source/ThirdParty/Perforce/...
 			</Property>
 			<Tag Files="$(StagePaths)" Except=".../.vs/...;.../.git/...;.../bin/...;.../obj/..." With="#InputFiles"/>
 
diff --git a/Engine/Source/Programs/Horde/Horde.Build.Tests/Stubs/Services/PerforceServiceStub.cs b/Engine/Source/Programs/Horde/Horde.Build.Tests/Stubs/Services/PerforceServiceStub.cs
--- a/Engine/Source/Programs/Horde/Horde.Build.Tests/Stubs/Services/PerforceServiceStub.cs
+++ b/Engine/Source/Programs/Horde/Horde.Build.Tests/Stubs/Services/PerforceServiceStub.cs
@@ -3,6 +3,7 @@
 using System;
 using System.Collections.Generic;
 using System.Linq;
+using System.Threading;
 using System.Threading.Tasks;
 using EpicGames.Perforce;
 using Horde.Build.Perforce;
@@ -60,12 +61,12 @@ namespace Horde.Build.Tests.Stubs.Services
 			_userCollection = userCollection;
 		}
 
-		public Task<IPerforceConnection?> GetServiceUserConnection(string? clusterName)
+		public Task<IPerforceConnection> ConnectAsync(string? clusterName, string? userName, CancellationToken cancellationToken)
 		{
 			throw new NotImplementedException();
 		}
 
-		public async ValueTask<IUser> FindOrAddUserAsync(string clusterName, string userName)
+		public async ValueTask<IUser> FindOrAddUserAsync(string clusterName, string userName, CancellationToken cancellationToken)
 		{
 			return await _userCollection.FindOrAddUserByLoginAsync(userName);
 		}
@@ -81,7 +82,7 @@ namespace Horde.Build.Tests.Stubs.Services
 			streamChanges.Add(number, new ChangeDetails(number, author, null!, description, files.Select(x => PerforceExtensions.CreateChangeFile(x)).ToList(), DateTime.Now));
 		}
 
-		public Task<List<ChangeSummary>> GetChangesAsync(string clusterName, string streamName, int? minChange, int? maxChange, int numResults, string? impersonateUser)
+		public Task<List<ChangeSummary>> GetChangesAsync(string clusterName, string streamName, int? minChange, int? maxChange, int numResults, string? impersonateUser, CancellationToken cancellationToken)
 		{
 			List<ChangeSummary> results = new List<ChangeSummary>();
 
@@ -108,12 +109,12 @@ namespace Horde.Build.Tests.Stubs.Services
 			return Task.FromResult(results);
 		}
 
-		public Task<(CheckShelfResult, string?)> CheckShelfAsync(string clusterName, string streamName, int changeNumber, string? impersonateUser)
+		public Task<(CheckShelfResult, string?)> CheckShelfAsync(string clusterName, string streamName, int changeNumber, string? impersonateUser, CancellationToken cancellationToken)
 		{
 			throw new NotImplementedException();
 		}
 
-		public Task<List<ChangeDetails>> GetChangeDetailsAsync(string clusterName, string streamName, IReadOnlyList<int> changeNumbers, string? impersonateUser)
+		public Task<List<ChangeDetails>> GetChangeDetailsAsync(string clusterName, string streamName, IReadOnlyList<int> changeNumbers, string? impersonateUser, CancellationToken cancellationToken)
 		{
 			List<ChangeDetails> results = new List<ChangeDetails>();
 			foreach (int changeNumber in changeNumbers)
@@ -128,7 +129,7 @@ namespace Horde.Build.Tests.Stubs.Services
 			return Task.FromResult("bogus-ticket");
 		}
 
-		public Task<int> GetCodeChangeAsync(string clusterName, string streamName, int change)
+		public Task<int> GetCodeChangeAsync(string clusterName, string streamName, int change, CancellationToken cancellationToken)
 		{
 			int codeChange = 0;
 
@@ -148,59 +149,54 @@ namespace Horde.Build.Tests.Stubs.Services
 			return Task.FromResult(codeChange);
 		}
 
-		public Task<int> CreateNewChangeAsync(string clusterName, string streamName, string path, string description)
+		public Task<int> CreateNewChangeAsync(string clusterName, string streamName, string path, string description, CancellationToken cancellationToken)
 		{
 			ChangeDetails newChange = new ChangeDetails(Changes[streamName].First().Key + 1, _testUser, null!, description, new List<ChangeFile> { PerforceExtensions.CreateChangeFile(path) }, DateTime.Now);
 			Changes[streamName].Add(newChange.Number, newChange);
 			return Task.FromResult(newChange.Number);
 		}
 
-		public Task<List<FileSummary>> FindFilesAsync(string clusterName, IEnumerable<string> paths)
+		public Task<List<FileSummary>> FindFilesAsync(string clusterName, IEnumerable<string> paths, CancellationToken cancellationToken)
 		{
 			throw new NotImplementedException();
 		}
 
-		public Task<byte[]> PrintAsync(string clusterName, string path)
+		public Task<byte[]> PrintAsync(string clusterName, string path, CancellationToken cancellationToken)
 		{
 			throw new NotImplementedException();
 		}
 
-		public Task<(int? Change, string Message)> SubmitShelvedChangeAsync(string clusterName, int shelvedChange, int originalChange)
+		public Task<(int? Change, string Message)> SubmitShelvedChangeAsync(string clusterName, int shelvedChange, int originalChange, CancellationToken cancellationToken)
 		{
 			throw new NotImplementedException();
 		}
 
-		public Task<int> DuplicateShelvedChangeAsync(string clusterName, int shelvedChange)
+		public Task<int> DuplicateShelvedChangeAsync(string clusterName, int shelvedChange, CancellationToken cancellationToken)
 		{
 			return Task.FromResult(shelvedChange);
 		}
 
-		public Task DeleteShelvedChangeAsync(string clusterName, int shelvedChange)
+		public Task DeleteShelvedChangeAsync(string clusterName, int shelvedChange, CancellationToken cancellationToken)
 		{
 			return Task.CompletedTask;
 		}
 
-		public Task UpdateChangelistDescription(string clusterName, int change, string description)
+		public Task UpdateChangelistDescription(string clusterName, int change, string description, CancellationToken cancellationToken)
 		{
 			return Task.CompletedTask;
 		}
 
-		public Task<List<ChangeSummary>> GetChangesAsync(string clusterName, int? minChange, int? maxChange, int maxResults)
+		public Task<List<ChangeSummary>> GetChangesAsync(string clusterName, int? minChange, int? maxChange, int maxResults, CancellationToken cancellationToken)
 		{
 			throw new NotImplementedException();
 		}
 
-		public Task<ChangeDetails> GetChangeDetailsAsync(string clusterName, int changeNumber)
+		public Task<ChangeDetails> GetChangeDetailsAsync(string clusterName, int changeNumber, CancellationToken cancellationToken)
 		{
 			throw new NotImplementedException();
 		}
 
-		public Task<List<ChangeFile>> GetStreamSnapshotAsync(string clusterName, string streamName, int change)
-		{
-			throw new NotImplementedException();
-		}
-
-		public Task<ChangeDetails> GetChangeDetailsAsync(string clusterName, string streamName, int changeNumber)
+		public Task<ChangeDetails> GetChangeDetailsAsync(string clusterName, string streamName, int changeNumber, CancellationToken cancellationToken)
 		{
 			return Task.FromResult(Changes[streamName][changeNumber]);
 		}
diff --git a/Engine/Source/Programs/Horde/Horde.Build/Dockerfile b/Engine/Source/Programs/Horde/Horde.Build/Dockerfile
--- a/Engine/Source/Programs/Horde/Horde.Build/Dockerfile
+++ b/Engine/Source/Programs/Horde/Horde.Build/Dockerfile
@@ -38,7 +38,6 @@ COPY Source/Programs/Shared/EpicGames.Serialization.Tests/*.csproj ./Source/Prog
 COPY Source/Programs/Horde/Horde.Agent/*.csproj ./Source/Programs/Horde/Horde.Agent/
 COPY Source/Programs/Horde/Horde.Build/*.csproj ./Source/Programs/Horde/Horde.Build/
 COPY Source/Programs/Horde/Horde.Build.Tests/*.csproj ./Source/Programs/Horde/Horde.Build.Tests/
-COPY Source/ThirdParty/Perforce/P4Api.Net/ ./Source/ThirdParty/Perforce/P4Api.Net/
 RUN dotnet restore Source/Programs/Horde/Horde.Build/Horde.Build.csproj -p:WithHordeStorage=false
 
 # Build the server
diff --git a/Engine/Source/Programs/Horde/Horde.Build/Horde.Build.csproj b/Engine/Source/Programs/Horde/Horde.Build/Horde.Build.csproj
--- a/Engine/Source/Programs/Horde/Horde.Build/Horde.Build.csproj
+++ b/Engine/Source/Programs/Horde/Horde.Build/Horde.Build.csproj
@@ -1,4 +1,4 @@
-﻿<Project Sdk="Microsoft.NET.Sdk.Web">
+<Project Sdk="Microsoft.NET.Sdk.Web">
 
   <PropertyGroup>
     <TargetFramework>net6.0</TargetFramework>
@@ -101,23 +101,9 @@
     <ProjectReference Include="..\..\Shared\EpicGames.Perforce\EpicGames.Perforce.csproj" />
     <ProjectReference Include="..\..\Shared\EpicGames.Redis\EpicGames.Redis.csproj" />
     <ProjectReference Include="..\..\Shared\EpicGames.Serialization\EpicGames.Serialization.csproj" />
-    <ProjectReference Include="..\..\..\ThirdParty\Perforce\P4Api.Net\p4api.net\p4api.net.csproj" />
     <ProjectReference Include="..\..\Shared\EpicGames.Slack\EpicGames.Slack.csproj" />
   </ItemGroup>
 
-  <ItemGroup>
-    <Content Include="..\..\..\ThirdParty\Perforce\P4Api.Net\native\windows\release\p4bridge.*" Condition=" '$(OS)' == 'Windows_NT' ">
-      <CopyToOutputDirectory>PreserveNewest</CopyToOutputDirectory>
-    </Content>
-    <Content Include="..\..\..\ThirdParty\Perforce\P4Api.Net\native\linux\release\p4bridge.*" Condition="$([MSBuild]::IsOSPlatform('Linux'))">
-      <CopyToOutputDirectory>PreserveNewest</CopyToOutputDirectory>
-    </Content>
-    <Content Include="..\ThirdParty\**" Condition="$([MSBuild]::IsOSPlatform('Windows'))">
-      <CopyToOutputDirectory>PreserveNewest</CopyToOutputDirectory>
-      <Link>ThirdParty/%(RecursiveDir)/%(Filename)%(Extension)</Link>
-    </Content>
-  </ItemGroup>
-
   <ItemGroup>
     <Content Include="Defaults\ue5.project.png">
       <CopyToOutputDirectory>PreserveNewest</CopyToOutputDirectory>
diff --git a/Engine/Source/Programs/Horde/Horde.Build/Perforce/CommitService.cs b/Engine/Source/Programs/Horde/Horde.Build/Perforce/CommitService.cs
--- a/Engine/Source/Programs/Horde/Horde.Build/Perforce/CommitService.cs
+++ b/Engine/Source/Programs/Horde/Horde.Build/Perforce/CommitService.cs
@@ -265,7 +265,7 @@ namespace Horde.Build.Perforce
 			_logger.LogInformation("Replicating metadata for cluster {Name}", clusterName);
 
 			// Create a connection to the server
-			using IPerforceConnection? connection = await _perforceService.GetServiceUserConnection(clusterName);
+			using IPerforceConnection? connection = await _perforceService.ConnectAsync(clusterName);
 			if (connection == null)
 			{
 				throw new PerforceException($"Unable to create cluster connection for {clusterName}");
diff --git a/Engine/Source/Programs/Horde/Horde.Build/Perforce/IPerforceService.cs b/Engine/Source/Programs/Horde/Horde.Build/Perforce/IPerforceService.cs
--- a/Engine/Source/Programs/Horde/Horde.Build/Perforce/IPerforceService.cs
+++ b/Engine/Source/Programs/Horde/Horde.Build/Perforce/IPerforceService.cs
@@ -4,6 +4,7 @@ using System;
 using System.Collections.Generic;
 using System.Diagnostics.CodeAnalysis;
 using System.Text.RegularExpressions;
+using System.Threading;
 using System.Threading.Tasks;
 using EpicGames.Perforce;
 using Horde.Build.Jobs.Templates;
@@ -49,19 +50,22 @@ namespace Horde.Build.Perforce
 	public interface IPerforceService
 	{
 		/// <summary>
-		/// Finds or adds a user from the given Perforce server, adding the user (and populating their profile with Perforce data) if they do not currently exist
+		/// Connect to the given Perforce cluster
 		/// </summary>
-		/// <param name="clusterName"></param>
-		/// <param name="userName"></param>
-		/// <returns></returns>
-		public ValueTask<IUser> FindOrAddUserAsync(string clusterName, string userName);
+		/// <param name="clusterName">Name of the cluster</param>
+		/// <param name="userName">Name of the user to connect as. Uses the service account if null.</param>
+		/// <param name="cancellationToken">Cancellation token for the operation</param>
+		/// <returns>Connection information</returns>
+		public Task<IPerforceConnection> ConnectAsync(string clusterName, string? userName = null, CancellationToken cancellationToken = default);
 
 		/// <summary>
-		/// 
+		/// Finds or adds a user from the given Perforce server, adding the user (and populating their profile with Perforce data) if they do not currently exist
 		/// </summary>
 		/// <param name="clusterName"></param>
+		/// <param name="userName"></param>
+		/// <param name="cancellationToken">Cancellation token for the operation</param>
 		/// <returns></returns>
-		public Task<IPerforceConnection?> GetServiceUserConnection(string? clusterName);
+		public ValueTask<IUser> FindOrAddUserAsync(string clusterName, string userName, CancellationToken cancellationToken = default);
 
 		/// <summary>
 		/// Create a new changelist by submitting the given file
@@ -70,8 +74,9 @@ namespace Horde.Build.Perforce
 		/// <param name="streamName">The stream to query</param>
 		/// <param name="path">Path for the file to submit</param>
 		/// <param name="description">Description for the changelist</param>
+		/// <param name="cancellationToken">Cancellation token for the operation</param>
 		/// <returns>New changelist number</returns>
-		public Task<int> CreateNewChangeAsync(string clusterName, string streamName, string path, string description);
+		public Task<int> CreateNewChangeAsync(string clusterName, string streamName, string path, string description, CancellationToken cancellationToken = default);
 
 		/// <summary>
 		/// Gets the code change corresponding to an actual change submitted to a stream
@@ -79,8 +84,9 @@ namespace Horde.Build.Perforce
 		/// <param name="clusterName">Name of the Perforce cluster</param>
 		/// <param name="streamName"></param>
 		/// <param name="change">The changelist number to query</param>
+		/// <param name="cancellationToken">Cancellation token for the operation</param>
 		/// <returns>Code change for the latest change</returns>
-		public Task<int> GetCodeChangeAsync(string clusterName, string streamName, int change);
+		public Task<int> GetCodeChangeAsync(string clusterName, string streamName, int change, CancellationToken cancellationToken = default);
 
 		/// <summary>
 		/// Finds changes submitted to a depot Gets the latest change for a particular stream
@@ -88,9 +94,10 @@ namespace Horde.Build.Perforce
 		/// <param name="clusterName">Name of the Perforce cluster</param>
 		/// <param name="minChange">The minimum changelist number</param>
 		/// <param name="maxChange">The maximum changelist number</param>
-		/// <param name="maxResults"></param>
+		/// <param name="maxResults">Maximum number of results to return</param>
+		/// <param name="cancellationToken">Cancellation token for the operation</param>
 		/// <returns>Changelist information</returns>
-		public Task<List<ChangeSummary>> GetChangesAsync(string clusterName, int? minChange, int? maxChange, int maxResults);
+		public Task<List<ChangeSummary>> GetChangesAsync(string clusterName, int? minChange, int? maxChange, int maxResults, CancellationToken cancellationToken = default);
 
 		/// <summary>
 		/// Gets the latest change for a particular stream
@@ -100,9 +107,10 @@ namespace Horde.Build.Perforce
 		/// <param name="minChange">The minimum changelist number</param>
 		/// <param name="maxChange">The maximum changelist number</param>
 		/// <param name="results">Number of results to return</param>
-		/// <param name="impersonateUser">Name of the user to impersonate</param>
+		/// <param name="impersonateUser">User to impersonate for the query</param>
+		/// <param name="cancellationToken">Cancellation token for the operation</param>
 		/// <returns>Latest changelist number</returns>
-		public Task<List<ChangeSummary>> GetChangesAsync(string clusterName, string streamName, int? minChange, int? maxChange, int results, string? impersonateUser);
+		public Task<List<ChangeSummary>> GetChangesAsync(string clusterName, string streamName, int? minChange, int? maxChange, int results, string? impersonateUser, CancellationToken cancellationToken = default);
 
 		/// <summary>
 		/// Checks a shelf is valid for the given stream
@@ -111,8 +119,9 @@ namespace Horde.Build.Perforce
 		/// <param name="streamName">The stream to query</param>
 		/// <param name="changeNumber">Shelved changelist number</param>
 		/// <param name="impersonateUser">Name of the user to impersonate</param>
+		/// <param name="cancellationToken">Cancellation token for the operation</param>
 		/// <returns>Whether the shelf is valid, plus the description for it</returns>
-		public Task<(CheckShelfResult, string?)> CheckShelfAsync(string clusterName, string streamName, int changeNumber, string? impersonateUser);
+		public Task<(CheckShelfResult, string?)> CheckShelfAsync(string clusterName, string streamName, int changeNumber, string? impersonateUser, CancellationToken cancellationToken = default);
 
 		/// <summary>
 		/// Gets the latest change for a particular stream
@@ -120,8 +129,9 @@ namespace Horde.Build.Perforce
 		/// <param name="clusterName">Name of the Perforce cluster</param>
 		/// <param name="streamName">The stream to query</param>
 		/// <param name="changeNumber">Change numbers to query</param>
+		/// <param name="cancellationToken">Cancellation token for the operation</param>
 		/// <returns>Commit details</returns>
-		public Task<ChangeDetails> GetChangeDetailsAsync(string clusterName, string streamName, int changeNumber);
+		public Task<ChangeDetails> GetChangeDetailsAsync(string clusterName, string streamName, int changeNumber, CancellationToken cancellationToken = default);
 
 		/// <summary>
 		/// Gets the latest change for a particular stream
@@ -130,49 +140,55 @@ namespace Horde.Build.Perforce
 		/// <param name="streamName">The stream to query</param>
 		/// <param name="changeNumbers">Change numbers to query</param>
 		/// <param name="impersonateUser">Name of the user to impersonate</param>
+		/// <param name="cancellationToken">Cancellation token for the operation</param>
 		/// <returns>Commit details</returns>
-		public Task<List<ChangeDetails>> GetChangeDetailsAsync(string clusterName, string streamName, IReadOnlyList<int> changeNumbers, string? impersonateUser);
+		public Task<List<ChangeDetails>> GetChangeDetailsAsync(string clusterName, string streamName, IReadOnlyList<int> changeNumbers, string? impersonateUser, CancellationToken cancellationToken = default);
 
 		/// <summary>
 		/// Gets the latest changes for a set of depot paths
 		/// </summary>
 		/// <param name="clusterName">Name of the Perforce cluster</param>
 		/// <param name="paths"></param>
+		/// <param name="cancellationToken">Cancellation token for the operation</param>
 		/// <returns></returns>
-		public Task<List<FileSummary>> FindFilesAsync(string clusterName, IEnumerable<string> paths);
+		public Task<List<FileSummary>> FindFilesAsync(string clusterName, IEnumerable<string> paths, CancellationToken cancellationToken = default);
 
 		/// <summary>
 		/// Gets the contents of a file in the depot
 		/// </summary>
 		/// <param name="clusterName">Name of the Perforce cluster</param>
 		/// <param name="path">Path to read</param>
+		/// <param name="cancellationToken">Cancellation token for the operation</param>
 		/// <returns>Data for the file</returns>
-		public Task<byte[]> PrintAsync(string clusterName, string path);
+		public Task<byte[]> PrintAsync(string clusterName, string path, CancellationToken cancellationToken = default);
 
 		/// <summary>
 		/// Duplicates a shelved changelist
 		/// </summary>
 		/// <param name="clusterName">Name of the Perforce cluster</param>
 		/// <param name="shelvedChange">The shelved changelist</param>
+		/// <param name="cancellationToken">Cancellation token for the operation</param>
 		/// <returns>The duplicated changelist</returns>
-		public Task<int> DuplicateShelvedChangeAsync(string clusterName, int shelvedChange);
+		public Task<int> DuplicateShelvedChangeAsync(string clusterName, int shelvedChange, CancellationToken cancellationToken = default);
 
 		/// <summary>
 		/// Submit a shelved changelist
 		/// </summary>
 		/// <param name="clusterName">Name of the Perforce cluster</param>
-		/// <param name="shelvedChange">The shelved changelist number, created by <see cref="DuplicateShelvedChangeAsync(String,Int32)"/></param>
+		/// <param name="shelvedChange">The shelved changelist number, created by <see cref="DuplicateShelvedChangeAsync(String,Int32,CancellationToken)"/></param>
 		/// <param name="originalChange">The original changelist number</param>
+		/// <param name="cancellationToken">Cancellation token for the operation</param>
 		/// <returns>Tuple consisting of the submitted changelist number and message</returns>
-		public Task<(int? Change, string Message)> SubmitShelvedChangeAsync(string clusterName, int shelvedChange, int originalChange);
+		public Task<(int? Change, string Message)> SubmitShelvedChangeAsync(string clusterName, int shelvedChange, int originalChange, CancellationToken cancellationToken = default);
 
 		/// <summary>
 		/// Deletes a changelist containing shelved files.
 		/// </summary>
 		/// <param name="clusterName">Name of the Perforce cluster</param>
 		/// <param name="shelvedChange">The changelist containing shelved files</param>
+		/// <param name="cancellationToken">Cancellation token for the operation</param>
 		/// <returns>Async tasy</returns>
-		public Task DeleteShelvedChangeAsync(string clusterName, int shelvedChange);
+		public Task DeleteShelvedChangeAsync(string clusterName, int shelvedChange, CancellationToken cancellationToken = default);
 
 		/// <summary>
 		/// Updates a changelist description
@@ -180,8 +196,9 @@ namespace Horde.Build.Perforce
 		/// <param name="clusterName">Name of the Perforce cluster</param>
 		/// <param name="change">The change to update</param>
 		/// <param name="description">The new description</param>
+		/// <param name="cancellationToken">Cancellation token for the operation</param>
 		/// <returns>Async task</returns>
-		public Task UpdateChangelistDescription(string clusterName, int change, string description);
+		public Task UpdateChangelistDescription(string clusterName, int change, string description, CancellationToken cancellationToken = default);
 	}
 
 	/// <summary>
diff --git a/Engine/Source/Programs/Horde/Horde.Build/Perforce/PerforceService.cs b/Engine/Source/Programs/Horde/Horde.Build/Perforce/PerforceService.cs
--- a/Engine/Source/Programs/Horde/Horde.Build/Perforce/PerforceService.cs
+++ b/Engine/Source/Programs/Horde/Horde.Build/Perforce/PerforceService.cs
@@ -2,6 +2,7 @@
 
 using System;
 using System.Collections.Generic;
+using System.Diagnostics;
 using System.Diagnostics.CodeAnalysis;
 using System.Globalization;
 using System.IO;
@@ -11,9 +12,11 @@ using System.Reflection;
 using System.Runtime.InteropServices;
 using System.Text;
 using System.Text.RegularExpressions;
+using System.Threading;
 using System.Threading.Tasks;
 using EpicGames.Core;
 using EpicGames.Perforce;
+using Grpc.Core;
 using Horde.Build.Server;
 using Horde.Build.Users;
 using Horde.Build.Utilities;
@@ -22,75 +25,120 @@ using Microsoft.Extensions.Logging;
 using Microsoft.Extensions.Options;
 using OpenTracing;
 using OpenTracing.Util;
+using StackExchange.Redis;
 
 namespace Horde.Build.Perforce
 {
-	using P4 = global::Perforce.P4;
-
-	static class PerforceExtensions
-	{
-		static readonly FieldInfo s_field = typeof(P4.Changelist).GetField("_baseForm", BindingFlags.Instance | BindingFlags.NonPublic)!;
-
-		public static string GetPath(this P4.Changelist changelist)
-		{
-			P4.FormBase? formBase = s_field.GetValue(changelist) as P4.FormBase;
-			object? pathValue = null;
-			formBase?.TryGetValue("path", out pathValue);
-			return (pathValue as string) ?? "//...";
-		}
-	}
-
 	/// <summary>
 	/// P4API implementation of the Perforce service
 	/// </summary>
 	class PerforceService : IPerforceService, IDisposable
 	{
-		class CachedTicketInfo
+		sealed class PooledConnection : IDisposable
 		{
-			public IPerforceServer _server;
-			public string _userName;
-			public P4.Credential _ticket;
+			public IPerforceConnection Perforce { get; }
+			public IPerforceSettings Settings => Perforce.Settings;
+			public Credentials Credentials { get; }
+			public ClientRecord? Client { get; }
+			public DateTime ExpiresAt { get; }
 
-			public CachedTicketInfo(IPerforceServer server, string userName, P4.Credential ticket)
+			InfoRecord? _info;
+
+			public PooledConnection(IPerforceConnection perforce, Credentials credentials, ClientRecord? client, DateTime expiresAt)
+			{
+				Perforce = perforce;
+				Credentials = credentials;
+				Client = client;
+				ExpiresAt = expiresAt;
+			}
+
+			public void Dispose()
 			{
-				_server = server;
-				_userName = userName;
-				_ticket = ticket;
+				Perforce.Dispose();
+			}
+
+			public async ValueTask<InfoRecord> GetInfoAsync(CancellationToken cancellationToken)
+			{
+				_info ??= await Perforce.GetInfoAsync(InfoOptions.ShortOutput, cancellationToken);
+				return _info;
+			}
+		}
+
+		class PooledConnectionHandle : IPerforceConnection
+		{
+			readonly PerforceService _owner;
+			PooledConnection _inner;
+
+			public ClientRecord? Client => _inner.Client;
+			public Credentials Credentials => _inner.Credentials;
+
+			public PooledConnectionHandle(PerforceService owner, PooledConnection inner)
+			{
+				_owner = owner;
+				_inner = inner;
+			}
+
+			/// <inheritdoc/>
+			public void Dispose()
+			{
+				if (_inner != null)
+				{
+					_owner.ReleasePooledConnection(_inner);
+					_inner = null!;
+				}
+			}
+
+			/// <inheritdoc/>
+			public IPerforceSettings Settings => _inner.Perforce.Settings;
+
+			/// <inheritdoc/>
+			public ILogger Logger => _inner.Perforce.Logger;
+
+			/// <inheritdoc/>
+			public IPerforceOutput Command(string command, IReadOnlyList<string> arguments, IReadOnlyList<string>? fileArguments, byte[]? inputData, string? promptResponse, bool interceptIo)
+			{
+				return _inner.Perforce.Command(command, arguments, fileArguments, inputData, promptResponse, interceptIo);
+			}
+
+			/// <inheritdoc/>
+			public PerforceRecord CreateRecord(List<KeyValuePair<string, object>> fields)
+			{
+				return _inner.Perforce.CreateRecord(fields);
+			}
+
+			public ValueTask<InfoRecord> GetInfoAsync(CancellationToken cancellationToken) => _inner.GetInfoAsync(cancellationToken);
+
+		}
+
+		class Credentials
+		{
+			public string UserName { get; }
+			public string? Password { get; }
+			public DateTime? ExpiresAt { get; }
+
+			public Credentials(string userName, string? password, DateTime? expiresAt)
+			{
+				UserName = userName;
+				Password = password;
+				ExpiresAt = expiresAt;
 			}
 		}
 
 		readonly PerforceLoadBalancer _loadBalancer;
 		readonly LazyCachedValue<Task<Globals>> _cachedGlobals;
 		readonly ILogger _logger;
-		readonly bool _haveNativePerforceApi;
 
 		// Useful overrides for local debugging with read-only data
 		readonly string? _perforceServerOverride;
 		readonly string? _perforceUserOverride;
 
-		/// <summary>
-		/// Object used for controlling access to the access user tickets
-		/// </summary>
-		static readonly object s_ticketLock = new object();
-
-		/// <summary>
-		/// Object used for controlling access to the p4 command output
-		/// </summary>
-		static readonly object s_p4LogLock = new object();
-
-		/// <summary>
-		/// Native -> managed debug logging callback
-		/// </summary>
-		readonly P4.P4CallBacks.LogMessageDelegate _logBridgeDelegate;
-
-		/// <summary>
-		/// The server settings
-		/// </summary>
 		readonly ServerSettings _settings;
-		readonly Dictionary<string, Dictionary<string, CachedTicketInfo>> _clusterTickets = new Dictionary<string, Dictionary<string, CachedTicketInfo>>(StringComparer.OrdinalIgnoreCase);
+		readonly Dictionary<string, Dictionary<string, Credentials>> _userCredentialsByCluster = new Dictionary<string, Dictionary<string, Credentials>>(StringComparer.OrdinalIgnoreCase);
 		readonly IUserCollection _userCollection;
 		readonly MemoryCache _userCache = new MemoryCache(new MemoryCacheOptions { SizeLimit = 2000 });
 
+		readonly List<PooledConnection> _pooledConnections = new List<PooledConnection>();
+
 		/// <summary>
 		/// Constructor
 		/// </summary>
@@ -102,15 +150,6 @@ namespace Horde.Build.Perforce
 			_settings = settings.Value;
 			_logger = logger;
 
-			_logBridgeDelegate = new P4.P4CallBacks.LogMessageDelegate(LogBridgeMessage);
-
-			_haveNativePerforceApi = !(RuntimeInformation.IsOSPlatform(OSPlatform.OSX) && RuntimeInformation.OSArchitecture == Architecture.Arm64);
-			if (_haveNativePerforceApi)
-			{
-				P4.P4Debugging.SetBridgeLogFunction(_logBridgeDelegate);
-				P4.LogFile.SetLoggingFunction(LogPerforce);
-			}
-
 			if(settings.Value.UseLocalPerforceEnv)
 			{
 				IPerforceSettings perforceSettings = PerforceSettings.Default;
@@ -121,120 +160,118 @@ namespace Horde.Build.Perforce
 
 		public void Dispose()
 		{
+			foreach (PooledConnection pooledConnection in _pooledConnections)
+			{
+				pooledConnection.Dispose();
+			}
+			_pooledConnections.Clear();
+
 			_userCache.Dispose();
 		}
 
-		public async ValueTask<IUser> FindOrAddUserAsync(string clusterName, string userName)
+		async Task<PooledConnectionHandle> CreatePooledConnectionAsync(string serverAndPort, Credentials credentials, ClientRecord? clientRecord, CancellationToken cancellationToken)
 		{
-			using IScope scope = GlobalTracer.Instance.BuildSpan("PerforceService.FindOrAddUserAsync").StartActive();
-			scope.Span.SetTag("ClusterName", clusterName);
-			scope.Span.SetTag("UserName", userName);
-			
-			IUser? user;
-			if (!_userCache.TryGetValue((clusterName, userName), out user))
+			_ = cancellationToken;
+
+			PerforceSettings settings = new PerforceSettings(serverAndPort, credentials.UserName);
+			settings.AppName = "Horde.Build";
+			settings.Password = credentials.Password;
+			settings.ClientName = clientRecord?.Name ?? "__DOES_NOT_EXIST__";
+			settings.PreferNativeClient = true;
+
+			DateTime expiresAt = DateTime.UtcNow + TimeSpan.FromHours(1.0);
+			if (credentials.ExpiresAt.HasValue && credentials.ExpiresAt.Value < expiresAt)
 			{
-				user = await _userCollection.FindUserByLoginAsync(userName);
-				if (user == null)
-				{
-					P4.User? userInfo = await GetUserInfoAsync(clusterName, userName);
-					user = await _userCollection.FindOrAddUserByLoginAsync(userName, userInfo?.FullName, userInfo?.EmailAddress);
-				}
+				expiresAt = credentials.ExpiresAt.Value;
+			}
+
+			IPerforceConnection connection = await PerforceConnection.CreateAsync(settings, _logger);
+			PooledConnection pooledConnection = new PooledConnection(connection, credentials, clientRecord, expiresAt);
+			return new PooledConnectionHandle(this, pooledConnection);
+		}
 
-				using (ICacheEntry entry = _userCache.CreateEntry((clusterName, userName)))
+		PooledConnectionHandle? GetPooledConnection(Predicate<PooledConnection> predicate)
+		{
+			lock (_pooledConnections)
+			{
+				ReleaseExpiredConnections();
+				for (int idx = 0; idx < _pooledConnections.Count; idx++)
 				{
-					entry.SetValue(user);
-					entry.SetSize(1);
-					entry.SetAbsoluteExpiration(TimeSpan.FromDays(1.0));
+					PooledConnection connection = _pooledConnections[idx];
+					if (predicate(connection))
+					{
+						_pooledConnections.RemoveAt(idx);
+						return new PooledConnectionHandle(this, connection);
+					}
 				}
+				return null;
 			}
-			return user!;
 		}
 
-		async Task<PerforceCluster> GetClusterAsync(string? clusterName, string? serverAndPort = null)
+		PooledConnectionHandle? GetPooledConnectionForUser(PerforceCluster cluster, string? userName)
 		{
-			Globals globals = await _cachedGlobals.GetCached();
-
-			PerforceCluster? cluster = globals.FindPerforceCluster(clusterName, serverAndPort);
-			if (cluster == null)
-			{
-				throw new Exception($"Unknown Perforce cluster '{clusterName}'");
-			}
+			string? poolUserName = userName ?? GetServiceUserCredentials(cluster).UserName;
+			return GetPooledConnection(x => IsValidServer(x, cluster) && String.Equals(x.Settings.UserName, poolUserName, StringComparison.Ordinal) && x.Client == null);
+		}
 
-			return cluster;
+		PooledConnectionHandle? GetPooledConnectionForClient(PerforceCluster cluster, string? userName, string clientName)
+		{
+			string? poolUserName = userName ?? GetServiceUserCredentials(cluster).UserName;
+			return GetPooledConnection(x => IsValidServer(x, cluster) && String.Equals(x.Settings.UserName, poolUserName, StringComparison.Ordinal) && String.Equals(x.Client?.Name, clientName, StringComparison.Ordinal));
 		}
 
-		async ValueTask<string> SelectServerAddressAsync(PerforceCluster cluster)
+		PooledConnectionHandle? GetPooledConnectionForStream(PerforceCluster cluster, string? userName, string streamName)
 		{
-			return (await SelectServer(cluster)).ServerAndPort;
+			string? poolUserName = userName ?? GetServiceUserCredentials(cluster).UserName;
+			return GetPooledConnection(x => IsValidServer(x, cluster) && String.Equals(x.Settings.UserName, poolUserName, StringComparison.Ordinal) && x.Client != null && String.Equals(x.Client.Stream, streamName, StringComparison.Ordinal));
 		}
 
-		async Task<IPerforceServer> SelectServer(PerforceCluster cluster)
+		bool IsValidServer(PooledConnection connection, PerforceCluster cluster)
 		{
-			using IScope scope = GlobalTracer.Instance.BuildSpan("PerforceService.SelectServer").StartActive();
-			scope.Span.SetTag("ClusterName", cluster.Name);
-			
-			IPerforceServer? server = await _loadBalancer.SelectServerAsync(cluster);
-			if (server == null)
+			string serverAndPort = connection.Settings.ServerAndPort;
+			if (_perforceServerOverride != null)
 			{
-				throw new Exception($"Unable to select server from '{cluster.Name}'");
+				return String.Equals(serverAndPort, _perforceServerOverride, StringComparison.Ordinal);
+			}
+			else
+			{
+				return cluster.Servers.Any(x => x.ServerAndPort.Equals(serverAndPort, StringComparison.Ordinal));
 			}
-			return server;
 		}
 
-		[SuppressMessage("Microsoft.Reliability", "CA2000:DisposeObjectsBeforeLosingScope")]
-		static P4.Repository CreateConnection(IPerforceServer server, string? userName, string? ticket)
+		void ReleaseExpiredConnections()
 		{
-			using IScope scope = GlobalTracer.Instance.BuildSpan("PerforceService.CreateConnection").StartActive();
-			scope.Span.SetTag("Cluster", server.Cluster);
-			scope.Span.SetTag("ServerAndPort", server.ServerAndPort);
-			scope.Span.SetTag("UserName", userName);
-			
-			P4.Repository repository = new P4.Repository(new P4.Server(new P4.ServerAddress(server.ServerAndPort)));
-			try
+			DateTime utcNow = DateTime.UtcNow;
+			for (int idx = _pooledConnections.Count - 1; idx >= 0; idx--)
 			{
-				P4.Connection connection = repository.Connection;
-				if (userName != null)
-				{
-					connection.UserName = userName;
-				}
-
-				P4.Options options = new P4.Options();
-				if (ticket != null)
+				PooledConnection connection = _pooledConnections[idx];
+				if (connection.ExpiresAt < utcNow)
 				{
-					options["Ticket"] = ticket;
+					connection.Dispose();
+					_pooledConnections.RemoveAt(idx);
 				}
+			}
+		}
 
-				// connect to the server
-				if (!connection.Connect(options))
+		void ReleasePooledConnection(PooledConnection connection)
+		{
+			lock (_pooledConnections)
+			{
+				ReleaseExpiredConnections();
+				if (_pooledConnections.Count >= _settings.PerforceConnectionPoolSize)
 				{
-					throw new Exception("Unable to get P4 server connection");
+					_pooledConnections[0].Dispose();
+					_pooledConnections.RemoveAt(0);
 				}
+				_pooledConnections.Add(connection);
 			}
-			catch
-			{
-				repository.Dispose();
-				throw;
-			}
-			return repository;
 		}
 
-		/// <summary>
-		/// 
-		/// </summary>
-		/// <param name="clusterName"></param>
-		/// <returns></returns>
-		public async Task<IPerforceConnection?> GetServiceUserConnection(string? clusterName)
+		Credentials GetServiceUserCredentials(PerforceCluster cluster)
 		{
-			using IScope scope = GlobalTracer.Instance.BuildSpan("PerforceService.GetServiceUserConnection").StartActive();
-			scope.Span.SetTag("ClusterName", clusterName);
-
-			PerforceCluster cluster = await GetClusterAsync(clusterName);
-
-			string? userName, password;
 			if (_perforceUserOverride != null)
 			{
-				userName = _perforceUserOverride;
-				password = null;
+				return new Credentials(_perforceUserOverride, null, null);
 			}
 			else if (cluster.ServiceAccount != null)
 			{
@@ -243,639 +280,376 @@ namespace Horde.Build.Perforce
 				{
 					throw new Exception($"No credentials defined for {cluster.ServiceAccount} on {cluster.Name}");
 				}
-
-				userName = credentials.UserName;
-				password = credentials.Password;
-			}
-			else
-			{
-				userName = PerforceSettings.Default.UserName;
-				password = null;
-			}
-
-			string serverAndPort;
-			if (_perforceServerOverride != null)
-			{
-				serverAndPort = _perforceServerOverride;
+				return new Credentials(credentials.UserName, credentials.Password, null);
 			}
 			else
 			{
-				serverAndPort = await SelectServerAddressAsync(cluster);
-			}
-
-			PerforceSettings settings = new PerforceSettings(serverAndPort, userName);
-			settings.Password = password;
-			settings.AppName = "Horde.Build";
-			settings.ClientName = "__DOES_NOT_EXIST__";
-			settings.PreferNativeClient = true;
-
-			return await PerforceConnection.CreateAsync(settings, _logger);
-		}
-
-		async Task<P4.Repository> GetServiceUserConnection(PerforceCluster cluster)
-		{
-			using IScope scope = GlobalTracer.Instance.BuildSpan("PerforceService.GetServiceUserConnectionAsync").StartActive();
-			scope.Span.SetTag("ClusterName", cluster.Name);
-			
-			IPerforceServer server = await SelectServer(cluster);
-			return GetServiceUserConnection(cluster, server);
-		}
-
-		P4.Repository GetServiceUserConnection(PerforceCluster cluster, IPerforceServer server)
-		{
-			using IScope scope = GlobalTracer.Instance.BuildSpan("PerforceService.GetServiceUserConnection").StartActive();
-			scope.Span.SetTag("ClusterName", cluster.Name);
-
-			if (cluster.Name == PerforceCluster.DefaultName && _settings.P4BridgeServiceUsername != null && _settings.P4BridgeServicePassword != null)
-			{
-				return CreateConnection(server, _settings.P4BridgeServiceUsername, _settings.P4BridgeServicePassword);
-			}
-
-			string? userName = null;
-			string? password = null;
-			if (cluster.ServiceAccount != null)
-			{
-				PerforceCredentials? credentials = cluster.Credentials.FirstOrDefault(x => x.UserName.Equals(cluster.ServiceAccount, StringComparison.OrdinalIgnoreCase));
-				if (credentials == null)
-				{
-					throw new Exception($"No credentials defined for {cluster.ServiceAccount} on {cluster.Name}");
-				}
-				userName = credentials.UserName;
-				password = credentials.Password;
+				return new Credentials(PerforceSettings.Default.UserName, null, null);
 			}
-			return CreateConnection(server, userName, password);
 		}
 
-		async Task<CachedTicketInfo> GetImpersonateCredential(PerforceCluster cluster, string impersonateUser)
+		async Task<Credentials> GetCredentialsAsync(PerforceCluster cluster, string? userName, CancellationToken cancellationToken)
 		{
-			using IScope scope = GlobalTracer.Instance.BuildSpan("PerforceService.GetImpersonateCredential").StartActive();
-			scope.Span.SetTag("ClusterName", cluster.Name);
-			scope.Span.SetTag("ImpersonateUser", impersonateUser);
-			
-			if (!cluster.CanImpersonate)
+			if (_perforceUserOverride != null)
 			{
-				throw new Exception($"Service account required to impersonate user {impersonateUser}");
+				return new Credentials(_perforceUserOverride, null, null);
 			}
-
-			CachedTicketInfo? ticketInfo = null;
-
-			Dictionary<string, CachedTicketInfo>? userTickets;
-			lock (s_ticketLock)
+			else if (userName != null)
 			{
-				// Check if we have a ticket
-				if (!_clusterTickets.TryGetValue(cluster.Name, out userTickets))
+				if (cluster.CanImpersonate && cluster.ServiceAccount != null)
 				{
-					userTickets = new Dictionary<string, CachedTicketInfo>(StringComparer.OrdinalIgnoreCase);
-					_clusterTickets[cluster.Name] = userTickets;
+					return await GetTicketAsync(cluster, userName, cancellationToken);
 				}
-				if (userTickets.TryGetValue(impersonateUser, out ticketInfo))
+				else
 				{
-					// if the credential expires within the next 15 minutes, refresh
-					TimeSpan time = new TimeSpan(0, 15, 0);
-					if (ticketInfo._ticket.Expires.Subtract(time) <= DateTime.UtcNow)
-					{
-						userTickets.Remove(impersonateUser);
-						ticketInfo = null;
-					}
+					return new Credentials(userName, null, null);
 				}
 			}
-
-			if (ticketInfo == null)
+			else
 			{
-				IPerforceServer server = await SelectServer(cluster);
-
-				P4.Credential? credential;
-				using (P4.Repository repository = GetServiceUserConnection(cluster, server))
-				{
-					credential = repository.Connection.Login(null, new P4.LoginCmdOptions(P4.LoginCmdFlags.AllHosts | P4.LoginCmdFlags.DisplayTicket, null), impersonateUser);
-				}
-				if (credential == null)
-				{
-					throw new Exception($"GetImpersonateCredential - Unable to get impersonation credential for {impersonateUser} from {Dns.GetHostName()}");
-				}
-
-				ticketInfo = new CachedTicketInfo(server, impersonateUser, credential);
-
-				lock (s_ticketLock)
-				{
-					userTickets[impersonateUser] = ticketInfo;
-				}
+				return GetServiceUserCredentials(cluster);
 			}
-
-			return ticketInfo;
 		}
 
-		async Task<P4.Repository> GetImpersonatedConnection(PerforceCluster cluster, string impersonateUser)
+		public async Task<IPerforceConnection> ConnectAsync(string clusterName, string? userName = null, CancellationToken cancellationToken = default)
 		{
-			using IScope scope = GlobalTracer.Instance.BuildSpan("PerforceService.GetImpersonatedConnection").StartActive();
-			scope.Span.SetTag("ClusterName", cluster.Name);
-			scope.Span.SetTag("ImpersonateUser", impersonateUser);
-			
-			CachedTicketInfo ticketInfo = await GetImpersonateCredential(cluster, impersonateUser);
-			return CreateConnection(ticketInfo._server, ticketInfo._userName, ticketInfo._ticket.Ticket);
+			PerforceCluster cluster = await GetClusterAsync(clusterName);
+			return await ConnectAsync(cluster, userName, cancellationToken);
 		}
-		
-		async Task<string?> GetPortFromChange(PerforceCluster cluster, int change)
-		{
-			using IScope scope = GlobalTracer.Instance.BuildSpan("PerforceService.GetPortFromChange").StartActive();
-			scope.Span.SetTag("ClusterName", cluster.Name);
-			scope.Span.SetTag("Change", change);
-
-			string? clientPort = null;
-
-			using (P4.Repository repository = await GetConnection(cluster, noClient: true))
-			{				
-				P4.Changelist changelist = repository.GetChangelist(change);
-				if (changelist == null || String.IsNullOrEmpty(changelist.ClientId))
-				{
-					throw new Exception($"GetPortFromChange - Unable to get changelist for {change}");
-				}
-
-				P4.Client client = repository.GetClient(changelist.ClientId);
 
-				if (client == null)
-				{
-					throw new Exception($"GetPortFromChange - Unable to get client for {change}");
-				}
-
-				// Check whether not restricted to a specific server
-				if (String.IsNullOrEmpty(client.ServerID))
-				{
-					return null;
-				}
-
-				List<string> args = new List<string> { "-o", client.ServerID };
-
-				using (P4.P4Command command = new P4.P4Command(repository, "server", true, args.ToArray()))
-				{
-					P4.P4CommandResult result = command.Run();
-
-					if (!result.Success || result.TaggedOutput == null || result.TaggedOutput.Count == 0)
-					{
-						throw new Exception($"Unable to get tagged output for {client.ServerID}");
-					}
-
-					foreach (P4.TaggedObject taggedObject in result.TaggedOutput)
-					{
-						if (taggedObject.TryGetValue("Address", out clientPort))
-						{
-							break;
-						}
-					}
-				}
-			}
-
-			if (String.IsNullOrEmpty(clientPort))
+		async Task<PooledConnectionHandle> ConnectAsync(PerforceCluster cluster, string? userName, CancellationToken cancellationToken)
+		{
+			PooledConnectionHandle? handle = GetPooledConnectionForUser(cluster, userName);
+			if (handle == null)
 			{
-				throw new Exception($"GetPortFromChange - Unable to get client port for {change}");
+				IPerforceServer server = await GetServerAsync(cluster);
+				Credentials credentials = await GetCredentialsAsync(cluster, userName, cancellationToken);
+				handle = await CreatePooledConnectionAsync(server.ServerAndPort, credentials, null, cancellationToken);
 			}
-
-			return clientPort;
+			return handle;
 		}
 
-		[SuppressMessage("Microsoft.Reliability", "CA2000:DisposeObjectsBeforeLosingScope")]
-		async Task<P4.Repository> GetConnection(PerforceCluster cluster, string? stream = null, string? username = null, bool readOnly = true, bool createChange = false, int? clientFromChange = null, bool useClientFromChange = false, int? usePortFromChange = null, string? clientId = null, bool noClient = false)
+		async Task<PooledConnectionHandle> ConnectAsChangeOwnerAsync(PerforceCluster cluster, int change, CancellationToken cancellationToken)
 		{
-			using IScope scope = GlobalTracer.Instance.BuildSpan("PerforceService.GetConnection").StartActive();
-			scope.Span.SetTag("ClusterName", cluster.Name);
-			scope.Span.SetTag("Stream", stream);
-			scope.Span.SetTag("Username", username);
-			scope.Span.SetTag("NoClient", noClient);
-			scope.Span.SetTag("ClientId", clientId);
-
-			if (usePortFromChange.HasValue)
+			using (IPerforceConnection perforce = await ConnectAsync(cluster, null, cancellationToken))
 			{
-				try
+				// Get the change information
+				ChangeRecord changeRecord = await perforce.GetChangeAsync(GetChangeOptions.None, change, cancellationToken);
+				if (changeRecord.Client == null)
 				{
-					string? clientPort = await GetPortFromChange(cluster, usePortFromChange.Value);
-					if (clientPort != null)
-					{
-						PerforceCluster clientCluster = await GetClusterAsync(null, clientPort);
-						if (!String.Equals(clientCluster.Name, cluster.Name, StringComparison.OrdinalIgnoreCase))
-						{
-							_logger.LogInformation("Perforce: Overriding Cluster {Cluster} with Client Cluster {ClientCluster} for server restricted change {Change}", cluster.Name, clientCluster.Name, usePortFromChange.Value);
-							cluster = clientCluster;
-						}
-					}
-					else
-					{
-						_logger.LogInformation("Change {UsePortFromChange} is not restricted to a server", usePortFromChange.Value);
-					}
+					throw new PerforceException("Client not specified on changelist {Change}", changeRecord.Number);
 				}
-				catch (Exception ex)
+				if (changeRecord.User == null)
 				{
-					_logger.LogError(ex, "Unable to get client port for changelist {Change} using cluster {ClusterName}", usePortFromChange, cluster.Name);
+					throw new PerforceException("User not specified on changelist {Change}", changeRecord.Number);
 				}
-			}
-
-			P4.Repository repository;
-			if (username == null || !cluster.CanImpersonate || username.Equals(cluster.ServiceAccount, StringComparison.OrdinalIgnoreCase))
-			{
-				repository = await GetServiceUserConnection(cluster);
-			}
-			else
-			{
-				repository = await GetImpersonatedConnection(cluster, username);
-			}
 
-			if (!noClient)
-			{
-				try
+				// Check if there's an existing connection for this client and user. We can forgo any server checks if the client matches, because we know it'll be on the same edge server.
+				PooledConnectionHandle? handle = GetPooledConnectionForClient(cluster, changeRecord.User, changeRecord.Client);
+				if (handle != null)
 				{
-					if (clientId != null)
-					{
-						repository.Connection.SetClient(clientId);
-					}
-					else
-					{
-						P4.Client client = GetOrCreateClient(cluster.ServiceAccount, repository, stream, username, readOnly, createChange, clientFromChange, useClientFromChange);
-						repository.Connection.SetClient(client.Name);
+					return handle;
+				}
 
-					}
+				// Get the client that owns this change
+				ClientRecord clientRecord = await perforce.GetClientAsync(changeRecord.Client, cancellationToken);
 
-					if (useClientFromChange)
+				// Get the server address to connect as, if it's locked to an edge server
+				string? serverAndPort = null;
+				if (!String.IsNullOrEmpty(clientRecord.ServerId))
+				{
+					ServerRecord serverRecord = await perforce.GetServerAsync(clientRecord.ServerId, cancellationToken);
+					if (!String.IsNullOrEmpty(serverRecord.Address))
 					{
-						string? clientHost = repository.Connection.Client.Host;
-						if (!String.IsNullOrEmpty(clientHost))
-						{
-							repository.Connection.getP4Server().SetConnectionHost(clientHost);
-						}
-					}					
+						serverAndPort = serverRecord.Address;
+					}
 				}
-				catch
+
+				// Otherwise connect to the default
+				if (serverAndPort == null)
 				{
-					repository.Dispose();
-					throw;
+					IPerforceServer server = await GetServerAsync(cluster);
+					serverAndPort = server.ServerAndPort;
 				}
-			}
 
-			if (repository == null)
-			{
-				throw new Exception($"Unable to get connection for Stream:{stream} Username:{username} ReadOnly:{readOnly} CreateChange:{createChange} ClientFromChange:{clientFromChange} UseClientFromChange:{useClientFromChange} UsePortFromChange:{usePortFromChange}");
+				// Create a new connection
+				Credentials credentials = await GetCredentialsAsync(cluster, changeRecord.User, cancellationToken);
+				return await CreatePooledConnectionAsync(serverAndPort, credentials, clientRecord, cancellationToken);
 			}
-
-			repository.Connection.CommandEcho += LogPerforceCommand;
-
-			return repository;
-
 		}
 
-		static string GetClientName(string? serviceUserName, string serverName, string stream, bool readOnly, bool createChange, string? username = null)
+		async Task<PooledConnectionHandle> ConnectWithStreamClientAsync(PerforceCluster cluster, string? userName, string streamName, CancellationToken cancellationToken)
 		{
-			string clientName = $"zzt-horde-p4bridge-{serverName}-{Dns.GetHostName()}-{serviceUserName ?? "default"}-{stream.Replace("/", "+", StringComparison.OrdinalIgnoreCase)}";
-
-			if (!readOnly)
+			PooledConnectionHandle? handle = GetPooledConnectionForStream(cluster, userName, streamName);
+			if (handle != null)
 			{
-				clientName += "-write";
-			}
-
-			if (createChange)
-			{
-				clientName += "-create";
+				return handle;
 			}
 
-			if (!String.IsNullOrEmpty(username))
+			using (PooledConnectionHandle perforce = await ConnectAsync(cluster, userName, cancellationToken))
 			{
-				clientName += "-" + username;
-			}
-
-			return clientName;
+				InfoRecord info = await perforce.GetInfoAsync(cancellationToken);
+				string clientName = GetClientName(info.UserName, info.ServerId ?? "unknown", streamName, readOnly: true);
 
-		}
+				string root = RuntimeInformation.IsOSPlatform(OSPlatform.Linux) ? $"/tmp/{clientName}/" : $"{Path.GetTempPath()}{clientName}\\";
 
-		static void ResetClient(P4.Repository repository)
-		{
-			repository.Connection.Client.RevertFiles(new P4.Options(), new P4.DepotPath("//..."));
+				ClientRecord client = new ClientRecord(clientName, info.UserName ?? String.Empty, root);
+				client.Stream = streamName;
+				await perforce.CreateClientAsync(client, cancellationToken);
 
-			IList<P4.Changelist> changes = repository.GetChangelists(new P4.ChangesCmdOptions(P4.ChangesCmdFlags.None, repository.Connection.Client.Name, 100, P4.ChangeListStatus.Pending, null));
-			if (changes != null)
-			{
-				foreach (P4.Changelist change in changes)
-				{
-					repository.DeleteChangelist(change, new P4.Options());
-				}
+				return await CreatePooledConnectionAsync(perforce.Settings.ServerAndPort, perforce.Credentials, client, cancellationToken);
 			}
 		}
 
-		static P4.Client GetOrCreateClient(string? serviceUserName, P4.Repository repository, string? stream, string? username = null, bool readOnly = true, bool createChange = false, int? clientFromChange = null, bool useClientFromChange = false)
+		async Task<Credentials> GetTicketAsync(PerforceCluster cluster, string userName, CancellationToken cancellationToken)
 		{
-			using IScope scope = GlobalTracer.Instance.BuildSpan("PerforceService.GetOrCreateClient").StartActive();
-			scope.Span.SetTag("ServiceUserName", serviceUserName);
-			scope.Span.SetTag("Stream", stream);
-			scope.Span.SetTag("Username", username);
-			
-			P4.Client? client = null;
-			P4.Changelist? changelist = null;
+			using IScope scope = GlobalTracer.Instance.BuildSpan("PerforceService.TryGetImpersonatedConnectionSettings").StartActive();
+			scope.Span.SetTag("ClusterName", cluster.Name);
+			scope.Span.SetTag("UserName", userName);
 
-			string? clientHost = null;
+			Credentials? ticketInfo = null;
 
-			if (clientFromChange.HasValue)
+			Dictionary<string, Credentials>? userTickets;
+			lock (_userCredentialsByCluster)
 			{
-				changelist = repository.GetChangelist(clientFromChange.Value);
-
-				if (changelist == null)
+				// Check if we have a ticket
+				if (!_userCredentialsByCluster.TryGetValue(cluster.Name, out userTickets))
 				{
-					throw new Exception($"Unable to get changelist for client {clientFromChange}");
+					userTickets = new Dictionary<string, Credentials>(StringComparer.OrdinalIgnoreCase);
+					_userCredentialsByCluster[cluster.Name] = userTickets;
 				}
-
-				client = repository.GetClient(changelist.ClientId);
-
-				if (client == null)
+				if (userTickets.TryGetValue(userName, out ticketInfo))
 				{
-					throw new Exception($"Unable to get client for id {changelist.ClientId}");
+					// if the credential expires within the next 15 minutes, refresh
+					if (DateTime.UtcNow > ticketInfo.ExpiresAt)
+					{
+						userTickets.Remove(userName);
+						ticketInfo = null;
+					}
 				}
-
-				stream = client.Stream;
-				username = client.OwnerName;
-				clientHost = client.Host;
-
 			}
 
-			if (!useClientFromChange)
+			if (ticketInfo == null)
 			{
-				if (stream == null)
+				PerforceResponse<LoginRecord> response;
+				using (IPerforceConnection connection = await ConnectAsync(cluster, null, cancellationToken: cancellationToken))
 				{
-					throw new Exception("Stream required for client");
-
+					response = await connection.TryLoginAsync(LoginOptions.AllHosts | LoginOptions.PrintTicket, userName, null, null, cancellationToken);
 				}
-				string clientName = GetClientName(serviceUserName, repository.Server.Metadata.Name, stream, readOnly, createChange, username);
-
-				IList<P4.Client>? clients = repository.GetClients(new P4.ClientsCmdOptions(P4.ClientsCmdFlags.None, username, clientName, 1, stream));
-
-				if (clients != null && clients.Count == 1)
+				if (!response.Succeeded || response.Data.Ticket == null)
 				{
-					client = clients[0];
-					if (!String.IsNullOrEmpty(client.Host))
-					{
-						client.Host = String.Empty;
-						repository.UpdateClient(client);
-					}
+					throw new PerforceException($"Unable to get impersonation credentials for {userName} for cluster {cluster.Name}");
 				}
-				else
-				{
 
-					P4.Client newClient = new P4.Client();
+				DateTime expiresAt = DateTime.UtcNow + new TimeSpan(response.Data.TicketExpiration * TimeSpan.TicksPerSecond) - TimeSpan.FromMinutes(15.0);
+				ticketInfo = new Credentials(userName, response.Data.Ticket, expiresAt);
 
-					newClient.Name = clientName;
-					newClient.Root = RuntimeInformation.IsOSPlatform(OSPlatform.Linux) ? $"/tmp/{clientName}/" : $"{Path.GetTempPath()}{clientName}\\";
-					newClient.Stream = stream;
-					if (readOnly)
-					{
-						newClient.ClientType = P4.ClientType.@readonly;
-					}
+				lock (_userCredentialsByCluster)
+				{
+					userTickets[userName] = ticketInfo;
+				}
+			}
 
-					newClient.OwnerName = username ?? serviceUserName;
+			return ticketInfo;
+		}
 
-					client = repository.CreateClient(newClient);
+		public async ValueTask<IUser> FindOrAddUserAsync(string clusterName, string userName, CancellationToken cancellationToken = default)
+		{
+			PerforceCluster cluster = await GetClusterAsync(clusterName);
+			return await FindOrAddUserAsync(cluster, userName, cancellationToken);
+		}
 
-					if (client == null)
+		public async ValueTask<IUser> FindOrAddUserAsync(PerforceCluster cluster, string userName, CancellationToken cancellationToken = default)
+		{
+			using IScope scope = GlobalTracer.Instance.BuildSpan("PerforceService.FindOrAddUserAsync").StartActive();
+			scope.Span.SetTag("ClusterName", cluster.Name);
+			scope.Span.SetTag("UserName", userName);
+			
+			IUser? user;
+			if (!_userCache.TryGetValue((cluster.Name, userName), out user))
+			{
+				user = await _userCollection.FindUserByLoginAsync(userName);
+				if (user == null)
+				{
+					UserRecord? userRecord = null;
+					using (IPerforceConnection perforce = await ConnectAsync(cluster, null, cancellationToken))
 					{
-						throw new Exception($"Unable to create client for ${stream} : {username ?? serviceUserName}");
+						PerforceResponseList<UserRecord> responses = await perforce.TryGetUserAsync(userName, cancellationToken);
+						if (responses.Count > 0 && responses[0].Succeeded)
+						{
+							userRecord = responses[0].Data;
+						}
+						else
+						{
+							_logger.LogWarning("Unable to find user {UserName} on cluster {ClusterName}", userName, cluster.Name);
+						}
 					}
+					user = await _userCollection.FindOrAddUserByLoginAsync(userRecord?.UserName ?? userName, userRecord?.FullName, userRecord?.Email);
 				}
-			}
 
-			if (client == null)
-			{
-				throw new Exception($"Unable to create client for Stream:{stream} Username:{username} ReadOnly:{readOnly} CreateChange:{createChange} ClientFromChange:{clientFromChange} UseClientFromChange:{useClientFromChange}");
+				using (ICacheEntry entry = _userCache.CreateEntry((cluster.Name, userName)))
+				{
+					entry.SetValue(user);
+					entry.SetSize(1);
+					entry.SetAbsoluteExpiration(TimeSpan.FromDays(1.0));
+				}
 			}
-
-			return client;
-
+			return user!;
 		}
 
-		/// <summary>
-		/// Logs perforce commands 
-		/// </summary>
-		/// <param name="log">The p4 command log info</param>
-		static void LogPerforceCommand(string log)
+		async Task<PerforceCluster> GetClusterAsync(string? clusterName, string? serverAndPort = null)
 		{
-			lock (s_p4LogLock)
+			Globals globals = await _cachedGlobals.GetCached();
+
+			PerforceCluster? cluster = globals.FindPerforceCluster(clusterName, serverAndPort);
+			if (cluster == null)
 			{
-				Serilog.Log.Information("Perforce: {Message}", log);
+				throw new Exception($"Unknown Perforce cluster '{clusterName}'");
 			}
+
+			return cluster;
 		}
 
-		/// <summary>
-		/// Perforce logging funcion
-		/// </summary>
-		/// <param name="logLevel">The level whether error, warning, message, or debug</param>
-		/// <param name="source">The log source</param>
-		/// <param name="message">The log message</param>
-		static void LogPerforce(int logLevel, string source, string message)
+		async Task<IPerforceServer> GetServerAsync(PerforceCluster cluster)
 		{
-			lock (s_p4LogLock)
+			using IScope scope = GlobalTracer.Instance.BuildSpan("PerforceService.SelectServer").StartActive();
+			scope.Span.SetTag("ClusterName", cluster.Name);
+
+			IPerforceServer? server = await _loadBalancer.SelectServerAsync(cluster);
+			if (server == null)
 			{
-				switch (logLevel)
-				{
-					case 0:
-					case 1:
-						Serilog.Log.Error("Perforce (Error): {Message} {Source} : ", message, source);
-						break;
-					case 2:
-						Serilog.Log.Warning("Perforce (Warning): {Message} {Source} : ", message, source);
-						break;
-					case 3:
-						Serilog.Log.Information("Perforce: {Message} {Source} : ", message, source);
-						break;
-					default:
-						Serilog.Log.Debug("Perforce (Debug): {Message} {Source} : ", message, source);
-						break;
-				}
+				throw new Exception($"Unable to select server from '{cluster.Name}'");
 			}
+			return server;
 		}
 
-		static void LogBridgeMessage(int logLevel, string filename, int line, string message)
+		static string GetClientName(string? serviceUserName, string serverId, string stream, bool readOnly)
 		{
+			string clientName = $"zzt-horde-p4bridge-{serverId}-{Dns.GetHostName()}-{serviceUserName ?? "default"}-{stream.Replace("/", "+", StringComparison.OrdinalIgnoreCase)}";
 
-			// Note, we do not get log level 4 unless it is defined in native code as it is very, very spammy (P4BridgeServer.cpp)
-
-			// remove the full path to the source, keep just the file name
-			string fileName = Path.GetFileName(filename);
-
-			string category = String.Format(CultureInfo.CurrentCulture, "P4Bridge({0}:{1})", fileName, line);
+			if (!readOnly)
+			{
+				clientName += "-write";
+			}
 
-			P4.LogFile.LogMessage(logLevel, category, message);
+			return clientName;
 		}
 
-		static int GetSyncRevision(string path, P4.FileAction headAction, int headRev)
+		static int GetSyncRevision(string path, FileAction headAction, int headRev)
 		{
 			switch (headAction)
 			{
-				case P4.FileAction.None:
-				case P4.FileAction.Add:
-				case P4.FileAction.Branch:
-				case P4.FileAction.MoveAdd:
-				case P4.FileAction.Edit:
-				case P4.FileAction.Integrate:
+				case FileAction.None:
+				case FileAction.Add:
+				case FileAction.Branch:
+				case FileAction.MoveAdd:
+				case FileAction.Edit:
+				case FileAction.Integrate:
 					return headRev;
-				case P4.FileAction.Delete:
-				case P4.FileAction.MoveDelete:
-				case P4.FileAction.Purge:
+				case FileAction.Delete:
+				case FileAction.MoveDelete:
+				case FileAction.Purge:
 					return -1;
 				default:
 					throw new Exception($"Unrecognized P4 file change type '{headAction}' for file {path}#{headRev}");
 			}
 		}
 
-		/// <summary>
-		/// Creates a <see cref="ChangeFile"/> from a <see cref="P4.FileMetaData"/>
-		/// </summary>
-		/// <param name="relativePath"></param>
-		/// <param name="metaData"></param>
-		/// <returns></returns>
-		static ChangeFile CreateChangeFile(string relativePath, P4.FileMetaData metaData)
+		static ChangeFile CreateChangeFile(string relativePath, DescribeFileRecord metaData)
 		{
 			using IScope scope = GlobalTracer.Instance.BuildSpan("PerforceService.CreateChangeFile (FileMetaData)").StartActive();
 			scope.Span.SetTag("RelativePath", relativePath);
 			
-			int revision = GetSyncRevision(metaData.DepotPath.Path, metaData.HeadAction, metaData.HeadRev);
-			Md5Hash? digest = String.IsNullOrEmpty(metaData.Digest) ? (Md5Hash?)null : Md5Hash.Parse(metaData.Digest);
-			return new ChangeFile(relativePath, metaData.DepotPath.Path, revision, metaData.FileSize, digest, (metaData.Type ?? metaData.HeadType).ToString());
-		}
-
-		/// <summary>
-		/// Creates a <see cref="ChangeFile"/> from a <see cref="P4.FileMetaData"/>
-		/// </summary>
-		/// <param name="relativePath"></param>
-		/// <param name="metaData"></param>
-		/// <returns></returns>
-		static ChangeFile CreateChangeFile(string relativePath, P4.ShelvedFile metaData)
-		{
-			using IScope scope = GlobalTracer.Instance.BuildSpan("PerforceService.CreateChangeFile (ShelvedFile)").StartActive();
-			scope.Span.SetTag("RelativePath", relativePath);
-			
-			int revision = GetSyncRevision(metaData.Path.Path, metaData.Action, metaData.Revision);
+			int revision = GetSyncRevision(metaData.DepotFile, metaData.Action, metaData.Revision);
 			Md5Hash? digest = String.IsNullOrEmpty(metaData.Digest) ? (Md5Hash?)null : Md5Hash.Parse(metaData.Digest);
-			return new ChangeFile(relativePath, metaData.Path.Path, revision, metaData.Size, digest, metaData.Type.ToString());
+			return new ChangeFile(relativePath, metaData.DepotFile, revision, metaData.FileSize, digest, metaData.Type);
 		}
 
 		/// <inheritdoc/>
-		public async Task<List<ChangeSummary>> GetChangesAsync(string clusterName, int? minChange, int? maxChange, int maxResults)
+		public async Task<List<ChangeSummary>> GetChangesAsync(string clusterName, int? minChange, int? maxChange, int maxResults, CancellationToken cancellationToken)
 		{
 			using IScope scope = GlobalTracer.Instance.BuildSpan("PerforceService.GetChangesAsync").StartActive();
 			scope.Span.SetTag("ClusterName", clusterName);
 			scope.Span.SetTag("MinChange", minChange ?? -1);
 			scope.Span.SetTag("MaxResults", maxResults);
-			
-			PerforceCluster cluster = await GetClusterAsync(clusterName);
-			using (P4.Repository repository = await GetConnection(cluster, noClient: true))
-			{
-				P4.ChangesCmdOptions options = new P4.ChangesCmdOptions(P4.ChangesCmdFlags.IncludeTime | P4.ChangesCmdFlags.FullDescription, null, maxResults, P4.ChangeListStatus.Submitted, null);
 
-				IList<P4.Changelist> changelists = repository.GetChangelists(options, new P4.FileSpec(new P4.DepotPath(GetFilter("//...", minChange, maxChange)), null, null, null));
+			List<ChangeSummary> results = new List<ChangeSummary>();
+			using (IPerforceConnection perforce = await ConnectAsync(clusterName, null, cancellationToken))
+			{
+				string filter = GetFilter("//...", minChange, maxChange);
+				List<ChangesRecord> changes = await perforce.GetChangesAsync(ChangesOptions.IncludeTimes | ChangesOptions.LongOutput, maxResults, ChangeStatus.Submitted, filter, cancellationToken);
 
-				List<ChangeSummary> changes = new List<ChangeSummary>();
-				if (changelists != null)
+				foreach (ChangesRecord change in changes)
 				{
-					foreach (P4.Changelist changelist in changelists)
-					{
-						IUser user = await FindOrAddUserAsync(clusterName, changelist.OwnerName);
-						changes.Add(new ChangeSummary(changelist.Id, user, changelist.GetPath(), changelist.Description));
-					}
+					IUser user = await FindOrAddUserAsync(clusterName, change.User, cancellationToken);
+					results.Add(new ChangeSummary(change.Number, user, change.Path ?? String.Empty, change.Description));
 				}
-				return changes;
 			}
+			return results;
 		}
 
 		/// <inheritdoc/>
-		public async Task<List<ChangeSummary>> GetChangesAsync(string clusterName, string streamName, int? minChange, int? maxChange, int results, string? impersonateUser)
+		public async Task<List<ChangeSummary>> GetChangesAsync(string clusterName, string streamName, int? minChange, int? maxChange, int maxResults, string? impersonateUser, CancellationToken cancellationToken)
 		{
 			using IScope scope = GlobalTracer.Instance.BuildSpan("PerforceService.GetChangesAsync").StartActive();
 			scope.Span.SetTag("ClusterName", clusterName);
 			scope.Span.SetTag("MinChange", minChange ?? -1);
 			scope.Span.SetTag("MaxChange", maxChange ?? -1);
-			scope.Span.SetTag("Results", results);
+			scope.Span.SetTag("MaxResults", maxResults);
 			scope.Span.SetTag("ImpersonateUser", impersonateUser);
 			
 			PerforceCluster cluster = await GetClusterAsync(clusterName);
-			using (P4.Repository repository = await GetConnection(cluster, streamName, impersonateUser))
-			{
 
-				P4.ChangesCmdOptions options = new P4.ChangesCmdOptions(P4.ChangesCmdFlags.IncludeTime | P4.ChangesCmdFlags.FullDescription, null, results, P4.ChangeListStatus.Submitted, null);
-
-				string filter = GetFilter($"//{repository.Connection.Client.Name}/...", minChange, maxChange);
-
-				P4.FileSpec fileSpec = new P4.FileSpec(new P4.DepotPath(filter), null, null, null);
-
-				IList<P4.Changelist> changelists = repository.GetChangelists(options, fileSpec);
+			List<ChangeSummary> results = new List<ChangeSummary>();
+			using (IPerforceConnection perforce = await ConnectWithStreamClientAsync(cluster, impersonateUser, streamName, cancellationToken))
+			{
+				string filter = GetFilter($"//{perforce.Settings.ClientName}/...", minChange, maxChange);
 
-				List<ChangeSummary> changes = new List<ChangeSummary>();
+				List<ChangesRecord> changes = await perforce.GetChangesAsync(ChangesOptions.IncludeTimes | ChangesOptions.LongOutput, maxResults, ChangeStatus.Submitted, filter, cancellationToken);
 
-				if (changelists != null)
+				foreach (ChangesRecord change in changes)
 				{
-					foreach (P4.Changelist changelist in changelists)
-					{
-						IUser user = await FindOrAddUserAsync(clusterName, changelist.OwnerName);
-						changes.Add(new ChangeSummary(changelist.Id, user, changelist.GetPath(), changelist.Description));
-					}
+					IUser user = await FindOrAddUserAsync(clusterName, change.User, cancellationToken);
+					results.Add(new ChangeSummary(change.Number, user, change.Path ?? String.Empty, change.Description));
 				}
-
-				return changes;
 			}
+			return results;
 		}
 
 		/// <inheritdoc/>
-		public async Task<ChangeDetails> GetChangeDetailsAsync(string clusterName, string streamName, int changeNumber)
+		public async Task<ChangeDetails> GetChangeDetailsAsync(string clusterName, string streamName, int changeNumber, CancellationToken cancellationToken)
 		{
 			using IScope scope = GlobalTracer.Instance.BuildSpan("PerforceService.GetChangeDetailsAsync").StartActive();
 			scope.Span.SetTag("ClusterName", clusterName);
 			scope.Span.SetTag("StreamName", streamName);
-			scope.Span.SetTag("ChangeNumber", changeNumber);
-			
-			PerforceCluster cluster = await GetClusterAsync(clusterName);
-			using (P4.Repository repository = await GetConnection(cluster, noClient: true))
-			{
-				List<ChangeDetails> results = new List<ChangeDetails>();
-
-				List<string> args = new List<string> { "-s", $"{changeNumber}" };
-
-				using (P4.P4Command command = new P4.P4Command(repository, "describe", true, args.ToArray()))
-				{
-					P4.P4CommandResult result = command.Run();
-
-					if (!result.Success || result.TaggedOutput == null || result.TaggedOutput.Count == 0)
-					{
-						throw new Exception("Unable to get changes");
-					}
-
-					List<P4.Changelist> changelists = new List<P4.Changelist>() { };
+			scope.Span.SetTag("ChangeNumber", changeNumber);
+			
+			PerforceCluster cluster = await GetClusterAsync(clusterName);
 
-					bool dstMismatch = false;
-					string offset = String.Empty;
+			using (PooledConnectionHandle perforce = await ConnectAsync(cluster, null, cancellationToken))
+			{
+				DescribeRecord record = await perforce.DescribeAsync(changeNumber, cancellationToken);
 
-					P4.Server server = repository.Server;
-					if (server != null && server.Metadata != null)
-					{
-						offset = server.Metadata.DateTimeOffset;
-						dstMismatch = P4.FormBase.DSTMismatch(server.Metadata);
-					}
+				InfoRecord info = await perforce.GetInfoAsync(cancellationToken);
 
-					P4.TaggedObject taggedObject = result.TaggedOutput[0];
+				List<ChangeFile> files = new List<ChangeFile>();
+				foreach (DescribeFileRecord describeFile in record.Files)
+				{
+					string? relativePath;
+					if (TryGetStreamRelativePath(describeFile.DepotFile, streamName, out relativePath))
 					{
-						List<ChangeFile> files = new List<ChangeFile>();
-
-						P4.Changelist change = new P4.Changelist();
-						change.FromChangeCmdTaggedOutput(taggedObject, false, offset, dstMismatch);
-
-						foreach (P4.FileMetaData describeFile in change.Files)
-						{
-							string? relativePath;
-							if (TryGetStreamRelativePath(describeFile.DepotPath.Path, streamName, out relativePath))
-							{
-								files.Add(CreateChangeFile(relativePath, describeFile));
-							}
-						}
-
-						IUser user = await FindOrAddUserAsync(clusterName, change.OwnerName);
-						return new ChangeDetails(change.Id, user, change.GetPath(), change.Description, files, change.ModifiedDate);
+						files.Add(CreateChangeFile(relativePath, describeFile));
 					}
 				}
+
+				IUser user = await FindOrAddUserAsync(cluster, record.User, cancellationToken);
+				DateTime timeUtc = new DateTime(record.Time.Ticks - info.TimeZoneOffsetSecs * TimeSpan.TicksPerSecond, DateTimeKind.Utc);
+				return new ChangeDetails(changeNumber, user, record.Path ?? String.Empty, record.Description, files, timeUtc);
 			}
 		}
 
 		/// <inheritdoc/>
-		public async Task<(CheckShelfResult, string?)> CheckShelfAsync(string clusterName, string streamName, int changeNumber, string? impersonateUser)
+		public async Task<(CheckShelfResult, string?)> CheckShelfAsync(string clusterName, string streamName, int changeNumber, string? impersonateUser, CancellationToken cancellationToken)
 		{
 			using IScope scope = GlobalTracer.Instance.BuildSpan("PerforceService.CheckPreflightAsync").StartActive();
 			scope.Span.SetTag("ClusterName", clusterName);
@@ -884,26 +658,37 @@ namespace Horde.Build.Perforce
 			scope.Span.SetTag("ImpersonateUser", impersonateUser);
 
 			PerforceCluster cluster = await GetClusterAsync(clusterName);
-			using (P4.Repository repository = await GetConnection(cluster, stream: streamName, username: impersonateUser))
+			IPerforceServer server = await GetServerAsync(cluster);
+
+			using (IPerforceConnection perforce = await ConnectAsync(cluster, impersonateUser, cancellationToken))
 			{
-				P4.Changelist change = repository.GetChangelist(changeNumber, new P4.DescribeCmdOptions(P4.DescribeChangelistCmdFlags.Omit | P4.DescribeChangelistCmdFlags.Shelved, 0, 0));
-				if(change == null)
+				PerforceResponse<DescribeRecord> response = await perforce.TryDescribeAsync(DescribeOptions.Shelved, -1, changeNumber, cancellationToken);
+				if(response.Error != null)
 				{
-					return (CheckShelfResult.NoChange, null);
+					if (response.Error.Generic == PerforceGenericCode.Empty)
+					{
+						return (CheckShelfResult.NoChange, null);
+					}
+					else
+					{
+						response.EnsureSuccess();
+					}
 				}
-				if (change.ShelvedFiles == null || change.ShelvedFiles.Count == 0)
+
+				DescribeRecord change = response.Data;
+				if (change.Files.Count == 0)
 				{
 					return (CheckShelfResult.NoShelvedFiles, null);
 				}
 
-				P4.Stream stream = repository.GetStream(streamName, new P4.StreamCmdOptions(P4.StreamCmdFlags.View, null, null));
-				ViewMap map = new ViewMap(stream.View);
+				StreamRecord stream = await perforce.GetStreamAsync(streamName, true, cancellationToken);
+				PerforceViewMap viewMap = PerforceViewMap.Parse(stream.View);
 
 				bool bHasMappedFile = false;
 				bool bHasUnmappedFile = false;
-				foreach (P4.ShelvedFile shelvedFile in change.ShelvedFiles)
+				foreach (DescribeFileRecord shelvedFile in change.Files)
 				{
-					if (map.TryMapFile(shelvedFile.Path.Path, Utf8StringComparer.OrdinalIgnoreCase, out Utf8String _))
+					if (viewMap.TryMapFile(shelvedFile.DepotFile, StringComparison.OrdinalIgnoreCase, out string _))
 					{
 						bHasMappedFile = true;
 					}
@@ -923,7 +708,7 @@ namespace Horde.Build.Perforce
 		}
 
 		/// <inheritdoc/>
-		public async Task<List<ChangeDetails>> GetChangeDetailsAsync(string clusterName, string streamName, IReadOnlyList<int> changeNumbers, string? impersonateUser)
+		public async Task<List<ChangeDetails>> GetChangeDetailsAsync(string clusterName, string streamName, IReadOnlyList<int> changeNumbers, string? impersonateUser, CancellationToken cancellationToken)
 		{
 			using IScope scope = GlobalTracer.Instance.BuildSpan("PerforceService.GetChangeDetailsAsync").StartActive();
 			scope.Span.SetTag("ClusterName", clusterName);
@@ -932,120 +717,65 @@ namespace Horde.Build.Perforce
 			scope.Span.SetTag("ImpersonateUser", impersonateUser);
 			
 			PerforceCluster cluster = await GetClusterAsync(clusterName);
-			using (P4.Repository repository = await GetConnection(cluster, stream: streamName, username: impersonateUser))
-			{
-				List<ChangeDetails> results = new List<ChangeDetails>();
 
-				List<string> args = new List<string> { "-s", "-S" };
-				args.AddRange(changeNumbers.Select(change => change.ToString(CultureInfo.InvariantCulture)));
+			List<ChangeDetails> results = new List<ChangeDetails>();
+			using (PooledConnectionHandle perforce = await ConnectWithStreamClientAsync(cluster, impersonateUser, streamName, cancellationToken))
+			{
+				InfoRecord info = await perforce.GetInfoAsync(cancellationToken);
 
-				using (P4.P4Command command = new P4.P4Command(repository, "describe", true, args.ToArray()))
+				List<DescribeRecord> records = await perforce.DescribeAsync(DescribeOptions.Shelved, -1, changeNumbers.ToArray(), cancellationToken);
+				foreach (DescribeRecord record in records)
 				{
-					P4.P4CommandResult result = command.Run();
-
-					if (!result.Success)
-					{
-						throw new Exception("Unable to get changes");
-					}
-
-					if (result.TaggedOutput == null || result.TaggedOutput.Count <= 0)
-					{
-						return results;
-					}
-
-					List<P4.Changelist> changelists = new List<P4.Changelist>() { };
-
-					bool dstMismatch = false;
-					string offset = String.Empty;
-
-					P4.Server server = repository.Server;
-					if (server != null && server.Metadata != null)
-					{
-						offset = server.Metadata.DateTimeOffset;
-						dstMismatch = P4.FormBase.DSTMismatch(server.Metadata);
-					}
-
-					foreach (P4.TaggedObject taggedObject in result.TaggedOutput)
+					List<ChangeFile> files = new List<ChangeFile>();
+					foreach (DescribeFileRecord describeFile in record.Files)
 					{
-						List<ChangeFile> files = new List<ChangeFile>();
-
-						P4.Changelist change = new P4.Changelist();
-						change.FromChangeCmdTaggedOutput(taggedObject, true, offset, dstMismatch);
-
-						foreach (P4.FileMetaData describeFile in change.Files)
+						string? relativePath;
+						if (TryGetStreamRelativePath(describeFile.DepotFile, streamName, out relativePath))
 						{
-							string? relativePath;
-							if (TryGetStreamRelativePath(describeFile.DepotPath.Path, streamName, out relativePath))
-							{
-								files.Add(CreateChangeFile(relativePath, describeFile));
-							}
+							files.Add(CreateChangeFile(relativePath, describeFile));
 						}
+					}
 
-						if (change.ShelvedFiles != null && change.ShelvedFiles.Count > 0)
-						{
-							foreach (P4.ShelvedFile shelvedFile in change.ShelvedFiles)
-							{
-								string? relativePath;
-								if (TryGetStreamRelativePath(shelvedFile.Path.ToString(), streamName, out relativePath))
-								{
-									files.Add(CreateChangeFile(relativePath, shelvedFile));
-								}
-							}
-						}
+					IUser user = await FindOrAddUserAsync(cluster, record.User, cancellationToken);
+					DateTime timeUtc = new DateTime(record.Time.Ticks - info.TimeZoneOffsetSecs * TimeSpan.TicksPerSecond, DateTimeKind.Utc);
 
-						IUser user = await FindOrAddUserAsync(clusterName, change.OwnerName);
-						results.Add(new ChangeDetails(change.Id, user, change.GetPath(), change.Description, files, change.ModifiedDate));
-					}
+					results.Add(new ChangeDetails(record.Number, user, record.Path ?? String.Empty, record.Description, files, timeUtc));
 				}
-
-				return results;
 			}
+			return results;
 		}
 
 		/// <inheritdoc/>
-		public async Task<List<FileSummary>> FindFilesAsync(string clusterName, IEnumerable<string> paths)
+		public async Task<List<FileSummary>> FindFilesAsync(string clusterName, IEnumerable<string> paths, CancellationToken cancellationToken)
 		{
 			using IScope scope = GlobalTracer.Instance.BuildSpan("PerforceService.FindFilesAsync").StartActive();
 			scope.Span.SetTag("ClusterName", clusterName);
 			
-			List<FileSummary> results = new List<FileSummary>();
-
 			PerforceCluster cluster = await GetClusterAsync(clusterName);
-			using (P4.Repository repository = await GetConnection(cluster, noClient: true))
-			{
-				P4.GetFileMetaDataCmdOptions options = new P4.GetFileMetaDataCmdOptions(P4.GetFileMetadataCmdFlags.ExcludeClientData, "", "", 0, "", "", "");
-
-				IList<P4.FileMetaData>? files = repository.GetFileMetaData(options, paths.Select(path => 
-				{ 
-					P4.FileSpec fileSpec = new P4.FileSpec(new P4.DepotPath(path)); 
-					return fileSpec; 
-				}).ToArray());
-
-				if (files == null)
-				{
-					files = new List<P4.FileMetaData>();
-				}
 
+			List<FileSummary> results = new List<FileSummary>();
+			using (IPerforceConnection perforce = await ConnectAsync(cluster, null, cancellationToken))
+			{
+				List<FStatRecord> records = await perforce.FStatAsync(FStatOptions.ShortenOutput, paths.ToArray(), cancellationToken).ToListAsync(cancellationToken);
 				foreach (string path in paths)
 				{
-					P4.FileMetaData? meta = files.FirstOrDefault(file => file.DepotPath?.Path == path);
-
-					if (meta == null)
+					// TODO: use server case
+					FStatRecord? record = records.FirstOrDefault(x => x.DepotFile == path);
+					if (record == null)
 					{
 						results.Add(new FileSummary(path, false, 0));
 					}
 					else
 					{
-						results.Add(new FileSummary(path, meta.HeadAction != P4.FileAction.Delete && meta.HeadAction != P4.FileAction.MoveDelete, meta.HeadChange));
+						results.Add(new FileSummary(path, record.HeadAction != FileAction.Delete && record.HeadAction != FileAction.MoveDelete, record.HeadChange));
 					}
 				}
-
 				return results;
 			}
 		}
 
 		/// <inheritdoc/>
-		public async Task<byte[]> PrintAsync(string clusterName, string depotPath)
+		public async Task<byte[]> PrintAsync(string clusterName, string depotPath, CancellationToken cancellationToken)
 		{
 			using IScope scope = GlobalTracer.Instance.BuildSpan("PerforceService.PrintAsync").StartActive();
 			scope.Span.SetTag("ClusterName", clusterName);
@@ -1057,112 +787,38 @@ namespace Horde.Build.Perforce
 			}
 
 			PerforceCluster cluster = await GetClusterAsync(clusterName);
-			using (P4.Repository repository = await GetConnection(cluster, noClient: true))
-			{
-				using (P4.P4Command command = new P4.P4Command(repository, "print", false, new string[] { "-q", depotPath }))
-				{
-					P4.P4CommandResult result = command.Run();
-
-					if (result.BinaryOutput != null)
-					{
-						return result.BinaryOutput;
-					}
 
-					return Encoding.Default.GetBytes(result.TextOutput);
-				}
+			using (IPerforceConnection perforce = await ConnectAsync(cluster, null, cancellationToken))
+			{
+				PerforceResponse<PrintRecord<byte[]>> response = await perforce.TryPrintDataAsync(depotPath, cancellationToken);
+				return response.Data.Contents!;
 			}
 		}
 
 		/// <inheritdoc/>
-		public async Task<int> DuplicateShelvedChangeAsync(string clusterName, int shelvedChange)
+		public Task<int> DuplicateShelvedChangeAsync(string clusterName, int shelvedChange, CancellationToken cancellationToken)
 		{
-			using IScope scope = GlobalTracer.Instance.BuildSpan("PerforceService.DuplicateShelvedChangeAsync").StartActive();
-			scope.Span.SetTag("ClusterName", clusterName);
-			scope.Span.SetTag("ShelvedChange", shelvedChange);
-
-			string? changeOwner = null;
-
-			// Get the owner of the shelf
-			PerforceCluster cluster = await GetClusterAsync(clusterName);
-			using (P4.Repository repository = await GetConnection(cluster, noClient: true))
-			{
-				List<string> args = new List<string> { "-S", shelvedChange.ToString(CultureInfo.InvariantCulture) };
-
-				using (P4.P4Command command = new P4.P4Command(repository, "describe", true, args.ToArray()))
-				{
-					P4.P4CommandResult result = command.Run();
-
-					if (!result.Success)
-					{
-						throw new Exception($"Unable to get change {shelvedChange}");
-					}
-
-					if (result.TaggedOutput == null || result.TaggedOutput.Count != 1)
-					{
-
-						throw new Exception($"Unable to get tagged output for change: {shelvedChange}");
-					}
-
-					P4.Changelist changelist = new P4.Changelist();
-					changelist.FromChangeCmdTaggedOutput(result.TaggedOutput[0], true, String.Empty, false);
-					changeOwner = changelist.OwnerName;
-
-				}
-			}
-
-			if (String.IsNullOrEmpty(changeOwner))
-			{
-				throw new Exception($"Unable to get owner for shelved change {shelvedChange}");
-			}
-
-			using (P4.Repository repository = await GetConnection(cluster, readOnly: false, usePortFromChange: shelvedChange, clientFromChange: shelvedChange, useClientFromChange: false, username: changeOwner))
-			{
-				return ReshelveChange(repository, shelvedChange);
-			}
-
-			throw new Exception($"Unable to duplicate shelve change {shelvedChange}");
+			throw new NotImplementedException();
 		}
 
 		/// <inheritdoc/>
-		public async Task DeleteShelvedChangeAsync(string clusterName, int shelvedChange)
+		public async Task DeleteShelvedChangeAsync(string clusterName, int shelvedChange, CancellationToken cancellationToken)
 		{
 			using IScope scope = GlobalTracer.Instance.BuildSpan("PerforceService.DeleteShelvedChangeAsync").StartActive();
 			scope.Span.SetTag("ClusterName", clusterName);
 			scope.Span.SetTag("ShelvedChange", shelvedChange);
 			
 			PerforceCluster cluster = await GetClusterAsync(clusterName);
-			using (P4.Repository repository = await GetConnection(cluster, noClient: true))
-			{
-
-				P4.Changelist? changelist = repository.GetChangelist(shelvedChange, new P4.DescribeCmdOptions(P4.DescribeChangelistCmdFlags.Omit, 0, 0));
 
-				if (changelist == null)
-				{
-					throw new Exception($"Unable to get shelved changelist for delete: {shelvedChange}");
-				}
-
-				string clientId = changelist.ClientId;
-
-				if (String.IsNullOrEmpty(clientId))
-				{
-					throw new Exception($"Unable to get shelved changelist client id for delete: {shelvedChange}");
-				}
-
-				using (P4.Repository clientRepository = await GetConnection(cluster, changelist.Stream, changelist.OwnerName, clientId: clientId, useClientFromChange: true, usePortFromChange: shelvedChange))
-				{
-
-					if (changelist.Shelved)
-					{
-						IList<P4.FileSpec> files = clientRepository.Connection.Client.ShelveFiles(new P4.ShelveFilesCmdOptions(P4.ShelveFilesCmdFlags.Delete, null, shelvedChange));
-					}
-
-					clientRepository.DeleteChangelist(changelist, null);
-				}
+			using (IPerforceConnection perforce = await ConnectAsChangeOwnerAsync(cluster, shelvedChange, cancellationToken))
+			{
+				await perforce.TryDeleteShelvedFilesAsync(shelvedChange, FileSpecList.Any, cancellationToken);
+				await perforce.DeleteChangeAsync(DeleteChangeOptions.None, shelvedChange, cancellationToken);
 			}
 		}
 
 		/// <inheritdoc/>
-		public async Task UpdateChangelistDescription(string clusterName, int change, string description)
+		public async Task UpdateChangelistDescription(string clusterName, int change, string description, CancellationToken cancellationToken)
 		{
 			using IScope scope = GlobalTracer.Instance.BuildSpan("PerforceService.UpdateChangelistDescription").StartActive();
 			scope.Span.SetTag("ClusterName", clusterName);
@@ -1171,30 +827,33 @@ namespace Horde.Build.Perforce
 			try
 			{
 				PerforceCluster cluster = await GetClusterAsync(clusterName);
-
-				using (P4.Repository repository = await GetConnection(cluster, noClient: true))
+				using (IPerforceConnection perforce = await ConnectAsChangeOwnerAsync(cluster, change, cancellationToken))
 				{
-					P4.Changelist changelist = repository.GetChangelist(change);
-
-					repository.Connection.Disconnect();
-
-					// the client must exist for the change list, otherwise will fail (for example, CreateNewChangeAsync deletes the client before returning)
-					using (P4.Repository updateRepository = await GetConnection(cluster, clientFromChange: change, useClientFromChange: true, usePortFromChange: change, username: changelist.OwnerName))
-					{
-						P4.Changelist updatedChangelist = updateRepository.GetChangelist(change);
-						updatedChangelist.Description = description;
-						updateRepository.UpdateChangelist(updatedChangelist);
-					}
+					ChangeRecord record = await perforce.GetChangeAsync(GetChangeOptions.None, change, cancellationToken);
+					record.Description = description;
+					await perforce.UpdateChangeAsync(UpdateChangeOptions.None, record, cancellationToken);
 				}
 			}
 			catch (Exception ex)
 			{
-				LogPerforce(1, "", $"Unable to update Changelist for CL {change} to ${description}, {ex.Message}");
+				_logger.LogError(ex, "Unable to update Changelist for CL {Change} to {Description}, {Message}", change, description, ex.Message);
+			}
+		}
+
+		static async Task ResetClientAsync(IPerforceConnection perforce, CancellationToken cancellationToken)
+		{
+			await perforce.RevertAsync(-1, null, RevertOptions.KeepWorkspaceFiles, FileSpecList.Any, cancellationToken);
+
+			List<ChangesRecord> changes = await perforce.GetChangesAsync(ChangesOptions.None, perforce.Settings.ClientName, -1, ChangeStatus.Pending, null, FileSpecList.Any, cancellationToken);
+			foreach(ChangesRecord change in changes)
+			{
+				await perforce.DeleteShelvedFilesAsync(change.Number, FileSpecList.Any, cancellationToken);
+				await perforce.DeleteChangeAsync(DeleteChangeOptions.None, change.Number, cancellationToken);
 			}
 		}
 
 		/// <inheritdoc/>
-		public async Task<int> CreateNewChangeAsync(string clusterName, string streamName, string filePath, string description)
+		public async Task<int> CreateNewChangeAsync(string clusterName, string streamName, string filePath, string description, CancellationToken cancellationToken)
 		{
 			using IScope scope = GlobalTracer.Instance.BuildSpan("PerforceService.CreateNewChangeAsync").StartActive();
 			scope.Span.SetTag("ClusterName", clusterName);
@@ -1202,39 +861,24 @@ namespace Horde.Build.Perforce
 			scope.Span.SetTag("FilePath", filePath);
 			
 			PerforceCluster cluster = await GetClusterAsync(clusterName);
-			P4.SubmitResults? submitResults = null;
 
-			using (P4.Repository repository = await GetConnection(cluster, stream: streamName, readOnly: false, createChange: true))
+			using (PooledConnectionHandle perforce = await ConnectWithStreamClientAsync(cluster, null, streamName, cancellationToken))
 			{
-				P4.Client client = repository.Connection.Client;
-
-				string workspaceFilePath = $"//{client.Name}/{filePath.TrimStart('/')}";
-				string diskFilePath = $"{client.Root + filePath.TrimStart('/')}";
-
-				P4.FileSpec workspaceFileSpec = new P4.FileSpec(new P4.ClientPath(workspaceFilePath));
-
-				IList<P4.File> files = repository.GetFiles(new P4.FilesCmdOptions(P4.FilesCmdFlags.None, 1), workspaceFileSpec);
+				string workspaceFilePath = $"//{perforce.Client!.Name}/{filePath.TrimStart('/')}";
+				string diskFilePath = $"{perforce.Client!.Root + filePath.TrimStart('/')}";
 
-				P4.Changelist? submitChangelist = null;
-				P4.DepotPath? depotPath = null;
+				int attempt = 1;
+				const int MaxAttempts = 5;
 
-				const int MaxRetries = 5;
-				int retry = 0;
-
-				for (; ; )
+				for(; ;)
 				{
-					ResetClient(repository);
-
-					depotPath = null;
-
-					if (retry == MaxRetries)
-					{
-						break;
-					}
+					await ResetClientAsync(perforce, cancellationToken);
 
+					string? depotPath = null;
 					try
 					{
-						if (files == null || files.Count == 0)
+						List<FStatRecord> files = await perforce.FStatAsync(workspaceFilePath, cancellationToken).ToListAsync(cancellationToken);
+						if (files.Count == 0)
 						{
 							// File does not exist, create it
 							string? directoryName = Path.GetDirectoryName(diskFilePath);
@@ -1269,285 +913,133 @@ namespace Horde.Build.Perforce
 								}
 							}
 
-							IList<P4.FileSpec> depotFiles = client.AddFiles(new P4.Options(), workspaceFileSpec);
-
-							if (depotFiles == null || depotFiles.Count != 1)
+							List<AddRecord> addFiles = await perforce.AddAsync(-1, workspaceFilePath, cancellationToken);
+							if (addFiles.Count != 1)
 							{
-								throw new Exception($"Unable to add local change file,  local: {diskFilePath} : workspace: {workspaceFileSpec}");
+								throw new Exception($"Unable to add local change file,  local: {diskFilePath} : workspace: {workspaceFilePath}");
 							}
 
-							depotPath = depotFiles[0].DepotPath;
-
+							depotPath = addFiles[0].DepotFile;
 						}
 						else
 						{
-							IList<P4.FileSpec> syncResults = client.SyncFiles(new P4.SyncFilesCmdOptions(P4.SyncFilesCmdFlags.Force), workspaceFileSpec);
-
+							List<SyncRecord> syncResults = await perforce.SyncAsync(SyncOptions.Force, -1, workspaceFilePath, cancellationToken).ToListAsync(cancellationToken);
 							if (syncResults == null || syncResults.Count != 1)
 							{
-								throw new Exception($"Unable to sync file, workspace: {workspaceFileSpec}");
+								throw new Exception($"Unable to sync file, workspace: {workspaceFilePath}");
 							}
 
-							IList<P4.FileSpec> editResults = client.EditFiles(new P4.FileSpec[] { new P4.FileSpec(syncResults[0].DepotPath) }, new P4.Options());
-
+							List<EditRecord> editResults = await perforce.EditAsync(-1, syncResults[0].DepotFile.ToString(), cancellationToken);
 							if (editResults == null || editResults.Count != 1)
 							{
-								throw new Exception($"Unable to edit file, workspace: {workspaceFileSpec}");
+								throw new Exception($"Unable to edit file, workspace: {workspaceFilePath}");
 							}
 
-							depotPath = editResults[0].DepotPath;
-
+							depotPath = editResults[0].DepotFile;
 						}
 
-						if (depotPath == null || String.IsNullOrEmpty(depotPath.Path))
+						if (String.IsNullOrEmpty(depotPath))
 						{
-							throw new Exception($"Unable to get depot path for: {workspaceFileSpec}");
+							throw new Exception($"Unable to get depot path for: {workspaceFilePath}");
 						}
 
 						// create a new change
-						P4.Changelist changelist = new P4.Changelist();
-						changelist.Description = description;
-						changelist.Files.Add(new P4.FileMetaData(new P4.FileSpec(depotPath)));
-						submitChangelist = repository.CreateChangelist(changelist);
-
-						if (submitChangelist == null)
-						{
-							throw new Exception($"Unable to create a changelist for: {depotPath}");
-						}
-
-						submitResults = submitChangelist.Submit(null);
-
-						if (submitResults == null)
-						{
-							throw new Exception($"Unable to submit changelist for: {depotPath}");
-						}
-
-						break;
-					}
-					catch
-					{
-						retry++;
-						continue;
-					}
-				}
-
-				string clientName = client.Name;
-				try
-				{
-					repository.DeleteClient(client, new P4.Options());
-				}
-				catch
-				{
-					_logger.LogError("Unable to delete client {ClientName}", clientName);
-				}
-			}
-
-			if (submitResults == null)
-			{
-				throw new Exception($"Unable to submit change for {streamName} {filePath}");
-			}
-
-			return submitResults.ChangeIdAfterSubmit;
-		}
-
-		/// <inheritdoc/>
-		public async Task<(int? Change, string Message)> SubmitShelvedChangeAsync(string clusterName, int change, int originalChange)
-		{
-			using IScope scope = GlobalTracer.Instance.BuildSpan("PerforceService.SubmitShelvedChangeAsync").StartActive();
-			scope.Span.SetTag("ClusterName", clusterName);
-			scope.Span.SetTag("Change", change);
-			scope.Span.SetTag("OriginalChange", originalChange);
-			
-			PerforceCluster cluster = await GetClusterAsync(clusterName);
-			using (P4.Repository repository = await GetConnection(cluster, noClient: true))
-			{
-
-				List<string> args = new List<string> { "-S", change.ToString(CultureInfo.InvariantCulture), originalChange.ToString(CultureInfo.InvariantCulture) };
-
-				using (P4.P4Command command = new P4.P4Command(repository, "describe", true, args.ToArray()))
-				{
-					P4.P4CommandResult result = command.Run();
-
-					if (!result.Success)
-					{
-						return (null, $"Unable to get change {change}");
-
-					}
-
-					if (result.TaggedOutput == null || result.TaggedOutput.Count != 2)
-					{
-
-						return (null, $"Unable to get tagged output for change: {change} and original change: {originalChange}");
-					}
-
-					bool dstMismatch = false;
-					string offset = String.Empty;
-
-					P4.Server server = repository.Server;
-					if (server != null && server.Metadata != null)
-					{
-						offset = server.Metadata.DateTimeOffset;
-						dstMismatch = P4.FormBase.DSTMismatch(server.Metadata);
-					}
-
-					P4.Changelist changelist = new P4.Changelist();
-					changelist.FromChangeCmdTaggedOutput(result.TaggedOutput[0], true, offset, dstMismatch);
+						ChangeRecord change = new ChangeRecord();
+						change.Description = description;
+						change.Files.Add(depotPath);
+						change = await perforce.CreateChangeAsync(change, cancellationToken);
 
-					P4.Changelist originalChangelist = new P4.Changelist();
-					originalChangelist.FromChangeCmdTaggedOutput(result.TaggedOutput[1], true, offset, dstMismatch);
-
-					if (originalChangelist.ShelvedFiles.Count != changelist.ShelvedFiles.Count)
-					{
-						return (null, $"Mismatched number of shelved files for change: {change} and original change: {originalChange}");
-					}
-
-					if (originalChangelist.ShelvedFiles.Count == 0)
-					{
-						return (null, $"No shelved file for change: {change} and original change: {originalChange}");
+						SubmitRecord submit = await perforce.SubmitAsync(change.Number, SubmitOptions.SubmitUnchanged, cancellationToken);
+						return submit.ChangeNumber;
 					}
-
-					foreach (P4.ShelvedFile shelvedFile in changelist.ShelvedFiles)
+					catch (Exception ex)
 					{
-						P4.ShelvedFile? found = originalChangelist.ShelvedFiles.FirstOrDefault(original => original.Digest == shelvedFile.Digest && original.Action == shelvedFile.Action);
-
-						if (found == null)
+						if (attempt < MaxAttempts)
+						{
+							_logger.LogWarning(ex, "Unable to submit new changelist (file: {File}, attempt: {Attempt}/{MaxAttempts}) - {Message}", depotPath, attempt, MaxAttempts, ex.Message);
+						}
+						else
 						{
-							return (null, $"Mismatch in shelved file digest or action for {shelvedFile.Path}");
+							throw new PerforceException($"Unable to submit new changelist for {filePath} after {MaxAttempts} attempts");
 						}
 					}
-
-					repository.Connection.Disconnect();
-
-					using (P4.Repository submitRepository = await GetConnection(cluster, readOnly: false, clientFromChange: change, useClientFromChange: true, usePortFromChange: change, username: changelist.OwnerName))
+					finally
 					{
-						// we might not need a client here, possibly -e below facilitates this, check!
-						using (P4.P4Command submitCommand = new P4.P4Command(submitRepository, "submit", null, true, new string[] { "-e", change.ToString(CultureInfo.InvariantCulture) }))
+						if (File.Exists(diskFilePath))
 						{
 							try
 							{
-								result = submitCommand.Run();
-							}
-							catch (Exception ex)
-							{
-								return (null, $"Submit command failed: {ex.Message}");
-							}
-
-							if (!result.Success)
-							{
-								string message = (result.ErrorList != null && result.ErrorList.Count > 0) ? result.ErrorList[0].ErrorMessage : "Unknown error, no errors in list";
-								return (null, $"Unable to submit {change}, {message}");
-							}
-
-							int? submittedChangeId = null;
-
-							foreach (P4.TaggedObject taggedObject in result.TaggedOutput)
-							{
-								string? submitted;
-								if (taggedObject.TryGetValue("submittedChange", out submitted))
-								{
-									submittedChangeId = Int32.Parse(submitted, CultureInfo.InvariantCulture);
-								}
+								File.SetAttributes(diskFilePath, FileAttributes.Normal);
+								File.Delete(diskFilePath);
 							}
-
-							if (submittedChangeId == null)
+							catch (Exception ex2)
 							{
-								return (null, $"Submit command succeeded, though unable to parse submitted change number");
+								_logger.LogWarning(ex2, "Unable to delete temp file {File}", diskFilePath);
 							}
-
-							return (submittedChangeId, changelist.Description);
-
 						}
 					}
 				}
 			}
-
-			throw new Exception($"Unable to get shelve change: {change} for original change: {originalChange}");
 		}
 
-		/// <inheritdoc/>		
-		async Task<P4.User?> GetUserInfoAsync(string clusterName, string userName)
+		/// <inheritdoc/>
+		public async Task<(int? Change, string Message)> SubmitShelvedChangeAsync(string clusterName, int change, int originalChange, CancellationToken cancellationToken)
 		{
-			using IScope scope = GlobalTracer.Instance.BuildSpan("PerforceService.GetUserInfoAsync").StartActive();
+			using IScope scope = GlobalTracer.Instance.BuildSpan("PerforceService.SubmitShelvedChangeAsync").StartActive();
 			scope.Span.SetTag("ClusterName", clusterName);
-			scope.Span.SetTag("UserName", userName);
-			
-			PerforceCluster cluster = await GetClusterAsync(clusterName);
-			using (P4.Repository repository = await GetConnection(cluster, noClient: true))
-			{
-				P4.User user = repository.GetUser(userName, new P4.UserCmdOptions(P4.UserCmdFlags.Output));
-				return user;
-			}
-		}
-
-		static int ReshelveChange(P4.Repository repository, int change)
-		{
-			using IScope scope = GlobalTracer.Instance.BuildSpan("PerforceService.ReshelveChange").StartActive();
 			scope.Span.SetTag("Change", change);
+			scope.Span.SetTag("OriginalChange", originalChange);
+
+			PerforceCluster cluster = await GetClusterAsync(clusterName);
+			IPerforceServer server = await GetServerAsync(cluster);
 
-			bool edgeServer = false;
-			string? value;
-			if (repository.Server.Metadata.RawData.TryGetValue("serverServices", out value))
+			List<DescribeRecord> records;
+			using (IPerforceConnection perforce = await ConnectAsync(cluster, null, cancellationToken))
 			{
-				if (value == "edge-server")
+				records = await perforce.DescribeAsync(DescribeOptions.Shelved, -1, new int[] { change, originalChange }, cancellationToken);
+				if (records.Count != 2)
 				{
-					edgeServer = true;
+					throw new PerforceException("Unexpected response count when examining changes to submit; expected 2, got {NumRecords}", records.Count);
 				}
-			}
-
-			List<string> arguments = new List<string>();
-
-			// promote shelf if we're on an edge server
-			if (edgeServer)
-			{
-				arguments.Add("-p");
-			}
-
-			arguments.AddRange(new string[] { "-s", change.ToString(CultureInfo.InvariantCulture), "-f" });
-
-			using (P4.P4Command cmd = new P4.P4Command(repository, "reshelve", false, arguments.ToArray()))
-			{
-				P4.P4CommandResult results = cmd.Run();
 
-				if (results.Success)
+				if (records[0].Files.Count != records[1].Files.Count)
 				{
+					return (null, $"Mismatched number of shelved files for change {change} and original change: {originalChange}");
+				}
+				if (records[0].Files.Count == 0)
+				{
+					return (null, $"No shelved file for change {originalChange}");
+				}
 
-					if (results.InfoOutput.Count == 0)
-					{
-						Serilog.Log.Logger.Information("Perforce: Unexpected info output when reshelving change");
-						throw new Exception("Unexpected info output when reshelving change");
-					}
-
-					bool error = true;
-					int reshelvedChange = 0;
-					string message = results.InfoOutput[^1].Message;
-					Match changeMatch = Regex.Match(message, @"Change (\d+) files shelved");
-
-					if (changeMatch.Success && changeMatch.Groups.Count == 2)
-					{
-						if (Int32.TryParse(changeMatch.Groups[1].Value, out reshelvedChange))
-						{
-							error = false;
-						}
-					}
-
-					if (error)
+				List<DescribeFileRecord> files0 = records[0].Files.OrderBy(x => x.DepotFile, StringComparer.Ordinal).ToList();
+				List<DescribeFileRecord> files1 = records[1].Files.OrderBy(x => x.DepotFile, StringComparer.Ordinal).ToList();
+				foreach ((DescribeFileRecord file0, DescribeFileRecord file1) in Enumerable.Zip(files0, files1))
+				{
+					if (!String.Equals(file0.DepotFile, file1.DepotFile, StringComparison.Ordinal) || !String.Equals(file0.Digest, file1.Digest, StringComparison.Ordinal) || file0.Action != file1.Action)
 					{
-						Serilog.Log.Logger.Information("Perforce: Unable to parse cl for reshelf: {Message}", message);
-						throw new Exception($"Unable to parse cl for reshelf: {message}");
+						_logger.LogInformation("Mismatch in shelved files between {File0} or {File1}", file0.DepotFile, file1.DepotFile);
+						return (null, $"Shelved files have been modified.");
 					}
-
-					return reshelvedChange;
 				}
 			}
 
-			Serilog.Log.Logger.Information("Perforce: General reshelving failure");
-
-			throw new Exception($"Unable to reshelve CL {change}");
-
+			using (IPerforceConnection perforce = await ConnectAsync(cluster, records[0].User, cancellationToken))
+			{
+				try
+				{
+					SubmitRecord submit = await perforce.SubmitShelvedAsync(change, cancellationToken);
+					return (submit.ChangeNumber, records[0].Description);
+				}
+				catch (Exception ex)
+				{
+					_logger.LogError(ex, "Unable to submit shelved change {Change}: {Message}", change, ex.Message);
+					return (null, $"Submit command failed: {ex.Message}");
+				}
+			}
 		}
 
 		/// <inheritdoc/>
-		public virtual async Task<int> GetCodeChangeAsync(string clusterName, string streamName, int change)
+		public async Task<int> GetCodeChangeAsync(string clusterName, string streamName, int change, CancellationToken cancellationToken)
 		{
 			using IScope scope = GlobalTracer.Instance.BuildSpan("PerforceService.GetCodeChangeAsync").StartActive();
 			scope.Span.SetTag("ClusterName", clusterName);
@@ -1558,19 +1050,19 @@ namespace Horde.Build.Perforce
 			for (; ; )
 			{
 				// Query for the changes before this point
-				List<ChangeSummary> changes = await GetChangesAsync(clusterName, streamName, null, maxChange, 10, null);
-				Serilog.Log.Logger.Information("Finding last code change in {Stream} before {MaxChange}: {NumResults}", streamName, maxChange, changes.Count);
+				List<ChangeSummary> changes = await GetChangesAsync(clusterName, streamName, null, maxChange, 10, null, cancellationToken);
+				_logger.LogInformation("Finding last code change in {Stream} before {MaxChange}: {NumResults}", streamName, maxChange, changes.Count);
 				if (changes.Count == 0)
 				{
 					return 0;
 				}
 
 				// Get the details for them
-				List<ChangeDetails> detailsList = await GetChangeDetailsAsync(clusterName, streamName, changes.ConvertAll(x => x.Number), null);
+				List<ChangeDetails> detailsList = await GetChangeDetailsAsync(clusterName, streamName, changes.ConvertAll(x => x.Number), null, cancellationToken);
 				foreach (ChangeDetails details in detailsList.OrderByDescending(x => x.Number))
 				{
 					ChangeContentFlags contentFlags = details.GetContentFlags();
-					Serilog.Log.Logger.Information("Change {Change} = {Flags}", details.Number, contentFlags.ToString());
+					_logger.LogInformation("Change {Change} = {Flags}", details.Number, contentFlags.ToString());
 					if ((details.GetContentFlags() & ChangeContentFlags.ContainsCode) != 0)
 					{
 						return details.Number;
diff --git a/Engine/Source/Programs/Horde/Horde.Build/Perforce/ReplicationService.cs b/Engine/Source/Programs/Horde/Horde.Build/Perforce/ReplicationService.cs
--- a/Engine/Source/Programs/Horde/Horde.Build/Perforce/ReplicationService.cs
+++ b/Engine/Source/Programs/Horde/Horde.Build/Perforce/ReplicationService.cs
@@ -1022,7 +1022,7 @@ namespace Horde.Build.Perforce
 			ReplicationClient? clientInfo = await FindReplicationClientAsync(stream);
 			if (clientInfo == null)
 			{
-				using IPerforceConnection? perforce = await _perforceService.GetServiceUserConnection(stream.ClusterName);
+				using IPerforceConnection? perforce = await _perforceService.ConnectAsync(stream.ClusterName);
 				if (perforce == null)
 				{
 					throw new PerforceException($"Unable to create connection to Perforce server");
diff --git a/Engine/Source/Programs/Horde/Horde.Build/ServerSettings.cs b/Engine/Source/Programs/Horde/Horde.Build/ServerSettings.cs
--- a/Engine/Source/Programs/Horde/Horde.Build/ServerSettings.cs
+++ b/Engine/Source/Programs/Horde/Horde.Build/ServerSettings.cs
@@ -568,10 +568,15 @@ namespace Horde.Build
 		public NamespaceId ToolNamespaceId { get; set; } = new NamespaceId("horde.p4");
 
 		/// <summary>
-		/// Whether to 
+		/// Whether to use the local Perforce environment
 		/// </summary>
 		public bool UseLocalPerforceEnv { get; set; }
 
+		/// <summary>
+		/// Number of pooled perforce connections to keep
+		/// </summary>
+		public int PerforceConnectionPoolSize { get; set; } = 5;
+
 		/// <summary>
 		/// Lazily computed timezone value
 		/// </summary>
diff --git a/Engine/Source/Programs/Horde/Horde.Build/Utilities/ViewMap.cs b/Engine/Source/Programs/Horde/Horde.Build/Utilities/ViewMap.cs
--- a/Engine/Source/Programs/Horde/Horde.Build/Utilities/ViewMap.cs
+++ b/Engine/Source/Programs/Horde/Horde.Build/Utilities/ViewMap.cs
@@ -9,8 +9,6 @@ using EpicGames.Core;
 
 namespace Horde.Build.Utilities
 {
-	using P4 = global::Perforce.P4;
-
 	/// <summary>
 	/// Stores a mapping from one set of paths to another
 	/// </summary>
@@ -38,19 +36,6 @@ namespace Horde.Build.Utilities
 			Entries = new List<ViewMapEntry>(other.Entries);
 		}
 
-		/// <summary>
-		/// Constructor
-		/// </summary>
-		/// <param name="viewMap">Peforce viewmap definition</param>
-		public ViewMap(P4.ViewMap viewMap)
-			: this()
-		{
-			foreach (P4.MapEntry entry in viewMap)
-			{
-				Entries.Add(new ViewMapEntry(entry));
-			}
-		}
-
 		/// <summary>
 		/// Determines if a file is included in the view
 		/// </summary>
@@ -259,15 +244,6 @@ namespace Horde.Build.Utilities
 			TargetSuffix = targetSuffix;
 		}
 
-		/// <summary>
-		/// Constructor
-		/// </summary>
-		/// <param name="entry"></param>
-		public ViewMapEntry(P4.MapEntry entry)
-			: this(entry.Type == P4.MapType.Include, entry.Left.Path, entry.Right.Path)
-		{
-		}
-
 		/// <summary>
 		/// Maps a file to the target path
 		/// </summary>
diff --git a/Engine/Source/Programs/Shared/EpicGames.Perforce/FileSpecList.cs b/Engine/Source/Programs/Shared/EpicGames.Perforce/FileSpecList.cs
--- a/Engine/Source/Programs/Shared/EpicGames.Perforce/FileSpecList.cs
+++ b/Engine/Source/Programs/Shared/EpicGames.Perforce/FileSpecList.cs
@@ -1,13 +1,15 @@
 // Copyright Epic Games, Inc. All Rights Reserved.
 
+using System;
 using System.Collections.Generic;
+using System.Linq;
 
 namespace EpicGames.Perforce
 {
 	/// <summary>
 	/// Wrapper for a list of filespecs. Allows implicit conversion from string (a single entry) or list.
 	/// </summary>
-	public struct FileSpecList
+	public struct FileSpecList : IEquatable<FileSpecList>
 	{
 		/// <summary>
 		/// Empty filespec list
@@ -33,6 +35,33 @@ namespace EpicGames.Perforce
 			List = fileSpecList;
 		}
 
+		/// <inheritdoc/>
+		public override bool Equals(object? obj) => obj is FileSpecList other && Equals(other);
+
+		/// <inheritdoc/>
+		public override int GetHashCode()
+		{
+			HashCode code = new HashCode();
+			foreach (string item in List)
+			{
+				code.Add(item);
+			}
+			return code.ToHashCode();
+		}
+
+		/// <inheritdoc/>
+		public bool Equals(FileSpecList other) => List.SequenceEqual(other.List, StringComparer.Ordinal);
+
+		/// <summary>
+		/// Test two filespecs for equality
+		/// </summary>
+		public static bool operator ==(FileSpecList a, FileSpecList b) => a.Equals(b);
+
+		/// <summary>
+		/// Test two filespecs for inequality
+		/// </summary>
+		public static bool operator !=(FileSpecList a, FileSpecList b) => !a.Equals(b);
+
 		/// <summary>
 		/// Implicit conversion operator from a list of filespecs
 		/// </summary>
diff --git a/Engine/Source/Programs/Shared/EpicGames.Perforce/PerforceConnection.cs b/Engine/Source/Programs/Shared/EpicGames.Perforce/PerforceConnection.cs
--- a/Engine/Source/Programs/Shared/EpicGames.Perforce/PerforceConnection.cs
+++ b/Engine/Source/Programs/Shared/EpicGames.Perforce/PerforceConnection.cs
@@ -7,6 +7,7 @@ using System.Linq;
 using System.Reflection;
 using System.Runtime.CompilerServices;
 using System.Text;
+using System.Text.RegularExpressions;
 using System.Threading;
 using System.Threading.Tasks;
 using EpicGames.Core;
@@ -719,13 +720,13 @@ namespace EpicGames.Perforce
 				throw new PerforceException("Unexpected info response from change command: {0}", response);
 			}
 
-			string[] tokens = info.Data.Split(' ');
-			if (tokens.Length != 3)
+			Match match = Regex.Match(info.Data, @"^Change (\d+) created");
+			if (!match.Success)
 			{
 				throw new PerforceException("Unexpected info response from change command: {0}", response);
 			}
 
-			record.Number = Int32.Parse(tokens[1]);
+			record.Number = Int32.Parse(match.Groups[1].Value);
 			return new PerforceResponse<ChangeRecord>(record);
 		}
 
@@ -1026,8 +1027,11 @@ namespace EpicGames.Perforce
 			{
 				arguments.Add($"-u{userName}");
 			}
-			arguments.AddRange(fileSpecs.List);
-
+			if (fileSpecs != FileSpecList.Any) // Queries are slower with a path specification
+			{
+				arguments.AddRange(fileSpecs.List);
+			}
+	
 			return CommandAsync<ChangesRecord>(connection, "changes", arguments, null, cancellationToken);
 		}
 
@@ -1577,6 +1581,40 @@ namespace EpicGames.Perforce
 			return records[0];
 		}
 
+		/// <summary>
+		/// Describes a single changelist
+		/// </summary>
+		/// <param name="connection">Connection to the Perforce server</param>
+		/// <param name="options">Options for the command</param>
+		/// <param name="maxNumFiles">Maximum number of files to return</param>
+		/// <param name="changeNumber">The changelist number to retrieve description for</param>
+		/// <param name="cancellationToken">Token used to cancel the operation</param>
+		/// <returns>Response from the server; either a describe record or error code</returns>
+		public static async Task<DescribeRecord> DescribeAsync(this IPerforceConnection connection, DescribeOptions options, int maxNumFiles, int changeNumber, CancellationToken cancellationToken = default)
+		{
+			PerforceResponse<DescribeRecord> response = await TryDescribeAsync(connection, options, maxNumFiles, changeNumber, cancellationToken);
+			return response.Data;
+		}
+
+		/// <summary>
+		/// Describes a single changelist
+		/// </summary>
+		/// <param name="connection">Connection to the Perforce server</param>
+		/// <param name="options">Options for the command</param>
+		/// <param name="maxNumFiles">Maximum number of files to return</param>
+		/// <param name="changeNumber">The changelist number to retrieve description for</param>
+		/// <param name="cancellationToken">Token used to cancel the operation</param>
+		/// <returns>Response from the server; either a describe record or error code</returns>
+		public static async Task<PerforceResponse<DescribeRecord>> TryDescribeAsync(this IPerforceConnection connection, DescribeOptions options, int maxNumFiles, int changeNumber, CancellationToken cancellationToken = default)
+		{
+			PerforceResponseList<DescribeRecord> records = await TryDescribeAsync(connection, options, maxNumFiles, new int[] { changeNumber }, cancellationToken);
+			if (records.Count != 1)
+			{
+				throw new PerforceException("Expected only one record returned from p4 describe command, got {0}", records.Count);
+			}
+			return records[0];
+		}
+
 		/// <summary>
 		/// Describes a set of changelists
 		/// </summary>
@@ -2305,14 +2343,51 @@ namespace EpicGames.Perforce
 		/// <summary>
 		/// Attempts to log in to the server
 		/// </summary>
-		/// <param name="connection"></param>
+		/// <param name="connection">Connection to use</param>
 		/// <param name="password"></param>
 		/// <param name="cancellationToken"></param>
 		/// <returns></returns>
-		public static async Task<PerforceResponse<LoginRecord>> TryLoginAsync(this IPerforceConnection connection, string password, CancellationToken cancellationToken = default)
+		public static Task<PerforceResponse<LoginRecord>> TryLoginAsync(this IPerforceConnection connection, string? password, CancellationToken cancellationToken = default)
+		{
+			return TryLoginAsync(connection, LoginOptions.None, null, password, null, cancellationToken);
+		}
+
+		/// <summary>
+		/// Attempts to log in to the server
+		/// </summary>
+		/// <param name="connection">Connection to use</param>
+		/// <param name="options">Options for the command</param>
+		/// <param name="user">User to login as</param>
+		/// <param name="password">Password for the user</param>
+		/// <param name="host">Host for the ticket</param>
+		/// <param name="cancellationToken"></param>
+		/// <returns></returns>
+		public static async Task<PerforceResponse<LoginRecord>> TryLoginAsync(this IPerforceConnection connection, LoginOptions options, string? user, string? password, string? host, CancellationToken cancellationToken = default)
 		{
+			List<string> arguments = new List<string>();
+			if((options & LoginOptions.AllHosts) != 0)
+			{
+				arguments.Add("-a");
+			}
+			if ((options & LoginOptions.PrintTicket) != 0)
+			{
+				arguments.Add("-p");
+			}
+			if ((options & LoginOptions.Status) != 0)
+			{
+				arguments.Add("-s");
+			}
+			if (host != null)
+			{
+				arguments.Add($"-h{host}");
+			}
+			if (user != null)
+			{
+				arguments.Add(user);
+			}
+
 			List<PerforceResponse> parsedResponses;
-			await using (IPerforceOutput response = connection.Command("login", Array.Empty<string>(), null, null, password, false))
+			await using (IPerforceOutput response = connection.Command("login", arguments, null, null, password, false))
 			{
 				for (; ; )
 				{
@@ -2331,14 +2406,44 @@ namespace EpicGames.Perforce
 				// not DisposeAsync()ing here will cause a deadlock when calling `TryGetLoginStateAsync()` below
 			}
 
-			PerforceResponse firstResponse = parsedResponses.First();
-			if (firstResponse.Info != null)
+			PerforceResponse? error = parsedResponses.FirstOrDefault(x => !x.Succeeded);
+			if (error != null)
+			{
+				return new PerforceResponse<LoginRecord>(error);
+			}
+
+			string? ticket = null;
+			if ((options & LoginOptions.PrintTicket) != 0)
+			{
+				if (parsedResponses.Count != 2)
+				{
+					throw new PerforceException("Unable to parse login response; expected two records, one with ticket id");
+				}
+
+				PerforceInfo? info = parsedResponses[0].Info;
+				if (info == null)
+				{
+					throw new PerforceException("Unable to parse login response; expected two records, one with ticket id");
+				}
+
+				ticket = info.Data;
+				parsedResponses.RemoveAt(0);
+			}
+
+			LoginRecord? loginRecord = parsedResponses.First().Data as LoginRecord;
+			if (loginRecord == null)
 			{
 				// Older versions of P4.EXE do not return a login record for succesful login, instread just returning a string. Call p4 login -s to get the login state instead.
-				return await TryGetLoginStateAsync(connection, cancellationToken);
+				PerforceResponse<LoginRecord> legacyResponse = await TryGetLoginStateAsync(connection, cancellationToken);
+				if (!legacyResponse.Succeeded)
+				{
+					return legacyResponse;
+				}
+				loginRecord = legacyResponse.Data;
 			}
 
-			return new PerforceResponse<LoginRecord>(firstResponse);
+			loginRecord.Ticket = ticket;
+			return new PerforceResponse<LoginRecord>(loginRecord);
 		}
 
 		static void DiscardPasswordPrompt(IPerforceOutput output)
@@ -2623,7 +2728,14 @@ namespace EpicGames.Perforce
 			}
 			finally
 			{
-				File.Delete(tempFile);
+				try
+				{
+					File.SetAttributes(tempFile, FileAttributes.Normal);
+					File.Delete(tempFile);
+				}
+				catch
+				{
+				}
 			}
 		}
 
@@ -2990,6 +3102,38 @@ namespace EpicGames.Perforce
 
 		#endregion
 
+		#region p4 server
+
+		/// <summary>
+		/// Gets information about the specified server
+		/// </summary>
+		/// <param name="connection">Connection to the Perforce server</param>
+		/// <param name="serverId">The server identifier</param>
+		/// <param name="cancellationToken">Token used to cancel the operation</param>
+		/// <returns></returns>
+		public static async Task<ServerRecord> GetServerAsync(this IPerforceConnection connection, string serverId, CancellationToken cancellationToken = default)
+		{
+			PerforceResponse<ServerRecord> response = await TryGetServerAsync(connection, serverId, cancellationToken);
+			return response.Data;
+		}
+
+		/// <summary>
+		/// Gets information about the specified server
+		/// </summary>
+		/// <param name="connection">Connection to the Perforce server</param>
+		/// <param name="serverId">The server identifier</param>
+		/// <param name="cancellationToken">Token used to cancel the operation</param>
+		/// <returns></returns>
+		public static Task<PerforceResponse<ServerRecord>> TryGetServerAsync(this IPerforceConnection connection, string serverId, CancellationToken cancellationToken = default)
+		{
+			List<string> arguments = new List<string>();
+			arguments.Add("-o");
+			arguments.Add(serverId);
+			return SingleResponseCommandAsync<ServerRecord>(connection, "server", arguments, null, cancellationToken);
+		}
+
+		#endregion
+
 		#region p4 shelve
 
 		/// <summary>
@@ -3224,6 +3368,22 @@ namespace EpicGames.Perforce
 			{
 				arguments.Add("-r");
 			}
+			if ((options & SubmitOptions.SubmitUnchanged) != 0)
+			{
+				arguments.Add("-f");
+				arguments.Add("submitunchanged");
+			}
+			if ((options & SubmitOptions.RevertUnchanged) != 0)
+			{
+				arguments.Add("-f");
+				arguments.Add("revertunchanged");
+			}
+			if ((options & SubmitOptions.LeaveUnchanged) != 0)
+			{
+				arguments.Add("-f");
+				arguments.Add("leaveunchanged");
+			}
+
 			arguments.Add($"-c{changeNumber}");
 
 			return (await CommandAsync<SubmitRecord>(connection, "submit", arguments, null, cancellationToken))[0];
diff --git a/Engine/Source/Programs/Shared/EpicGames.Perforce/Records/ClientRecord.cs b/Engine/Source/Programs/Shared/EpicGames.Perforce/Records/ClientRecord.cs
--- a/Engine/Source/Programs/Shared/EpicGames.Perforce/Records/ClientRecord.cs
+++ b/Engine/Source/Programs/Shared/EpicGames.Perforce/Records/ClientRecord.cs
@@ -46,6 +46,12 @@ namespace EpicGames.Perforce
 		[PerforceTag("Description", Optional = true)]
 		public string? Description { get; set; }
 
+		/// <summary>
+		/// Name of the server containing this change.
+		/// </summary>
+		[PerforceTag("ServerID", Optional = true)]
+		public string? ServerId { get; set; }
+
 		/// <summary>
 		/// The directory (on the local host) relative to which all the files in the View: are specified. 
 		/// </summary>
diff --git a/Engine/Source/Programs/Shared/EpicGames.Perforce/Records/InfoRecord.cs b/Engine/Source/Programs/Shared/EpicGames.Perforce/Records/InfoRecord.cs
--- a/Engine/Source/Programs/Shared/EpicGames.Perforce/Records/InfoRecord.cs
+++ b/Engine/Source/Programs/Shared/EpicGames.Perforce/Records/InfoRecord.cs
@@ -84,5 +84,11 @@ namespace EpicGames.Perforce
 		/// </summary>
 		[PerforceTag("ServerID", Optional = true)]
 		public string? ServerId { get; set; }
+
+		/// <summary>
+		/// Timezone offset from UTC, in seconds
+		/// </summary>
+		[PerforceTag("tzoffset", Optional = true)]
+		public int TimeZoneOffsetSecs { get; set; }
 	}
 }
diff --git a/Engine/Source/Programs/Shared/EpicGames.Perforce/Records/LoginOptions.cs b/Engine/Source/Programs/Shared/EpicGames.Perforce/Records/LoginOptions.cs
new file mode 100644
--- /dev/null
+++ b/Engine/Source/Programs/Shared/EpicGames.Perforce/Records/LoginOptions.cs
@@ -0,0 +1,35 @@
+// Copyright Epic Games, Inc. All Rights Reserved.
+
+using System;
+
+namespace EpicGames.Perforce
+{
+	/// <summary>
+	/// Options for the login command
+	/// </summary>
+	[Flags]
+	public enum LoginOptions
+	{
+		/// <summary>
+		/// No options specified
+		/// </summary>
+		None = 0,
+
+		/// <summary>
+		/// Obtain a ticket that is valid for all IP addresses.
+		/// </summary>
+		AllHosts = 1,
+
+		/// <summary>
+		/// Display the ticket, rather than storing it in the local ticket file.
+		/// </summary>
+		PrintTicket = 2,
+
+		/// <summary>
+		/// Display the status of the current ticket, if one exists.
+		/// Use with -a to display status for all hosts, or with -h host to display status for a specfic host.
+		/// Users with super access can provide a username argument to display the status of that username's ticket.
+		/// </summary>
+		Status = 4,
+	}
+}
diff --git a/Engine/Source/Programs/Shared/EpicGames.Perforce/Records/LoginRecord.cs b/Engine/Source/Programs/Shared/EpicGames.Perforce/Records/LoginRecord.cs
--- a/Engine/Source/Programs/Shared/EpicGames.Perforce/Records/LoginRecord.cs
+++ b/Engine/Source/Programs/Shared/EpicGames.Perforce/Records/LoginRecord.cs
@@ -13,6 +13,11 @@ namespace EpicGames.Perforce
 		[PerforceTag("User")]
 		public string? User { get; set; }
 
+		/// <summary>
+		/// The ticket created for the user. Set only if <see cref="LoginOptions.PrintTicket"/> is set.
+		/// </summary>
+		public string? Ticket { get; set; }
+
 		/// <summary>
 		/// Time at which the current ticket expires
 		/// </summary>
diff --git a/Engine/Source/Programs/Shared/EpicGames.Perforce/Records/ReshelveOptions.cs b/Engine/Source/Programs/Shared/EpicGames.Perforce/Records/ReshelveOptions.cs
new file mode 100644
--- /dev/null
+++ b/Engine/Source/Programs/Shared/EpicGames.Perforce/Records/ReshelveOptions.cs
@@ -0,0 +1,28 @@
+// Copyright Epic Games, Inc. All Rights Reserved.
+
+using System;
+
+namespace EpicGames.Perforce
+{
+	/// <summary>
+	/// Options for the p4 reshelve command
+	/// </summary>
+	[Flags]
+	public enum ReshelveOptions
+	{
+		/// <summary>
+		/// No options specified
+		/// </summary>
+		None = 0,
+
+		/// <summary>
+		/// When the same file already exists in the target changelist, force the overwriting of it.
+		/// </summary>
+		Force = 1,
+
+		/// <summary>
+		/// Promote the new or target changelist where it can be accessed by other edge servers participating in the multi-server configuration. Once a shelved change has been promoted, all subsequent local modifications to the shelf are also pushed to the commit server and remain until the shelf is deleted.
+		/// </summary>
+		Promote = 2,
+	}
+}
diff --git a/Engine/Source/Programs/Shared/EpicGames.Perforce/Records/ServerRecord.cs b/Engine/Source/Programs/Shared/EpicGames.Perforce/Records/ServerRecord.cs
new file mode 100644
--- /dev/null
+++ b/Engine/Source/Programs/Shared/EpicGames.Perforce/Records/ServerRecord.cs
@@ -0,0 +1,37 @@
+// Copyright Epic Games, Inc. All Rights Reserved.
+
+using System;
+using System.Collections.Generic;
+
+namespace EpicGames.Perforce
+{
+	/// <summary>
+	/// Represents a Perforce clientspec
+	/// </summary>
+	public class ServerRecord
+	{
+		/// <summary>
+		/// The server id
+		/// </summary>
+		[PerforceTag("ServerID")]
+		public string ServerId { get; set; } = String.Empty;
+
+		/// <summary>
+		/// The server type
+		/// </summary>
+		[PerforceTag("Type")]
+		public string Type { get; set; } = String.Empty;
+
+		/// <summary>
+		/// The P4NAME used by this server (optional).
+		/// </summary>
+		[PerforceTag("Name", Optional = true)]
+		public string? Name { get; set; }
+
+		/// <summary>
+		/// The P4PORT used by this server (optional).
+		/// </summary>
+		[PerforceTag("Address", Optional = true)]
+		public string? Address { get; set; }
+	}
+}
diff --git a/Engine/Source/Programs/Shared/EpicGames.Perforce/Records/SubmitOptions.cs b/Engine/Source/Programs/Shared/EpicGames.Perforce/Records/SubmitOptions.cs
--- a/Engine/Source/Programs/Shared/EpicGames.Perforce/Records/SubmitOptions.cs
+++ b/Engine/Source/Programs/Shared/EpicGames.Perforce/Records/SubmitOptions.cs
@@ -19,5 +19,20 @@ namespace EpicGames.Perforce
 		/// Reopen files for edit in the default changelist after submission.
 		/// </summary>
 		ReopenAsEdit = 1,
+
+		/// <summary>
+		/// All open files (with or without changes) are submitted to the depot. This is the default behavior of Helix Server.
+		/// </summary>
+		SubmitUnchanged = 2,
+
+		/// <summary>
+		/// Only those files with content or type changes are submitted to the depot. Unchanged files are reverted.
+		/// </summary>
+		RevertUnchanged = 4,
+
+		/// <summary>
+		/// Only those files with content or type changes are submitted to the depot. Any unchanged files are moved to the default changelist.
+		/// </summary>
+		LeaveUnchanged = 8,
 	}
 }
