Horde: Remove user impersonation when querying submitted changes.

This functionality is insecure but expensive; there is no way to prevent jobs from leaking information about the contents of a stream, so so we already only grant access to streams that people are permitted to see any jobs for. Filtering changelist descriptions provides minimal additional benefit, but prevents us being able to cache responses.

diff --git a/Engine/Source/Programs/Horde/Horde.Build/Issues/IssueService.cs b/Engine/Source/Programs/Horde/Horde.Build/Issues/IssueService.cs
--- a/Engine/Source/Programs/Horde/Horde.Build/Issues/IssueService.cs
+++ b/Engine/Source/Programs/Horde/Horde.Build/Issues/IssueService.cs
@@ -1090,7 +1090,7 @@ namespace Horde.Build.Issues
 				if (stream != null)
 				{
 					_logger.LogInformation("Querying fix changelist {FixChange} in {StreamId}", fixChange, streamId);
-					List<ChangeSummary> changes = await _perforce.GetChangesAsync(stream.ClusterName, stream.Name, fixChange, fixChange, 1, null);
+					List<ChangeSummary> changes = await _perforce.GetChangesAsync(stream.ClusterName, stream.Name, fixChange, fixChange, 1);
 					bContainsFixChange = changes.Count > 0;
 				}
 				cachedContainsFixChange[(streamId, fixChange)] = bContainsFixChange;
@@ -1114,7 +1114,7 @@ namespace Horde.Build.Issues
 				_logger.LogDebug("Querying for changes in {StreamName} between {MinChange} and {MaxChange}", stream.Name, minChange, maxChange);
 
 				// Get the submitted changes before this job
-				List<ChangeDetails> changes = await PerforceServiceExtensions.GetChangeDetailsAsync(_perforce, stream.ClusterName, stream.Name, minChange, maxChange, MaxChanges, null);
+				List<ChangeDetails> changes = await PerforceServiceExtensions.GetChangeDetailsAsync(_perforce, stream.ClusterName, stream.Name, minChange, maxChange, MaxChanges);
 				_logger.LogDebug("Found {NumResults} changes", changes.Count);
 
 				// Get the handler to rank them
diff --git a/Engine/Source/Programs/Horde/Horde.Build/Jobs/JobService.cs b/Engine/Source/Programs/Horde/Horde.Build/Jobs/JobService.cs
--- a/Engine/Source/Programs/Horde/Horde.Build/Jobs/JobService.cs
+++ b/Engine/Source/Programs/Horde/Horde.Build/Jobs/JobService.cs
@@ -1095,7 +1095,7 @@ namespace Horde.Build.Jobs
 
 					_logger.LogInformation("Updating description for {ClonedPreflightChange}", clonedPreflightChange);
 
-					ChangeDetails details = await _perforceService.GetChangeDetailsAsync(stream.ClusterName, stream.Name, clonedPreflightChange, null);
+					ChangeDetails details = await _perforceService.GetChangeDetailsAsync(stream.ClusterName, stream.Name, clonedPreflightChange);
 					await _perforceService.UpdateChangelistDescription(stream.ClusterName, clonedPreflightChange, details.Description.TrimEnd() + $"\n#preflight {job.Id}");
 
 					_logger.LogInformation("Submitting change {Change} (through {ChangeCopy}) after successful completion of {JobId}", job.PreflightChange, clonedPreflightChange, job.Id);
diff --git a/Engine/Source/Programs/Horde/Horde.Build/Jobs/JobsController.cs b/Engine/Source/Programs/Horde/Horde.Build/Jobs/JobsController.cs
--- a/Engine/Source/Programs/Horde/Horde.Build/Jobs/JobsController.cs
+++ b/Engine/Source/Programs/Horde/Horde.Build/Jobs/JobsController.cs
@@ -149,7 +149,7 @@ namespace Horde.Build.Jobs
 			}
 			else
 			{
-				change = await _perforce.GetLatestChangeAsync(stream.ClusterName, stream.Name, null);
+				change = await _perforce.GetLatestChangeAsync(stream.ClusterName, stream.Name);
 			}
 
 			// And get the matching code changelist
@@ -162,7 +162,7 @@ namespace Horde.Build.Jobs
 			string? preflightDescription = null;
 			if (create.PreflightChange != null)
 			{
-				(CheckShelfResult result, preflightDescription) = await _perforce.CheckShelfAsync(stream.ClusterName, stream.Name, create.PreflightChange.Value, null);
+				(CheckShelfResult result, preflightDescription) = await _perforce.CheckShelfAsync(stream.ClusterName, stream.Name, create.PreflightChange.Value);
 				switch (result)
 				{
 					case CheckShelfResult.Ok:
@@ -210,7 +210,7 @@ namespace Horde.Build.Jobs
 			if (jobs.Count == 0)
 			{
 				_logger.LogInformation("Unable to find successful build of {TemplateId} target {Target}. Using latest change instead", templateId, target);
-				return await _perforce.GetLatestChangeAsync(stream.ClusterName, stream.Name, null);
+				return await _perforce.GetLatestChangeAsync(stream.ClusterName, stream.Name);
 			}
 			else
 			{
diff --git a/Engine/Source/Programs/Horde/Horde.Build/Jobs/Schedules/ScheduleService.cs b/Engine/Source/Programs/Horde/Horde.Build/Jobs/Schedules/ScheduleService.cs
--- a/Engine/Source/Programs/Horde/Horde.Build/Jobs/Schedules/ScheduleService.cs
+++ b/Engine/Source/Programs/Horde/Horde.Build/Jobs/Schedules/ScheduleService.cs
@@ -82,7 +82,7 @@ namespace Horde.Build.Jobs.Schedules
 			{
 				while (index >= _changes.Count)
 				{
-					List<ChangeSummary> newChanges = await _perforce.GetChangesAsync(_clusterName, _streamName, null, null, _maxResults, null);
+					List<ChangeSummary> newChanges = await _perforce.GetChangesAsync(_clusterName, _streamName, null, null, _maxResults);
 
 					int numResults = newChanges.Count;
 					if (_changes.Count > 0)
@@ -95,7 +95,7 @@ namespace Horde.Build.Jobs.Schedules
 					}
 					if(newChanges.Count > 0)
 					{
-						_changes.AddRange((await _perforce.GetChangeDetailsAsync(_clusterName, _streamName, newChanges.ConvertAll(x => x.Number), null)).OrderByDescending(x => x.Number));
+						_changes.AddRange((await _perforce.GetChangeDetailsAsync(_clusterName, _streamName, newChanges.ConvertAll(x => x.Number))).OrderByDescending(x => x.Number));
 					}
 					_maxResults += 10;
 				}
diff --git a/Engine/Source/Programs/Horde/Horde.Build/Perforce/IPerforceService.cs b/Engine/Source/Programs/Horde/Horde.Build/Perforce/IPerforceService.cs
--- a/Engine/Source/Programs/Horde/Horde.Build/Perforce/IPerforceService.cs
+++ b/Engine/Source/Programs/Horde/Horde.Build/Perforce/IPerforceService.cs
@@ -107,10 +107,9 @@ namespace Horde.Build.Perforce
 		/// <param name="minChange">The minimum changelist number</param>
 		/// <param name="maxChange">The maximum changelist number</param>
 		/// <param name="results">Number of results to return</param>
-		/// <param name="impersonateUser">User to impersonate for the query</param>
 		/// <param name="cancellationToken">Cancellation token for the operation</param>
 		/// <returns>Latest changelist number</returns>
-		public Task<List<ChangeSummary>> GetChangesAsync(string clusterName, string streamName, int? minChange, int? maxChange, int results, string? impersonateUser, CancellationToken cancellationToken = default);
+		public Task<List<ChangeSummary>> GetChangesAsync(string clusterName, string streamName, int? minChange, int? maxChange, int results, CancellationToken cancellationToken = default);
 
 		/// <summary>
 		/// Checks a shelf is valid for the given stream
@@ -118,10 +117,9 @@ namespace Horde.Build.Perforce
 		/// <param name="clusterName">Name of the Perforce cluster</param>
 		/// <param name="streamName">The stream to query</param>
 		/// <param name="changeNumber">Shelved changelist number</param>
-		/// <param name="impersonateUser">Name of the user to impersonate</param>
 		/// <param name="cancellationToken">Cancellation token for the operation</param>
 		/// <returns>Whether the shelf is valid, plus the description for it</returns>
-		public Task<(CheckShelfResult, string?)> CheckShelfAsync(string clusterName, string streamName, int changeNumber, string? impersonateUser, CancellationToken cancellationToken = default);
+		public Task<(CheckShelfResult, string?)> CheckShelfAsync(string clusterName, string streamName, int changeNumber, CancellationToken cancellationToken = default);
 
 		/// <summary>
 		/// Gets the latest change for a particular stream
@@ -139,10 +137,9 @@ namespace Horde.Build.Perforce
 		/// <param name="clusterName">Name of the Perforce cluster</param>
 		/// <param name="streamName">The stream to query</param>
 		/// <param name="changeNumbers">Change numbers to query</param>
-		/// <param name="impersonateUser">Name of the user to impersonate</param>
 		/// <param name="cancellationToken">Cancellation token for the operation</param>
 		/// <returns>Commit details</returns>
-		public Task<List<ChangeDetails>> GetChangeDetailsAsync(string clusterName, string streamName, IReadOnlyList<int> changeNumbers, string? impersonateUser, CancellationToken cancellationToken = default);
+		public Task<List<ChangeDetails>> GetChangeDetailsAsync(string clusterName, string streamName, IReadOnlyList<int> changeNumbers, CancellationToken cancellationToken = default);
 
 		/// <summary>
 		/// Gets the latest changes for a set of depot paths
@@ -235,11 +232,10 @@ namespace Horde.Build.Perforce
 		/// <param name="clusterName">Name of the Perforce cluster</param>
 		/// <param name="streamName">The stream to query</param>
 		/// <param name="changeNumber">Change number to query</param>
-		/// <param name="impersonateUser">Name of the user to impersonate</param>
 		/// <returns>Commit details</returns>
-		public static async Task<ChangeDetails> GetChangeDetailsAsync(this IPerforceService perforce, string clusterName, string streamName, int changeNumber, string? impersonateUser)
+		public static async Task<ChangeDetails> GetChangeDetailsAsync(this IPerforceService perforce, string clusterName, string streamName, int changeNumber)
 		{
-			List<ChangeDetails> results = await perforce.GetChangeDetailsAsync(clusterName, streamName, new[] { changeNumber }, impersonateUser);
+			List<ChangeDetails> results = await perforce.GetChangeDetailsAsync(clusterName, streamName, new[] { changeNumber });
 			return results[0];
 		}
 
@@ -252,12 +248,11 @@ namespace Horde.Build.Perforce
 		/// <param name="minChange">The minimum changelist number</param>
 		/// <param name="maxChange">The maximum changelist number</param>
 		/// <param name="results">Number of results to return</param>
-		/// <param name="impersonateUser">Name of the user to impersonate</param>
 		/// <returns>Commit details</returns>
-		public static async Task<List<ChangeDetails>> GetChangeDetailsAsync(this IPerforceService perforce, string clusterName, string streamName, int? minChange, int? maxChange, int results, string? impersonateUser)
+		public static async Task<List<ChangeDetails>> GetChangeDetailsAsync(this IPerforceService perforce, string clusterName, string streamName, int? minChange, int? maxChange, int results)
 		{
-			List<ChangeSummary> changes = await perforce.GetChangesAsync(clusterName, streamName, minChange, maxChange, results, impersonateUser);
-			return await perforce.GetChangeDetailsAsync(clusterName, streamName, changes.ConvertAll(x => x.Number), impersonateUser);
+			List<ChangeSummary> changes = await perforce.GetChangesAsync(clusterName, streamName, minChange, maxChange, results);
+			return await perforce.GetChangeDetailsAsync(clusterName, streamName, changes.ConvertAll(x => x.Number));
 		}
 
 		/// <summary>
@@ -266,11 +261,10 @@ namespace Horde.Build.Perforce
 		/// <param name="perforce">The perforce implementation</param>
 		/// <param name="clusterName">Name of the Perforce cluster</param>
 		/// <param name="streamName">The stream to query</param>
-		/// <param name="impersonateUser">Name of the user to impersonate</param>
 		/// <returns>Latest changelist number</returns>
-		public static async Task<int> GetLatestChangeAsync(this IPerforceService perforce, string clusterName, string streamName, string? impersonateUser)
+		public static async Task<int> GetLatestChangeAsync(this IPerforceService perforce, string clusterName, string streamName)
 		{
-			List<ChangeSummary> changes = await perforce.GetChangesAsync(clusterName, streamName, null, null, 1, impersonateUser);
+			List<ChangeSummary> changes = await perforce.GetChangesAsync(clusterName, streamName, null, null, 1);
 			if (changes.Count == 0)
 			{
 				throw new Exception($"No changes have been submitted to stream {streamName}");
diff --git a/Engine/Source/Programs/Horde/Horde.Build/Perforce/PerforceService.cs b/Engine/Source/Programs/Horde/Horde.Build/Perforce/PerforceService.cs
--- a/Engine/Source/Programs/Horde/Horde.Build/Perforce/PerforceService.cs
+++ b/Engine/Source/Programs/Horde/Horde.Build/Perforce/PerforceService.cs
@@ -589,19 +589,18 @@ namespace Horde.Build.Perforce
 		}
 
 		/// <inheritdoc/>
-		public async Task<List<ChangeSummary>> GetChangesAsync(string clusterName, string streamName, int? minChange, int? maxChange, int maxResults, string? impersonateUser, CancellationToken cancellationToken)
+		public async Task<List<ChangeSummary>> GetChangesAsync(string clusterName, string streamName, int? minChange, int? maxChange, int maxResults, CancellationToken cancellationToken)
 		{
 			using IScope scope = GlobalTracer.Instance.BuildSpan("PerforceService.GetChangesAsync").StartActive();
 			scope.Span.SetTag("ClusterName", clusterName);
 			scope.Span.SetTag("MinChange", minChange ?? -1);
 			scope.Span.SetTag("MaxChange", maxChange ?? -1);
 			scope.Span.SetTag("MaxResults", maxResults);
-			scope.Span.SetTag("ImpersonateUser", impersonateUser);
 			
 			PerforceCluster cluster = await GetClusterAsync(clusterName);
 
 			List<ChangeSummary> results = new List<ChangeSummary>();
-			using (IPerforceConnection perforce = await ConnectWithStreamClientAsync(cluster, impersonateUser, streamName, cancellationToken))
+			using (IPerforceConnection perforce = await ConnectWithStreamClientAsync(cluster, null, streamName, cancellationToken))
 			{
 				string filter = GetFilter($"//{perforce.Settings.ClientName}/...", minChange, maxChange);
 
@@ -649,18 +648,17 @@ namespace Horde.Build.Perforce
 		}
 
 		/// <inheritdoc/>
-		public async Task<(CheckShelfResult, string?)> CheckShelfAsync(string clusterName, string streamName, int changeNumber, string? impersonateUser, CancellationToken cancellationToken)
+		public async Task<(CheckShelfResult, string?)> CheckShelfAsync(string clusterName, string streamName, int changeNumber, CancellationToken cancellationToken)
 		{
 			using IScope scope = GlobalTracer.Instance.BuildSpan("PerforceService.CheckPreflightAsync").StartActive();
 			scope.Span.SetTag("ClusterName", clusterName);
 			scope.Span.SetTag("StreamName", streamName);
 			scope.Span.SetTag("ChangeNumber", changeNumber);
-			scope.Span.SetTag("ImpersonateUser", impersonateUser);
 
 			PerforceCluster cluster = await GetClusterAsync(clusterName);
 			IPerforceServer server = await GetServerAsync(cluster);
 
-			using (IPerforceConnection perforce = await ConnectAsync(cluster, impersonateUser, cancellationToken))
+			using (IPerforceConnection perforce = await ConnectAsync(cluster, null, cancellationToken))
 			{
 				PerforceResponse<DescribeRecord> response = await perforce.TryDescribeAsync(DescribeOptions.Shelved, -1, changeNumber, cancellationToken);
 				if(response.Error != null)
@@ -708,18 +706,17 @@ namespace Horde.Build.Perforce
 		}
 
 		/// <inheritdoc/>
-		public async Task<List<ChangeDetails>> GetChangeDetailsAsync(string clusterName, string streamName, IReadOnlyList<int> changeNumbers, string? impersonateUser, CancellationToken cancellationToken)
+		public async Task<List<ChangeDetails>> GetChangeDetailsAsync(string clusterName, string streamName, IReadOnlyList<int> changeNumbers, CancellationToken cancellationToken)
 		{
 			using IScope scope = GlobalTracer.Instance.BuildSpan("PerforceService.GetChangeDetailsAsync").StartActive();
 			scope.Span.SetTag("ClusterName", clusterName);
 			scope.Span.SetTag("StreamName", streamName);
 			scope.Span.SetTag("ChangeNumbers.Count", changeNumbers.Count);
-			scope.Span.SetTag("ImpersonateUser", impersonateUser);
 			
 			PerforceCluster cluster = await GetClusterAsync(clusterName);
 
 			List<ChangeDetails> results = new List<ChangeDetails>();
-			using (PooledConnectionHandle perforce = await ConnectWithStreamClientAsync(cluster, impersonateUser, streamName, cancellationToken))
+			using (PooledConnectionHandle perforce = await ConnectWithStreamClientAsync(cluster, null, streamName, cancellationToken))
 			{
 				InfoRecord info = await perforce.GetInfoAsync(cancellationToken);
 
@@ -1050,7 +1047,7 @@ namespace Horde.Build.Perforce
 			for (; ; )
 			{
 				// Query for the changes before this point
-				List<ChangeSummary> changes = await GetChangesAsync(clusterName, streamName, null, maxChange, 10, null, cancellationToken);
+				List<ChangeSummary> changes = await GetChangesAsync(clusterName, streamName, null, maxChange, 10, cancellationToken);
 				_logger.LogInformation("Finding last code change in {Stream} before {MaxChange}: {NumResults}", streamName, maxChange, changes.Count);
 				if (changes.Count == 0)
 				{
@@ -1058,7 +1055,7 @@ namespace Horde.Build.Perforce
 				}
 
 				// Get the details for them
-				List<ChangeDetails> detailsList = await GetChangeDetailsAsync(clusterName, streamName, changes.ConvertAll(x => x.Number), null, cancellationToken);
+				List<ChangeDetails> detailsList = await GetChangeDetailsAsync(clusterName, streamName, changes.ConvertAll(x => x.Number), cancellationToken);
 				foreach (ChangeDetails details in detailsList.OrderByDescending(x => x.Number))
 				{
 					ChangeContentFlags contentFlags = details.GetContentFlags();
diff --git a/Engine/Source/Programs/Horde/Horde.Build/Streams/StreamsController.cs b/Engine/Source/Programs/Horde/Horde.Build/Streams/StreamsController.cs
--- a/Engine/Source/Programs/Horde/Horde.Build/Streams/StreamsController.cs
+++ b/Engine/Source/Programs/Horde/Horde.Build/Streams/StreamsController.cs
@@ -188,13 +188,7 @@ namespace Horde.Build.Streams
 				return Forbid(AclAction.ViewChanges, streamId);
 			}
 
-			string? perforceUser = User.GetPerforceUser();
-			if(perforceUser == null)
-			{
-				return BadRequest("Current user does not have an associated Perforce user");
-			}
-
-			List<ChangeSummary> commits = await _perforceService.GetChangesAsync(stream.ClusterName, stream.Name, min, max, results, perforceUser);
+			List<ChangeSummary> commits = await _perforceService.GetChangesAsync(stream.ClusterName, stream.Name, min, max, results);
 
 			List<GetChangeSummaryResponse> responses = new List<GetChangeSummaryResponse>();
 			foreach (ChangeSummary commit in commits)
@@ -226,13 +220,7 @@ namespace Horde.Build.Streams
 				return Forbid(AclAction.ViewChanges, streamId);
 			}
 
-			string? perforceUser = User.GetPerforceUser();
-			if(perforceUser == null)
-			{
-				return BadRequest("Current user does not have an associated Perforce user");
-			}
-
-			ChangeDetails? changeDetails = await _perforceService.GetChangeDetailsAsync(stream.ClusterName, stream.Name, changeNumber, perforceUser);
+			ChangeDetails? changeDetails = await _perforceService.GetChangeDetailsAsync(stream.ClusterName, stream.Name, changeNumber);
 			if(changeDetails == null)
 			{
 				return NotFound("CL {Change} not found in stream {StreamId}", changeNumber, streamId);
