Horde: Fix tests.

diff --git a/Engine/Source/Programs/Horde/Horde.Build.Tests/Stubs/Services/PerforceServiceStub.cs b/Engine/Source/Programs/Horde/Horde.Build.Tests/Stubs/Services/PerforceServiceStub.cs
--- a/Engine/Source/Programs/Horde/Horde.Build.Tests/Stubs/Services/PerforceServiceStub.cs
+++ b/Engine/Source/Programs/Horde/Horde.Build.Tests/Stubs/Services/PerforceServiceStub.cs
@@ -82,7 +82,7 @@ namespace Horde.Build.Tests.Stubs.Services
 			streamChanges.Add(number, new ChangeDetails(number, author, null!, description, files.Select(x => PerforceExtensions.CreateChangeFile(x)).ToList(), DateTime.Now));
 		}
 
-		public Task<List<ChangeSummary>> GetChangesAsync(string clusterName, string streamName, int? minChange, int? maxChange, int numResults, string? impersonateUser, CancellationToken cancellationToken)
+		public Task<List<ChangeSummary>> GetChangesAsync(string clusterName, string streamName, int? minChange, int? maxChange, int numResults, CancellationToken cancellationToken)
 		{
 			List<ChangeSummary> results = new List<ChangeSummary>();
 
@@ -109,12 +109,12 @@ namespace Horde.Build.Tests.Stubs.Services
 			return Task.FromResult(results);
 		}
 
-		public Task<(CheckShelfResult, string?)> CheckShelfAsync(string clusterName, string streamName, int changeNumber, string? impersonateUser, CancellationToken cancellationToken)
+		public Task<(CheckShelfResult, string?)> CheckShelfAsync(string clusterName, string streamName, int changeNumber, CancellationToken cancellationToken)
 		{
 			throw new NotImplementedException();
 		}
 
-		public Task<List<ChangeDetails>> GetChangeDetailsAsync(string clusterName, string streamName, IReadOnlyList<int> changeNumbers, string? impersonateUser, CancellationToken cancellationToken)
+		public Task<List<ChangeDetails>> GetChangeDetailsAsync(string clusterName, string streamName, IReadOnlyList<int> changeNumbers, CancellationToken cancellationToken)
 		{
 			List<ChangeDetails> results = new List<ChangeDetails>();
 			foreach (int changeNumber in changeNumbers)
