Horde: Handle reading config files from Perforce as a special case, rather than including it in IPerforceService.

diff --git a/Engine/Source/Programs/Horde/Horde.Build.Tests/Stubs/Services/PerforceServiceStub.cs b/Engine/Source/Programs/Horde/Horde.Build.Tests/Stubs/Services/PerforceServiceStub.cs
--- a/Engine/Source/Programs/Horde/Horde.Build.Tests/Stubs/Services/PerforceServiceStub.cs
+++ b/Engine/Source/Programs/Horde/Horde.Build.Tests/Stubs/Services/PerforceServiceStub.cs
@@ -156,16 +156,6 @@ namespace Horde.Build.Tests.Stubs.Services
 			return Task.FromResult(newChange.Number);
 		}
 
-		public Task<List<FileSummary>> FindFilesAsync(string clusterName, IEnumerable<string> paths, CancellationToken cancellationToken)
-		{
-			throw new NotImplementedException();
-		}
-
-		public Task<byte[]> PrintAsync(string clusterName, string path, CancellationToken cancellationToken)
-		{
-			throw new NotImplementedException();
-		}
-
 		public Task<(int? Change, string Message)> SubmitShelvedChangeAsync(string clusterName, int shelvedChange, int originalChange, CancellationToken cancellationToken)
 		{
 			throw new NotImplementedException();
@@ -186,11 +176,6 @@ namespace Horde.Build.Tests.Stubs.Services
 			return Task.CompletedTask;
 		}
 
-		public Task<List<ChangeSummary>> GetChangesAsync(string clusterName, int? minChange, int? maxChange, int maxResults, CancellationToken cancellationToken)
-		{
-			throw new NotImplementedException();
-		}
-
 		public Task<ChangeDetails> GetChangeDetailsAsync(string clusterName, int changeNumber, CancellationToken cancellationToken)
 		{
 			throw new NotImplementedException();
diff --git a/Engine/Source/Programs/Horde/Horde.Build/Configuration/ConfigUpdateService.cs b/Engine/Source/Programs/Horde/Horde.Build/Configuration/ConfigUpdateService.cs
--- a/Engine/Source/Programs/Horde/Horde.Build/Configuration/ConfigUpdateService.cs
+++ b/Engine/Source/Programs/Horde/Horde.Build/Configuration/ConfigUpdateService.cs
@@ -11,6 +11,7 @@ using System.Text.Json.Serialization;
 using System.Threading;
 using System.Threading.Tasks;
 using EpicGames.Core;
+using EpicGames.Perforce;
 using Horde.Build.Acls;
 using Horde.Build.Agents;
 using Horde.Build.Agents.Pools;
@@ -31,7 +32,6 @@ using Microsoft.Extensions.Options;
 
 namespace Horde.Build.Configuration
 {
-	using PoolId = StringId<IPool>;
 	using ProjectId = StringId<IProject>;
 	using StreamId = StringId<IStream>;
 
@@ -56,7 +56,6 @@ namespace Horde.Build.Configuration
 		readonly IPerforceService _perforceService;
 		readonly INotificationService _notificationService;
 		readonly AgentService _agentService;
-		readonly PoolService _poolService;
 		readonly IOptionsMonitor<ServerSettings> _settings;
 		readonly ITicker _ticker;
 		readonly ILogger _logger;
@@ -64,7 +63,7 @@ namespace Horde.Build.Configuration
 		/// <summary>
 		/// Constructor
 		/// </summary>
-		public ConfigUpdateService(MongoService mongoService, ConfigCollection configCollection, IPerforceService perforceService, ToolCollection toolCollection, ProjectService projectService, StreamService streamService, INotificationService notificationService, PoolService poolService, AgentService agentService, IClock clock, IOptionsMonitor<ServerSettings> settings, ILogger<ConfigUpdateService> logger)
+		public ConfigUpdateService(MongoService mongoService, ConfigCollection configCollection, IPerforceService perforceService, ToolCollection toolCollection, ProjectService projectService, StreamService streamService, INotificationService notificationService, AgentService agentService, IClock clock, IOptionsMonitor<ServerSettings> settings, ILogger<ConfigUpdateService> logger)
 		{
 			_mongoService = mongoService;
 			_configCollection = configCollection;
@@ -73,9 +72,9 @@ namespace Horde.Build.Configuration
 			_projectService = projectService;
 			_streamService = streamService;
 			_notificationService = notificationService;
-			_poolService = poolService;
 			_agentService = agentService;
 			_settings = settings;
+
 			if (mongoService.ReadOnlyMode)
 			{
 				_ticker = new NullTicker();
@@ -84,6 +83,7 @@ namespace Horde.Build.Configuration
 			{
 				_ticker = clock.AddSharedTicker<ConfigUpdateService>(TimeSpan.FromMinutes(1.0), TickLeaderAsync, logger);
 			}
+
 			_logger = logger;
 		}
 
@@ -101,13 +101,13 @@ namespace Horde.Build.Configuration
 		Dictionary<ProjectId, (ProjectConfig Config, string Revision)> _cachedProjectConfigs = new Dictionary<ProjectId, (ProjectConfig, string)>();
 		readonly Dictionary<ProjectId, string?> _cachedLogoRevisions = new Dictionary<ProjectId, string?>();
 
-		async Task UpdateConfigAsync(Uri configPath)
+		async Task UpdateConfigAsync(Uri configPath, CancellationToken cancellationToken)
 		{
 			// Update the globals singleton
 			GlobalConfig globalConfig;
 			for (; ; )
 			{
-				Dictionary<Uri, string> globalRevisions = await FindRevisionsAsync(new[] { configPath });
+				Dictionary<Uri, string> globalRevisions = await FindRevisionsAsync(new[] { configPath }, cancellationToken);
 				if (globalRevisions.Count == 0)
 				{
 					throw new Exception($"Invalid config path: {configPath}");
@@ -119,12 +119,12 @@ namespace Horde.Build.Configuration
 					_logger.LogInformation("Caching global config from {Revision}", revision);
 					try
 					{
-						_cachedGlobalConfig = await ReadDataAsync<GlobalConfig>(revision, configPath);
+						_cachedGlobalConfig = await ReadDataAsync<GlobalConfig>(revision, configPath, cancellationToken);
 						_cachedGlobalConfigRevision = revision;
 					}
 					catch (Exception ex)
 					{
-						await SendFailureNotificationAsync(ex, configPath);
+						await SendFailureNotificationAsync(ex, configPath, cancellationToken);
 						return;
 					}
 				}
@@ -174,7 +174,7 @@ namespace Horde.Build.Configuration
 			HashSet<ProjectId> skipProjectIds = new HashSet<ProjectId>();
 
 			// Update any existing projects
-			Dictionary<Uri, string> projectRevisions = await FindRevisionsAsync(projectConfigs.Select(x => x.Path));
+			Dictionary<Uri, string> projectRevisions = await FindRevisionsAsync(projectConfigs.Select(x => x.Path), cancellationToken);
 			for (int idx = 0; idx < projectConfigs.Count; idx++)
 			{
 				// Make sure we were able to fetch metadata for 
@@ -199,7 +199,7 @@ namespace Horde.Build.Configuration
 					_logger.LogInformation("Caching configuration for project {ProjectId} ({Revision})", projectRef.Id, revision);
 					try
 					{
-						projectConfig = await ReadDataAsync<ProjectConfig>(revision, projectPath);
+						projectConfig = await ReadDataAsync<ProjectConfig>(revision, projectPath, cancellationToken);
 						if (update)
 						{
 							_logger.LogInformation("Updating configuration for project {ProjectId} ({Revision})", projectRef.Id, revision);
@@ -208,7 +208,7 @@ namespace Horde.Build.Configuration
 					}
 					catch (Exception ex)
 					{
-						await SendFailureNotificationAsync(ex, projectPath);
+						await SendFailureNotificationAsync(ex, projectPath, cancellationToken);
 						skipProjectIds.Add(projectRef.Id);
 						continue;
 					}
@@ -224,7 +224,7 @@ namespace Horde.Build.Configuration
 			}
 
 			// Get the logo revisions
-			Dictionary<Uri, string> logoRevisions = await FindRevisionsAsync(projectLogos.Select(x => x.Path));
+			Dictionary<Uri, string> logoRevisions = await FindRevisionsAsync(projectLogos.Select(x => x.Path), cancellationToken);
 			for (int idx = 0; idx < projectLogos.Count; idx++)
 			{
 				(ProjectId projectId, Uri path) = projectLogos[idx];
@@ -241,12 +241,12 @@ namespace Horde.Build.Configuration
 						_logger.LogInformation("Updating logo for project {ProjectId} ({Revision})", projectId, revision);
 						try
 						{
-							await _projectService.Collection.SetLogoAsync(projectId, path.ToString(), revision, GetMimeTypeFromPath(path), await ReadDataAsync(path));
+							await _projectService.Collection.SetLogoAsync(projectId, path.ToString(), revision, GetMimeTypeFromPath(path), await ReadDataAsync(path, cancellationToken));
 							_cachedLogoRevisions[projectId] = revision;
 						}
 						catch (Exception ex)
 						{
-							await SendFailureNotificationAsync(ex, path);
+							await SendFailureNotificationAsync(ex, path, cancellationToken);
 							continue;
 						}
 					}
@@ -257,7 +257,7 @@ namespace Horde.Build.Configuration
 			List<IStream> streams = await _streamService.GetStreamsAsync();
 
 			// Get the revisions for all the stream documents
-			Dictionary<Uri, string> streamRevisions = await FindRevisionsAsync(streamConfigs.Select(x => x.Path));
+			Dictionary<Uri, string> streamRevisions = await FindRevisionsAsync(streamConfigs.Select(x => x.Path), cancellationToken);
 			for (int idx = 0; idx < streamConfigs.Count; idx++)
 			{
 				(ProjectId projectId, StreamConfigRef streamRef, Uri streamPath) = streamConfigs[idx];
@@ -269,12 +269,12 @@ namespace Horde.Build.Configuration
 						_logger.LogInformation("Updating configuration for stream {StreamRef} ({Revision})", streamRef.Id, revision);
 						try
 						{
-							StreamConfig streamConfig = await ReadDataAsync<StreamConfig>(revision, streamPath);
+							StreamConfig streamConfig = await ReadDataAsync<StreamConfig>(revision, streamPath, cancellationToken);
 							stream = await _streamService.StreamCollection.CreateOrReplaceAsync(streamRef.Id, stream, revision, projectId);
 						}
 						catch (Exception ex)
 						{
-							await SendFailureNotificationAsync(ex, streamPath);
+							await SendFailureNotificationAsync(ex, streamPath, cancellationToken);
 							continue;
 						}
 					}
@@ -330,7 +330,7 @@ namespace Horde.Build.Configuration
 			return new Uri(baseUri, path);
 		}
 
-		async Task<Dictionary<Uri, string>> FindRevisionsAsync(IEnumerable<Uri> paths)
+		async Task<Dictionary<Uri, string>> FindRevisionsAsync(IEnumerable<Uri> paths, CancellationToken cancellationToken)
 		{
 			Dictionary<Uri, string> revisions = new Dictionary<Uri, string>();
 
@@ -355,17 +355,21 @@ namespace Horde.Build.Configuration
 			// Query all the Perforce revisions
 			foreach (IGrouping<string, Uri> perforcePath in perforcePaths.GroupBy(x => x.Host, StringComparer.OrdinalIgnoreCase))
 			{
-				List<FileSummary> files = await _perforceService.FindFilesAsync(perforcePath.Key, perforcePath.Select(x => x.AbsolutePath));
-				foreach (FileSummary file in files)
+				string[] absolutePaths = perforcePath.Select(x => x.AbsolutePath).ToArray();
+				using (IPerforceConnection perforce = await _perforceService.ConnectAsync(perforcePath.Key, null, cancellationToken))
 				{
-					Uri fileUri = new Uri($"{PerforceScheme}://{perforcePath.Key}{file.DepotPath}");
-					if (file.Error == null)
+					List<FStatRecord> files = await perforce.FStatAsync(FStatOptions.ShortenOutput, absolutePaths, cancellationToken).ToListAsync(cancellationToken);
+					foreach (string absolutePath in absolutePaths)
 					{
-						revisions[fileUri] = $"ver={Version},chg={file.Change},path={fileUri}";
-					}
-					else
-					{
-						_notificationService.NotifyConfigUpdateFailure(file.Error, file.DepotPath);
+						FStatRecord? file = files.FirstOrDefault(x => String.Equals(x.DepotFile, absolutePath, StringComparison.OrdinalIgnoreCase));
+						if (file == null || file.HeadAction == FileAction.Delete || file.HeadAction == FileAction.MoveDelete)
+						{
+							_notificationService.NotifyConfigUpdateFailure($"Unable to find '{absolutePath}'", absolutePath);
+							continue;
+						}
+
+						Uri fileUri = new Uri($"{PerforceScheme}://{perforcePath.Key}{file.DepotFile}");
+						revisions[fileUri] = $"ver={Version},chg={file.ChangeNumber},path={fileUri}";
 					}
 				}
 			}
@@ -373,27 +377,32 @@ namespace Horde.Build.Configuration
 			return revisions;
 		}
 
-		async Task<T> ReadDataAsync<T>(string revision, Uri configPath) where T : class
+		async Task<T> ReadDataAsync<T>(string revision, Uri configPath, CancellationToken cancellationToken) where T : class
 		{
-			byte[] data = await ReadDataAsync(configPath);
+			byte[] data = await ReadDataAsync(configPath, cancellationToken);
 			await _configCollection.AddConfigDataAsync(revision, data);
 			return await _configCollection.GetConfigAsync<T>(revision);
 		}
 
-		Task<byte[]> ReadDataAsync(Uri configPath)
+		async Task<byte[]> ReadDataAsync(Uri configPath, CancellationToken cancellationToken)
 		{
 			switch (configPath.Scheme)
 			{
 				case FileScheme:
-					return File.ReadAllBytesAsync(configPath.LocalPath);
+					return await File.ReadAllBytesAsync(configPath.LocalPath, cancellationToken);
 				case PerforceScheme:
-					return _perforceService.PrintAsync(configPath.Host, configPath.AbsolutePath);
+					using (IPerforceConnection perforce = await _perforceService.ConnectAsync(configPath.Host, null, cancellationToken))
+					{
+						PerforceResponse<PrintRecord<byte[]>> response = await perforce.TryPrintDataAsync(configPath.AbsolutePath, cancellationToken);
+						response.EnsureSuccess();
+						return response.Data.Contents!;
+					}
 				default:
 					throw new Exception($"Invalid config path: {configPath}");
 			}
 		}
 
-		async Task SendFailureNotificationAsync(Exception ex, Uri configPath)
+		async Task SendFailureNotificationAsync(Exception ex, Uri configPath, CancellationToken cancellationToken)
 		{
 			_logger.LogError(ex, "Unable to read data from {ConfigPath}: {Message}", configPath, ex.Message);
 
@@ -404,20 +413,22 @@ namespace Horde.Build.Configuration
 
 			if (configPath.Scheme == PerforceScheme)
 			{
-				try
+				using (IPerforceConnection perforce = await _perforceService.ConnectAsync(configPath.Host, null, cancellationToken))
 				{
-					List<FileSummary> files = await _perforceService.FindFilesAsync(configPath.Host, new[] { fileName });
-					change = files[0].Change;
-
-					List<ChangeSummary> changes = await _perforceService.GetChangesAsync(configPath.Host, change, change, 1);
-					if (changes.Count > 0 && changes[0].Number == change)
+					try
 					{
-						(author, description) = (changes[0].Author, changes[0].Description);
+						List<FStatRecord> records = await perforce.FStatAsync(FStatOptions.ShortenOutput, fileName, cancellationToken).ToListAsync(cancellationToken);
+						if (records.Count > 0)
+						{
+							DescribeRecord describeRecord = await perforce.DescribeAsync(records[0].HeadChange, cancellationToken);
+							author = await _perforceService.FindOrAddUserAsync(configPath.Host, describeRecord.User, cancellationToken);
+							description = describeRecord.Description;
+						}
+					}
+					catch (Exception ex2)
+					{
+						_logger.LogError(ex2, "Unable to identify change that last modified {ConfigPath} from Perforce", configPath);
 					}
-				}
-				catch (Exception ex2)
-				{
-					_logger.LogError(ex2, "Unable to identify change that last modified {ConfigPath} from Perforce", configPath);
 				}
 			}
 
@@ -427,7 +438,6 @@ namespace Horde.Build.Configuration
 		/// <inheritdoc/>
 		async ValueTask TickLeaderAsync(CancellationToken stoppingToken)
 		{
-
 			Uri? configUri = null;
 
 			if (Path.IsPathRooted(_settings.CurrentValue.ConfigPath) && !_settings.CurrentValue.ConfigPath.StartsWith("//", StringComparison.Ordinal))
@@ -443,7 +453,7 @@ namespace Horde.Build.Configuration
 
 			if (configUri != null)
 			{
-				await UpdateConfigAsync(configUri);
+				await UpdateConfigAsync(configUri, stoppingToken);
 			}
 		}
 	}
diff --git a/Engine/Source/Programs/Horde/Horde.Build/Notifications/Sinks/SlackNotificationSink.cs b/Engine/Source/Programs/Horde/Horde.Build/Notifications/Sinks/SlackNotificationSink.cs
--- a/Engine/Source/Programs/Horde/Horde.Build/Notifications/Sinks/SlackNotificationSink.cs
+++ b/Engine/Source/Programs/Horde/Horde.Build/Notifications/Sinks/SlackNotificationSink.cs
@@ -1797,44 +1797,6 @@ namespace Horde.Build.Notifications.Sinks
 
 		#endregion
 
-		#region Stream update (file)
-
-		/// <inheritdoc/>
-		public async Task NotifyStreamUpdateFailedAsync(FileSummary file)
-		{
-			_logger.LogDebug("Sending stream update failure notification for {File}", file.DepotPath);
-			if (_settings.UpdateStreamsNotificationChannel != null)
-			{
-				await SendStreamUpdateFailureMessage($"#{_settings.UpdateStreamsNotificationChannel}", file);
-			}
-		}
-
-		/// <summary>
-		/// Creates a stream update failure message in relation to a file
-		/// </summary>
-		/// <param name="recipient"></param>
-		/// <param name="file">The file</param>
-		/// <returns></returns>
-		async Task SendStreamUpdateFailureMessage(string recipient, FileSummary file)
-		{
-			string outcomeColor = ErrorColor;
-			SlackAttachment attachment = new SlackAttachment();
-			attachment.Color = outcomeColor;
-			attachment.FallbackText = $"{file.DepotPath} - Update Failure";
-
-			attachment.AddHeader($"Stream Update Failure :rip:", true);
-
-			attachment.AddSection($"<!here> Horde was unable to update {file.DepotPath}");
-			if (file.Error != null)
-			{
-				attachment.AddSection(QuoteText(file.Error));
-			}
-
-			await SendMessageAsync(recipient, attachment);
-		}
-
-		#endregion
-
 		#region Device notifications
 
 		/// <inheritdoc/>
diff --git a/Engine/Source/Programs/Horde/Horde.Build/Perforce/IPerforceService.cs b/Engine/Source/Programs/Horde/Horde.Build/Perforce/IPerforceService.cs
--- a/Engine/Source/Programs/Horde/Horde.Build/Perforce/IPerforceService.cs
+++ b/Engine/Source/Programs/Horde/Horde.Build/Perforce/IPerforceService.cs
@@ -88,17 +88,6 @@ namespace Horde.Build.Perforce
 		/// <returns>Code change for the latest change</returns>
 		public Task<int> GetCodeChangeAsync(string clusterName, string streamName, int change, CancellationToken cancellationToken = default);
 
-		/// <summary>
-		/// Finds changes submitted to a depot Gets the latest change for a particular stream
-		/// </summary>
-		/// <param name="clusterName">Name of the Perforce cluster</param>
-		/// <param name="minChange">The minimum changelist number</param>
-		/// <param name="maxChange">The maximum changelist number</param>
-		/// <param name="maxResults">Maximum number of results to return</param>
-		/// <param name="cancellationToken">Cancellation token for the operation</param>
-		/// <returns>Changelist information</returns>
-		public Task<List<ChangeSummary>> GetChangesAsync(string clusterName, int? minChange, int? maxChange, int maxResults, CancellationToken cancellationToken = default);
-
 		/// <summary>
 		/// Gets the latest change for a particular stream
 		/// </summary>
@@ -141,24 +130,6 @@ namespace Horde.Build.Perforce
 		/// <returns>Commit details</returns>
 		public Task<List<ChangeDetails>> GetChangeDetailsAsync(string clusterName, string streamName, IReadOnlyList<int> changeNumbers, CancellationToken cancellationToken = default);
 
-		/// <summary>
-		/// Gets the latest changes for a set of depot paths
-		/// </summary>
-		/// <param name="clusterName">Name of the Perforce cluster</param>
-		/// <param name="paths"></param>
-		/// <param name="cancellationToken">Cancellation token for the operation</param>
-		/// <returns></returns>
-		public Task<List<FileSummary>> FindFilesAsync(string clusterName, IEnumerable<string> paths, CancellationToken cancellationToken = default);
-
-		/// <summary>
-		/// Gets the contents of a file in the depot
-		/// </summary>
-		/// <param name="clusterName">Name of the Perforce cluster</param>
-		/// <param name="path">Path to read</param>
-		/// <param name="cancellationToken">Cancellation token for the operation</param>
-		/// <returns>Data for the file</returns>
-		public Task<byte[]> PrintAsync(string clusterName, string path, CancellationToken cancellationToken = default);
-
 		/// <summary>
 		/// Duplicates a shelved changelist
 		/// </summary>
diff --git a/Engine/Source/Programs/Horde/Horde.Build/Perforce/PerforceService.cs b/Engine/Source/Programs/Horde/Horde.Build/Perforce/PerforceService.cs
--- a/Engine/Source/Programs/Horde/Horde.Build/Perforce/PerforceService.cs
+++ b/Engine/Source/Programs/Horde/Horde.Build/Perforce/PerforceService.cs
@@ -565,29 +565,6 @@ namespace Horde.Build.Perforce
 			return new ChangeFile(relativePath, metaData.DepotFile, revision, metaData.FileSize, digest, metaData.Type);
 		}
 
-		/// <inheritdoc/>
-		public async Task<List<ChangeSummary>> GetChangesAsync(string clusterName, int? minChange, int? maxChange, int maxResults, CancellationToken cancellationToken)
-		{
-			using IScope scope = GlobalTracer.Instance.BuildSpan("PerforceService.GetChangesAsync").StartActive();
-			scope.Span.SetTag("ClusterName", clusterName);
-			scope.Span.SetTag("MinChange", minChange ?? -1);
-			scope.Span.SetTag("MaxResults", maxResults);
-
-			List<ChangeSummary> results = new List<ChangeSummary>();
-			using (IPerforceConnection perforce = await ConnectAsync(clusterName, null, cancellationToken))
-			{
-				string filter = GetFilter("//...", minChange, maxChange);
-				List<ChangesRecord> changes = await perforce.GetChangesAsync(ChangesOptions.IncludeTimes | ChangesOptions.LongOutput, maxResults, ChangeStatus.Submitted, filter, cancellationToken);
-
-				foreach (ChangesRecord change in changes)
-				{
-					IUser user = await FindOrAddUserAsync(clusterName, change.User, cancellationToken);
-					results.Add(new ChangeSummary(change.Number, user, change.Path ?? String.Empty, change.Description));
-				}
-			}
-			return results;
-		}
-
 		/// <inheritdoc/>
 		public async Task<List<ChangeSummary>> GetChangesAsync(string clusterName, string streamName, int? minChange, int? maxChange, int maxResults, CancellationToken cancellationToken)
 		{
@@ -742,56 +719,6 @@ namespace Horde.Build.Perforce
 			return results;
 		}
 
-		/// <inheritdoc/>
-		public async Task<List<FileSummary>> FindFilesAsync(string clusterName, IEnumerable<string> paths, CancellationToken cancellationToken)
-		{
-			using IScope scope = GlobalTracer.Instance.BuildSpan("PerforceService.FindFilesAsync").StartActive();
-			scope.Span.SetTag("ClusterName", clusterName);
-			
-			PerforceCluster cluster = await GetClusterAsync(clusterName);
-
-			List<FileSummary> results = new List<FileSummary>();
-			using (IPerforceConnection perforce = await ConnectAsync(cluster, null, cancellationToken))
-			{
-				List<FStatRecord> records = await perforce.FStatAsync(FStatOptions.ShortenOutput, paths.ToArray(), cancellationToken).ToListAsync(cancellationToken);
-				foreach (string path in paths)
-				{
-					// TODO: use server case
-					FStatRecord? record = records.FirstOrDefault(x => x.DepotFile == path);
-					if (record == null)
-					{
-						results.Add(new FileSummary(path, false, 0));
-					}
-					else
-					{
-						results.Add(new FileSummary(path, record.HeadAction != FileAction.Delete && record.HeadAction != FileAction.MoveDelete, record.HeadChange));
-					}
-				}
-				return results;
-			}
-		}
-
-		/// <inheritdoc/>
-		public async Task<byte[]> PrintAsync(string clusterName, string depotPath, CancellationToken cancellationToken)
-		{
-			using IScope scope = GlobalTracer.Instance.BuildSpan("PerforceService.PrintAsync").StartActive();
-			scope.Span.SetTag("ClusterName", clusterName);
-			scope.Span.SetTag("DepotPath", depotPath);
-			
-			if (depotPath.EndsWith("...", StringComparison.OrdinalIgnoreCase) || depotPath.EndsWith("*", StringComparison.OrdinalIgnoreCase))
-			{
-				throw new Exception("PrintAsync requires exactly one file to be specified");
-			}
-
-			PerforceCluster cluster = await GetClusterAsync(clusterName);
-
-			using (IPerforceConnection perforce = await ConnectAsync(cluster, null, cancellationToken))
-			{
-				PerforceResponse<PrintRecord<byte[]>> response = await perforce.TryPrintDataAsync(depotPath, cancellationToken);
-				return response.Data.Contents!;
-			}
-		}
-
 		/// <inheritdoc/>
 		public Task<int> DuplicateShelvedChangeAsync(string clusterName, int shelvedChange, CancellationToken cancellationToken)
 		{
