Horde: Fix notifications being sent to triage alias for closed issues.

diff --git a/Engine/Source/Programs/Horde/Horde.Build/Notifications/Sinks/SlackNotificationSink.cs b/Engine/Source/Programs/Horde/Horde.Build/Notifications/Sinks/SlackNotificationSink.cs
--- a/Engine/Source/Programs/Horde/Horde.Build/Notifications/Sinks/SlackNotificationSink.cs
+++ b/Engine/Source/Programs/Horde/Horde.Build/Notifications/Sinks/SlackNotificationSink.cs
@@ -803,7 +803,8 @@ namespace Horde.Build.Notifications.Sinks
 
 				string issueSummary = issue.UserSummary ?? issue.Summary;
 				string text = $"{workflow.TriagePrefix}{GetSeverityPrefix(issue.Severity)}Issue <{issueUrl}|{issue.Id}>: {issueSummary}{workflow.TriageSuffix}";
-				if (!spans.Any(x => x.NextSuccess == null && x.LastFailure.Annotations.WorkflowId != null)) // Thread may be shared by multiple workflows
+				bool closed = !spans.Any(x => x.NextSuccess == null && x.LastFailure.Annotations.WorkflowId != null);
+				if (closed) // Thread may be shared by multiple workflows
 				{
 					text = $"~{text}~";
 				}
@@ -975,7 +976,7 @@ namespace Horde.Build.Notifications.Sinks
 					}
 				}
 
-				if (workflow.TriageAlias != null && issue.OwnerId == null && (suspects.All(x => x.DeclinedAt != null) || notifyTriageAlias))
+				if (workflow.TriageAlias != null && issue.OwnerId == null && (suspects.All(x => x.DeclinedAt != null) || notifyTriageAlias) && !closed)
 				{
 					string triageMessage = $"(cc {FormatUserOrGroupMention(workflow.TriageAlias)} for triage).";
 					await SendOrUpdateMessageAsync(triageChannel, state.Ts, eventId + "_triage", null, triageMessage);
