Horde: Add an "Analyze" configuration to the Horde.Build project which enables the standard NET analyzers, and fix up existing warnings.

diff --git a/Engine/Source/Programs/Horde/Horde.Build/Agents/AgentMessages.cs b/Engine/Source/Programs/Horde/Horde.Build/Agents/AgentMessages.cs
--- a/Engine/Source/Programs/Horde/Horde.Build/Agents/AgentMessages.cs
+++ b/Engine/Source/Programs/Horde/Horde.Build/Agents/AgentMessages.cs
@@ -5,7 +5,6 @@ using System.Collections.Generic;
 using System.ComponentModel.DataAnnotations;
 using System.Linq;
 using Horde.Build.Acls;
-using Horde.Build.Agents;
 using EpicGames.Core;
 using HordeCommon;
 using Horde.Build.Agents.Leases;
diff --git a/Engine/Source/Programs/Horde/Horde.Build/Agents/Fleet/AutoscaleService.cs b/Engine/Source/Programs/Horde/Horde.Build/Agents/Fleet/AutoscaleService.cs
--- a/Engine/Source/Programs/Horde/Horde.Build/Agents/Fleet/AutoscaleService.cs
+++ b/Engine/Source/Programs/Horde/Horde.Build/Agents/Fleet/AutoscaleService.cs
@@ -232,7 +232,7 @@ namespace Horde.Build.Agents.Fleet
 
 							if (delta > 0)
 							{
-								await _fleetManager.ExpandPoolAsync(pool, poolData.Agents, delta);
+								await _fleetManager.ExpandPoolAsync(pool, poolData.Agents, delta, stoppingToken);
 								await _poolCollection.TryUpdateAsync(pool, lastScaleUpTime: DateTime.UtcNow);
 							}
 
@@ -241,7 +241,7 @@ namespace Horde.Build.Agents.Fleet
 								bool bShrinkIsOnCoolDown = pool.LastScaleDownTime != null && pool.LastScaleDownTime + s_shrinkPoolCoolDown > DateTime.UtcNow;
 								if (!bShrinkIsOnCoolDown)
 								{
-									await _fleetManager.ShrinkPoolAsync(pool, poolData.Agents, -delta);
+									await _fleetManager.ShrinkPoolAsync(pool, poolData.Agents, -delta, stoppingToken);
 									await _poolCollection.TryUpdateAsync(pool, lastScaleDownTime: DateTime.UtcNow);
 								}
 								else
diff --git a/Engine/Source/Programs/Horde/Horde.Build/Agents/Fleet/AutoscaleServiceV2.cs b/Engine/Source/Programs/Horde/Horde.Build/Agents/Fleet/AutoscaleServiceV2.cs
--- a/Engine/Source/Programs/Horde/Horde.Build/Agents/Fleet/AutoscaleServiceV2.cs
+++ b/Engine/Source/Programs/Horde/Horde.Build/Agents/Fleet/AutoscaleServiceV2.cs
@@ -36,7 +36,6 @@ namespace Horde.Build.Agents.Fleet
 		private readonly IClock _clock;
 		private readonly ITicker _ticker;
 		private readonly ITicker _tickerHighFrequency;
-		private readonly ServerSettings _settings;
 		private readonly TimeSpan _defaultScaleOutCooldown;
 		private readonly TimeSpan _defaultScaleInCooldown;
 		private readonly ILogger<AutoscaleServiceV2> _logger;
@@ -68,7 +67,6 @@ namespace Horde.Build.Agents.Fleet
 			_clock = clock;
 			_ticker = clock.AddSharedTicker<AutoscaleServiceV2>(TimeSpan.FromMinutes(5.0), TickLeaderAsync, logger);
 			_tickerHighFrequency = clock.AddSharedTicker("AutoscaleServiceV2.TickHighFrequency", TimeSpan.FromSeconds(30), TickHighFrequencyAsync, logger);
-			_settings = settings.Value;
 			_logger = logger;
 			_defaultScaleOutCooldown = TimeSpan.FromSeconds(settings.Value.AgentPoolScaleOutCooldownSeconds);
 			_defaultScaleInCooldown = TimeSpan.FromSeconds(settings.Value.AgentPoolScaleInCooldownSeconds);
diff --git a/Engine/Source/Programs/Horde/Horde.Build/Agents/Fleet/ComputeQueueAwsMetricStrategy.cs b/Engine/Source/Programs/Horde/Horde.Build/Agents/Fleet/ComputeQueueAwsMetricStrategy.cs
--- a/Engine/Source/Programs/Horde/Horde.Build/Agents/Fleet/ComputeQueueAwsMetricStrategy.cs
+++ b/Engine/Source/Programs/Horde/Horde.Build/Agents/Fleet/ComputeQueueAwsMetricStrategy.cs
@@ -8,17 +8,13 @@ using System.Threading.Tasks;
 using Amazon.CloudWatch;
 using Amazon.CloudWatch.Model;
 using EpicGames.Horde.Compute;
-using Horde.Build.Agents.Pools;
 using Horde.Build.Compute;
-using Horde.Build.Utilities;
 using Microsoft.Extensions.Logging;
 using OpenTracing;
 using OpenTracing.Util;
 
 namespace Horde.Build.Agents.Fleet
 {
-	using PoolId = StringId<IPool>;
-
 	/// <summary>
 	/// Settings for <see cref="ComputeQueueAwsMetricStrategy" />
 	/// </summary>
diff --git a/Engine/Source/Programs/Horde/Horde.Build/Agents/Fleet/Providers/AwsFleetManager.cs b/Engine/Source/Programs/Horde/Horde.Build/Agents/Fleet/Providers/AwsFleetManager.cs
--- a/Engine/Source/Programs/Horde/Horde.Build/Agents/Fleet/Providers/AwsFleetManager.cs
+++ b/Engine/Source/Programs/Horde/Horde.Build/Agents/Fleet/Providers/AwsFleetManager.cs
@@ -231,9 +231,12 @@ namespace Horde.Build.Agents.Fleet.Providers
 		/// <param name="agentCollection">Instance of the AgentCollection</param>
 		/// <param name="pool">Pool to resize</param>
 		/// <param name="agents">Current list of agents in the pool</param>
-		/// <param name="count">Number of agents to remove</param>		/// <param name="cancellationToken">Cancellation token for the call</param>
+		/// <param name="count">Number of agents to remove</param>
+		/// <param name="cancellationToken"></param>
 		public static async Task ShrinkPoolViaAgentShutdownRequestAsync(IAgentCollection agentCollection, IPool pool, IReadOnlyList<IAgent> agents, int count, CancellationToken cancellationToken)
 		{
+			_ = cancellationToken;
+
 			using IScope scope = GlobalTracer.Instance.BuildSpan("ShrinkPool").StartActive();
 			scope.Span.SetTag("poolName", pool.Name);
 			scope.Span.SetTag("count", count);
@@ -273,7 +276,7 @@ namespace Horde.Build.Agents.Fleet.Providers
 			describeRequest.Filters.Add(new Filter("instance-state-name", new List<string> { InstanceStateName.Stopped.Value }));
 			describeRequest.Filters.Add(new Filter("tag:" + PoolTagName, new List<string> { pool.Name }));
 
-			DescribeInstancesResponse describeResponse = await _client.DescribeInstancesAsync(describeRequest);
+			DescribeInstancesResponse describeResponse = await _client.DescribeInstancesAsync(describeRequest, cancellationToken);
 			return describeResponse.Reservations.SelectMany(x => x.Instances).Select(x => x.InstanceId).Distinct().Count();
 		}
 	}
diff --git a/Engine/Source/Programs/Horde/Horde.Build/Agents/Fleet/Providers/AwsReuseFleetManager.cs b/Engine/Source/Programs/Horde/Horde.Build/Agents/Fleet/Providers/AwsReuseFleetManager.cs
--- a/Engine/Source/Programs/Horde/Horde.Build/Agents/Fleet/Providers/AwsReuseFleetManager.cs
+++ b/Engine/Source/Programs/Horde/Horde.Build/Agents/Fleet/Providers/AwsReuseFleetManager.cs
@@ -65,7 +65,7 @@ namespace Horde.Build.Agents.Fleet.Providers
 				describeRequest.Filters = new List<Filter>();
 				describeRequest.Filters.Add(new Filter("instance-state-name", new List<string> { InstanceStateName.Stopped.Value }));
 				describeRequest.Filters.Add(new Filter("tag:" + AwsFleetManager.PoolTagName, new List<string> { pool.Name }));
-				describeResponse = await _client.DescribeInstancesAsync(describeRequest);
+				describeResponse = await _client.DescribeInstancesAsync(describeRequest, cancellationToken);
 				describeScope.Span.SetTag("res.statusCode", (int)describeResponse.HttpStatusCode);
 				describeScope.Span.SetTag("res.numReservations", describeResponse.Reservations.Count);
 			}
@@ -79,7 +79,7 @@ namespace Horde.Build.Agents.Fleet.Providers
 				startScope.Span.SetTag("req.instanceIds", String.Join(",", startRequest.InstanceIds));
 				if (startRequest.InstanceIds.Count > 0)
 				{
-					StartInstancesResponse startResponse = await _client.StartInstancesAsync(startRequest);
+					StartInstancesResponse startResponse = await _client.StartInstancesAsync(startRequest, cancellationToken);
 					startScope.Span.SetTag("res.statusCode", (int)startResponse.HttpStatusCode);
 					startScope.Span.SetTag("res.numInstances", startResponse.StartingInstances.Count);
 					if ((int)startResponse.HttpStatusCode >= 200 && (int)startResponse.HttpStatusCode <= 299)
@@ -114,7 +114,7 @@ namespace Horde.Build.Agents.Fleet.Providers
 			describeRequest.Filters.Add(new Filter("instance-state-name", new List<string> { InstanceStateName.Stopped.Value }));
 			describeRequest.Filters.Add(new Filter("tag:" + AwsFleetManager.PoolTagName, new List<string> { pool.Name }));
 
-			DescribeInstancesResponse describeResponse = await _client.DescribeInstancesAsync(describeRequest);
+			DescribeInstancesResponse describeResponse = await _client.DescribeInstancesAsync(describeRequest, cancellationToken);
 			return describeResponse.Reservations.SelectMany(x => x.Instances).Select(x => x.InstanceId).Distinct().Count();
 		}
 	}
diff --git a/Engine/Source/Programs/Horde/Horde.Build/Agents/Pools/PoolCollection.cs b/Engine/Source/Programs/Horde/Horde.Build/Agents/Pools/PoolCollection.cs
--- a/Engine/Source/Programs/Horde/Horde.Build/Agents/Pools/PoolCollection.cs
+++ b/Engine/Source/Programs/Horde/Horde.Build/Agents/Pools/PoolCollection.cs
@@ -43,7 +43,7 @@ namespace Horde.Build.Agents.Pools
 			public TimeSpan? ScaleOutCooldown { get; set; }
 			public TimeSpan? ScaleInCooldown { get; set; }
 			public PoolSizeStrategy? SizeStrategy { get; set; }
-			PoolSizeStrategy IPool.SizeStrategy => SizeStrategy ?? defaultStrategy!.Value;
+			PoolSizeStrategy IPool.SizeStrategy => SizeStrategy ?? s_defaultStrategy!.Value;
 			public LeaseUtilizationSettings? LeaseUtilizationSettings { get; set; }
 			public JobQueueSettings? JobQueueSettings { get; set; }
 			public ComputeQueueAwsMetricSettings? ComputeQueueAwsMetricSettings { get; set; }
@@ -86,7 +86,7 @@ namespace Horde.Build.Agents.Pools
 		/// </summary>
 		readonly IMongoCollection<PoolDocument> _pools;
 
-		static PoolSizeStrategy? defaultStrategy;
+		static PoolSizeStrategy? s_defaultStrategy;
 
 		/// <summary>
 		/// Constructor
@@ -96,7 +96,7 @@ namespace Horde.Build.Agents.Pools
 		public PoolCollection(MongoService mongoService, IOptions<ServerSettings> serverSettings)
 		{
 			_pools = mongoService.GetCollection<PoolDocument>("Pools");
-			defaultStrategy = serverSettings.Value.DefaultAgentPoolSizeStrategy;
+			s_defaultStrategy = serverSettings.Value.DefaultAgentPoolSizeStrategy;
 		}
 
 		/// <inheritdoc/>
diff --git a/Engine/Source/Programs/Horde/Horde.Build/Agents/Telemetry/TelemetryService.cs b/Engine/Source/Programs/Horde/Horde.Build/Agents/Telemetry/TelemetryService.cs
--- a/Engine/Source/Programs/Horde/Horde.Build/Agents/Telemetry/TelemetryService.cs
+++ b/Engine/Source/Programs/Horde/Horde.Build/Agents/Telemetry/TelemetryService.cs
@@ -93,7 +93,7 @@ namespace Horde.Build.Agents.Telemetry
 				{
 					if (pool.EnableAutoscaling)
 					{
-						int numStoppedInstances = await _fleetManager.GetNumStoppedInstancesAsync(pool);
+						int numStoppedInstances = await _fleetManager.GetNumStoppedInstancesAsync(pool, stoppingToken);
 						telemetry.NumAgents += numStoppedInstances;
 
 						NewPoolUtilizationTelemetry poolTelemetry = telemetry.FindOrAddPool(pool.Id);
diff --git a/Engine/Source/Programs/Horde/Horde.Build/Auditing/IAuditLog.cs b/Engine/Source/Programs/Horde/Horde.Build/Auditing/IAuditLog.cs
--- a/Engine/Source/Programs/Horde/Horde.Build/Auditing/IAuditLog.cs
+++ b/Engine/Source/Programs/Horde/Horde.Build/Auditing/IAuditLog.cs
@@ -126,7 +126,7 @@ namespace Horde.Build.Auditing
 			await bodyWriter.WriteAsync(Encoding.UTF8.GetBytes(prefix));
 
 			string separator = "";
-			await foreach (IAuditLogMessage<T> message in channel.FindAsync(minTime, maxTime, index, count))
+			await foreach (IAuditLogMessage message in channel.FindAsync(minTime, maxTime, index, count))
 			{
 				string line = $"{separator}\n\t\t{message.Data}";
 				await bodyWriter.WriteAsync(Encoding.UTF8.GetBytes(line));
diff --git a/Engine/Source/Programs/Horde/Horde.Build/Commands/Bundles/CreateCommand.cs b/Engine/Source/Programs/Horde/Horde.Build/Commands/Bundles/CreateCommand.cs
--- a/Engine/Source/Programs/Horde/Horde.Build/Commands/Bundles/CreateCommand.cs
+++ b/Engine/Source/Programs/Horde/Horde.Build/Commands/Bundles/CreateCommand.cs
@@ -1,17 +1,11 @@
 // Copyright Epic Games, Inc. All Rights Reserved.
 
-using System;
-using System.Buffers;
-using System.Collections.Generic;
-using System.IO;
 using System.Threading;
 using System.Threading.Tasks;
 using EpicGames.Core;
 using EpicGames.Horde.Storage;
-using EpicGames.Horde.Storage.Bundles;
 using EpicGames.Horde.Storage.Nodes;
 using Horde.Build.Perforce;
-using Microsoft.Extensions.Caching.Memory;
 using Microsoft.Extensions.Configuration;
 using Microsoft.Extensions.DependencyInjection;
 using Microsoft.Extensions.Logging;
@@ -27,8 +21,8 @@ namespace Horde.Build.Commands.Bundles
 		[CommandLine("-InputDir=", Required = true)]
 		public DirectoryReference InputDir { get; set; } = null!;
 
-		protected readonly IConfiguration _configuration;
-		protected readonly ILoggerProvider _loggerProvider;
+		readonly IConfiguration _configuration;
+		readonly ILoggerProvider _loggerProvider;
 
 		public CreateCommand(IConfiguration configuration, ILoggerProvider loggerProvider)
 		{
diff --git a/Engine/Source/Programs/Horde/Horde.Build/Commands/Bundles/ExtractCommand.cs b/Engine/Source/Programs/Horde/Horde.Build/Commands/Bundles/ExtractCommand.cs
--- a/Engine/Source/Programs/Horde/Horde.Build/Commands/Bundles/ExtractCommand.cs
+++ b/Engine/Source/Programs/Horde/Horde.Build/Commands/Bundles/ExtractCommand.cs
@@ -22,8 +22,8 @@ namespace Horde.Build.Commands.Bundles
 		[CommandLine("-OutputDir=", Required = true)]
 		public DirectoryReference OutputDir { get; set; } = null!;
 
-		protected readonly IConfiguration _configuration;
-		protected readonly ILoggerProvider _loggerProvider;
+		readonly IConfiguration _configuration;
+		readonly ILoggerProvider _loggerProvider;
 
 		public ExtractCommand(IConfiguration configuration, ILoggerProvider loggerProvider)
 		{
diff --git a/Engine/Source/Programs/Horde/Horde.Build/Commands/Bundles/PerforceCommand.cs b/Engine/Source/Programs/Horde/Horde.Build/Commands/Bundles/PerforceCommand.cs
--- a/Engine/Source/Programs/Horde/Horde.Build/Commands/Bundles/PerforceCommand.cs
+++ b/Engine/Source/Programs/Horde/Horde.Build/Commands/Bundles/PerforceCommand.cs
@@ -1,29 +1,19 @@
 // Copyright Epic Games, Inc. All Rights Reserved.
 
 using System;
-using System.Buffers;
 using System.Collections.Generic;
-using System.IO;
 using System.Linq;
 using System.Threading;
 using System.Threading.Tasks;
 using EpicGames.Core;
 using EpicGames.Horde.Storage;
-using EpicGames.Horde.Storage.Bundles;
 using EpicGames.Horde.Storage.Nodes;
-using Horde.Build.Agents.Leases;
-using Horde.Build.Configuration;
-using Horde.Build.Jobs.Templates;
 using Horde.Build.Perforce;
-using Horde.Build.Server;
 using Horde.Build.Streams;
 using Horde.Build.Utilities;
-using HordeCommon;
-using Microsoft.Extensions.Caching.Memory;
 using Microsoft.Extensions.Configuration;
 using Microsoft.Extensions.DependencyInjection;
 using Microsoft.Extensions.Logging;
-using Microsoft.Extensions.Options;
 
 namespace Horde.Build.Commands.Bundles
 {
@@ -59,8 +49,8 @@ namespace Horde.Build.Commands.Bundles
 		[CommandLine]
 		public DirectoryReference? OutputDir { get; set; }
 
-		protected readonly IConfiguration _configuration;
-		protected readonly ILoggerProvider _loggerProvider;
+		readonly IConfiguration _configuration;
+		readonly ILoggerProvider _loggerProvider;
 
 		public PerforceCommand(IConfiguration configuration, ILoggerProvider loggerProvider)
 		{
diff --git a/Engine/Source/Programs/Horde/Horde.Build/Commands/Bundles/RepoCommand.cs b/Engine/Source/Programs/Horde/Horde.Build/Commands/Bundles/RepoCommand.cs
--- a/Engine/Source/Programs/Horde/Horde.Build/Commands/Bundles/RepoCommand.cs
+++ b/Engine/Source/Programs/Horde/Horde.Build/Commands/Bundles/RepoCommand.cs
@@ -1,20 +1,14 @@
 // Copyright Epic Games, Inc. All Rights Reserved.
 
-using System;
-using System.Buffers;
-using System.Collections.Generic;
 using System.IO;
 using System.IO.Compression;
-using System.Security.Cryptography;
 using System.Threading;
 using System.Threading.Tasks;
 using EpicGames.Core;
 using EpicGames.Horde.Storage;
-using EpicGames.Horde.Storage.Bundles;
 using EpicGames.Horde.Storage.Git;
 using EpicGames.Horde.Storage.Nodes;
 using Horde.Build.Perforce;
-using Microsoft.Extensions.Caching.Memory;
 using Microsoft.Extensions.Configuration;
 using Microsoft.Extensions.DependencyInjection;
 using Microsoft.Extensions.Logging;
@@ -30,8 +24,8 @@ namespace Horde.Build.Commands.Bundles
 		[CommandLine("-OutputDir=", Required = true)]
 		public DirectoryReference OutputDir { get; set; } = null!;
 
-		protected readonly IConfiguration _configuration;
-		protected readonly ILoggerProvider _loggerProvider;
+		readonly IConfiguration _configuration;
+		readonly ILoggerProvider _loggerProvider;
 
 		public RepoCommand(IConfiguration configuration, ILoggerProvider loggerProvider)
 		{
diff --git a/Engine/Source/Programs/Horde/Horde.Build/Commands/ServerCommand.cs b/Engine/Source/Programs/Horde/Horde.Build/Commands/ServerCommand.cs
--- a/Engine/Source/Programs/Horde/Horde.Build/Commands/ServerCommand.cs
+++ b/Engine/Source/Programs/Horde/Horde.Build/Commands/ServerCommand.cs
@@ -1,9 +1,7 @@
 // Copyright Epic Games, Inc. All Rights Reserved.
 
 using System;
-using System.Collections.Generic;
 using System.IO;
-using System.Linq;
 using System.Security.Cryptography.X509Certificates;
 using System.Threading.Tasks;
 using EpicGames.Core;
diff --git a/Engine/Source/Programs/Horde/Horde.Build/Compute/ComputeService.cs b/Engine/Source/Programs/Horde/Horde.Build/Compute/ComputeService.cs
--- a/Engine/Source/Programs/Horde/Horde.Build/Compute/ComputeService.cs
+++ b/Engine/Source/Programs/Horde/Horde.Build/Compute/ComputeService.cs
@@ -65,7 +65,7 @@ namespace Horde.Build.Compute
 	/// <summary>
 	/// Dispatches remote actions. Does not implement any cross-pod communication to satisfy leases; only agents connected to this server instance will be stored.
 	/// </summary>
-	public class ComputeService : TaskSourceBase<ComputeTaskMessage>, IHostedService, IDisposable, IComputeService
+	public sealed class ComputeService : TaskSourceBase<ComputeTaskMessage>, IHostedService, IDisposable, IComputeService
 	{
 		[RedisConverter(typeof(QueueKeySerializer))]
 		class QueueKey
diff --git a/Engine/Source/Programs/Horde/Horde.Build/Configuration/ConfigUpdateService.cs b/Engine/Source/Programs/Horde/Horde.Build/Configuration/ConfigUpdateService.cs
--- a/Engine/Source/Programs/Horde/Horde.Build/Configuration/ConfigUpdateService.cs
+++ b/Engine/Source/Programs/Horde/Horde.Build/Configuration/ConfigUpdateService.cs
@@ -2,19 +2,14 @@
 
 using System;
 using System.Collections.Generic;
-using System.Globalization;
 using System.IO;
 using System.Linq;
-using System.Reflection;
-using System.Text.Json;
-using System.Text.Json.Serialization;
 using System.Threading;
 using System.Threading.Tasks;
 using EpicGames.Core;
 using EpicGames.Perforce;
 using Horde.Build.Acls;
 using Horde.Build.Agents;
-using Horde.Build.Agents.Pools;
 using Horde.Build.Notifications;
 using Horde.Build.Perforce;
 using Horde.Build.Projects;
@@ -25,7 +20,6 @@ using Horde.Build.Users;
 using Horde.Build.Utilities;
 using HordeCommon;
 using Microsoft.AspNetCore.StaticFiles;
-using Microsoft.Extensions.Configuration;
 using Microsoft.Extensions.Hosting;
 using Microsoft.Extensions.Logging;
 using Microsoft.Extensions.Options;
diff --git a/Engine/Source/Programs/Horde/Horde.Build/Devices/DeviceCollection.cs b/Engine/Source/Programs/Horde/Horde.Build/Devices/DeviceCollection.cs
--- a/Engine/Source/Programs/Horde/Horde.Build/Devices/DeviceCollection.cs
+++ b/Engine/Source/Programs/Horde/Horde.Build/Devices/DeviceCollection.cs
@@ -227,7 +227,6 @@ namespace Horde.Build.Devices
 			[BsonIgnoreIfNull]
 			public string? Notes { get; set; }
 
-
 			/// <summary>
 			/// [DEPRECATED]
 			/// </summary>
@@ -352,7 +351,6 @@ namespace Horde.Build.Devices
 				JobName = jobName;
 				StepName = stepName;
 			}
-
 		}
 
 		class DeviceReservationPoolTelemetryDocument : IDevicePoolReservationTelemetry
@@ -386,7 +384,6 @@ namespace Horde.Build.Devices
 				JobName = jobName;
 				StepName = stepName;
 			}
-
 		}
 
 		class DevicePlatformTelemetryDocument : IDevicePlatformTelemetry
@@ -488,7 +485,6 @@ namespace Horde.Build.Devices
 			}
 		}
 
-
 		readonly IMongoCollection<DevicePlatformDocument> _platforms;
 
 		readonly IMongoCollection<DeviceDocument> _devices;
@@ -903,10 +899,7 @@ namespace Horde.Build.Devices
 
 				List<DeviceUtilizationTelemetry>? utilization = device.Utilization;
 
-				if (utilization == null)
-				{
-					utilization = new List<DeviceUtilizationTelemetry>();
-				}
+				utilization ??= new List<DeviceUtilizationTelemetry>();
 
 				// keep up to 100, maintaining order
 				if (utilization.Count > 99)
@@ -1108,11 +1101,11 @@ namespace Horde.Build.Devices
 
 		struct DevicePoolTelemetryHelper
 		{
-			public List<DeviceId> available = new List<DeviceId>();
-			public List<DeviceId> problem = new List<DeviceId>();
-			public List<DeviceId> maintenance = new List<DeviceId>();
-			public List<DeviceId> disabled = new List<DeviceId>();
-			public Dictionary<StreamId, List<DeviceReservationPoolTelemetryDocument>> reserved = new Dictionary<StreamId, List<DeviceReservationPoolTelemetryDocument>>();
+			public List<DeviceId> _available = new List<DeviceId>();
+			public List<DeviceId> _problem = new List<DeviceId>();
+			public List<DeviceId> _maintenance = new List<DeviceId>();
+			public List<DeviceId> _disabled = new List<DeviceId>();
+			public Dictionary<StreamId, List<DeviceReservationPoolTelemetryDocument>> _reserved = new Dictionary<StreamId, List<DeviceReservationPoolTelemetryDocument>>();
 
 			public DevicePoolTelemetryHelper() { }
 		}
@@ -1168,10 +1161,10 @@ namespace Horde.Build.Devices
 						{
 							StreamId streamId = new StreamId(reservation.StreamId);
 							List<DeviceReservationPoolTelemetryDocument>? rdevices;
-							if (!helper.reserved.TryGetValue(streamId, out rdevices))
+							if (!helper._reserved.TryGetValue(streamId, out rdevices))
 							{
 								rdevices = new List<DeviceReservationPoolTelemetryDocument>();
-								helper.reserved[streamId] = rdevices;
+								helper._reserved[streamId] = rdevices;
 							}
 							rdevices.Add(new DeviceReservationPoolTelemetryDocument(device.Id, reservation.JobId, reservation.StepId, reservation.JobName, reservation.StepName));
 						}
@@ -1180,33 +1173,32 @@ namespace Horde.Build.Devices
 
 					if (problemDevices.Contains(device))
 					{
-						helper.problem.Add(device.Id);
+						helper._problem.Add(device.Id);
 						continue;
 					}
 
 					if (maintenanceDevices.Contains(device))
 					{
-						helper.maintenance.Add(device.Id);
+						helper._maintenance.Add(device.Id);
 						continue;
 					}
 
 					if (disabledDevices.Contains(device))
 					{
-						helper.disabled.Add(device.Id);
+						helper._disabled.Add(device.Id);
 						continue;
 					}
 
-					helper.available.Add(device.Id);
+					helper._available.Add(device.Id);
 				}
 
 				List<DevicePlatformTelemetryDocument> platformTelemtry = new List<DevicePlatformTelemetryDocument>();
 
-
 				foreach (KeyValuePair<DevicePlatformId, DevicePoolTelemetryHelper> platform in helpers)
 				{
 					DevicePoolTelemetryHelper helper = platform.Value;
 
-					platformTelemtry.Add(new DevicePlatformTelemetryDocument(platform.Key, helper.available, helper.reserved, helper.maintenance, helper.problem, helper.disabled));
+					platformTelemtry.Add(new DevicePlatformTelemetryDocument(platform.Key, helper._available, helper._reserved, helper._maintenance, helper._problem, helper._disabled));
 				}
 
 				poolTelemetry[pool.Id] = platformTelemtry;
@@ -1234,7 +1226,5 @@ namespace Horde.Build.Devices
 			List<DevicePoolTelemetryDocument> results = await _poolTelemetry.Find(filter).Range(index, count).ToListAsync();
 			return results.ConvertAll<IDevicePoolTelemetry>(x => x);
 		}
-
-
 	}
 }
diff --git a/Engine/Source/Programs/Horde/Horde.Build/Devices/DevicesController.cs b/Engine/Source/Programs/Horde/Horde.Build/Devices/DevicesController.cs
--- a/Engine/Source/Programs/Horde/Horde.Build/Devices/DevicesController.cs
+++ b/Engine/Source/Programs/Horde/Horde.Build/Devices/DevicesController.cs
@@ -539,8 +539,7 @@ namespace Horde.Build.Devices
 			return response;
 		}
 
-
-		// RESERVATIONS
+		#region Reservations
 
 		/// <summary>
 		/// Create a new device reservation
@@ -718,6 +717,8 @@ namespace Horde.Build.Devices
 			return Ok();
 		}
 
+		#endregion
+
 		/// <summary>
 		/// Get device telemetry
 		/// </summary>		
diff --git a/Engine/Source/Programs/Horde/Horde.Build/GlobalSuppressions.cs b/Engine/Source/Programs/Horde/Horde.Build/GlobalSuppressions.cs
--- a/Engine/Source/Programs/Horde/Horde.Build/GlobalSuppressions.cs
+++ b/Engine/Source/Programs/Horde/Horde.Build/GlobalSuppressions.cs
@@ -16,3 +16,5 @@ using System.Diagnostics.CodeAnalysis;
 [assembly: SuppressMessage("Design", "CA1030:Use events where appropriate", Justification = "<Pending>", Scope = "module")]
 [assembly: SuppressMessage("Design", "CA1031:Do not catch general exception types", Justification = "<Pending>", Scope = "module")]
 [assembly: SuppressMessage("Design", "CA1032:Implement standard exception constructors", Justification = "<Pending>", Scope = "module")]
+[assembly: SuppressMessage("Design", "CA1056:URI-like properties should not be strings", Justification = "<Pending>", Scope = "module")]
+[assembly: SuppressMessage("Design", "CA1055:URI-like return values should not be strings", Justification = "<Pending>", Scope = "module")]
diff --git a/Engine/Source/Programs/Horde/Horde.Build/Horde.Build.csproj b/Engine/Source/Programs/Horde/Horde.Build/Horde.Build.csproj
--- a/Engine/Source/Programs/Horde/Horde.Build/Horde.Build.csproj
+++ b/Engine/Source/Programs/Horde/Horde.Build/Horde.Build.csproj
@@ -8,6 +8,14 @@
     <DocumentationFile>Horde.Build.xml</DocumentationFile>
     <ApplicationIcon>HordeServer.ico</ApplicationIcon>
     <RunAnalyzersDuringBuild>False</RunAnalyzersDuringBuild>
+    <Configurations>Debug;Release;Analyze</Configurations>
+  </PropertyGroup>
+
+  <PropertyGroup Condition=" '$(Configuration)|$(Platform)' == 'Analyze|AnyCPU' ">
+    <RunAnalyzersDuringBuild>True</RunAnalyzersDuringBuild>
+    <TreatWarningsAsErrors>false</TreatWarningsAsErrors>
+    <GenerateDocumentationFile>True</GenerateDocumentationFile>
+    <DocumentationFile></DocumentationFile>
   </PropertyGroup>
 
   <ItemGroup>
diff --git a/Engine/Source/Programs/Horde/Horde.Build/Issues/External/IExternalIssue.cs b/Engine/Source/Programs/Horde/Horde.Build/Issues/External/IExternalIssue.cs
--- a/Engine/Source/Programs/Horde/Horde.Build/Issues/External/IExternalIssue.cs
+++ b/Engine/Source/Programs/Horde/Horde.Build/Issues/External/IExternalIssue.cs
@@ -1,18 +1,5 @@
 // Copyright Epic Games, Inc. All Rights Reserved.	
 
-using System;
-using System.Collections.Concurrent;
-using System.Collections.Generic;
-using System.Net.Http;
-using System.Net.Http.Headers;
-using System.Text;
-using System.Text.Json;
-using System.Threading.Tasks;
-using Microsoft.Extensions.Logging;
-using Microsoft.Extensions.Options;
-using Polly;
-using Polly.Extensions.Http;
-
 namespace Horde.Build.Issues.External
 {
 	/// <summary>
@@ -59,8 +46,6 @@ namespace Horde.Build.Issues.External
 		/// The current assignee's email address
 		/// </summary>
 		public string? AssigneeEmailAddress { get; set; }
-
 	}
-
 }
 
diff --git a/Engine/Source/Programs/Horde/Horde.Build/Issues/External/IExternalIssueService.cs b/Engine/Source/Programs/Horde/Horde.Build/Issues/External/IExternalIssueService.cs
--- a/Engine/Source/Programs/Horde/Horde.Build/Issues/External/IExternalIssueService.cs
+++ b/Engine/Source/Programs/Horde/Horde.Build/Issues/External/IExternalIssueService.cs
@@ -6,13 +6,10 @@ using System.Threading;
 using System.Threading.Tasks;
 using Horde.Build.Streams;
 using Horde.Build.Users;
-using Horde.Build.Utilities;
 using Microsoft.Extensions.Hosting;
 
 namespace Horde.Build.Issues.External
 {
-	using StreamId = StringId<IStream>;
-
 	/// <summary>
 	/// External issue tracking project
 	/// </summary>
@@ -77,14 +74,12 @@ namespace Horde.Build.Issues.External
 		/// <param name="hordeIssueLink"></param>
 		Task<(string? key, string? url)> CreateIssueAsync(IUser user, string? externalIssueUser, int issueId, string summary, string projectId, string componentId, string issueType, string? description, string? hordeIssueLink);
 
-
 		/// <summary>
 		///  Get projects for provided keys
 		/// </summary>
 		/// <param name="stream"></param>
 		/// <returns></returns>;
 		Task<List<IExternalIssueProject>> GetProjects(IStream stream);
-
 	}
 
 	/// <summary>
@@ -125,5 +120,4 @@ namespace Horde.Build.Issues.External
 		/// <inheritdoc/>
 		public Task StopAsync(CancellationToken token) => Task.CompletedTask;
 	}
-
 }
\ No newline at end of file
diff --git a/Engine/Source/Programs/Horde/Horde.Build/Issues/External/JiraService.cs b/Engine/Source/Programs/Horde/Horde.Build/Issues/External/JiraService.cs
--- a/Engine/Source/Programs/Horde/Horde.Build/Issues/External/JiraService.cs
+++ b/Engine/Source/Programs/Horde/Horde.Build/Issues/External/JiraService.cs
@@ -18,7 +18,6 @@ using Microsoft.Extensions.Hosting;
 using HordeCommon;
 using System.Threading;
 using System.Linq;
-using Horde.Build.Server;
 using Horde.Build.Streams;
 using Horde.Build.Users;
 
@@ -125,13 +124,11 @@ namespace Horde.Build.Issues.External
 		/// </summary>
 		readonly ConfigCollection _configCollection;
 
-		readonly IClock _clock;
 		readonly ITicker _ticker;
 
 		readonly Uri _jiraUrl;
 
-		Dictionary<string, JiraProject> _cachedJiraProjects = new Dictionary<string, JiraProject>();
-
+		readonly Dictionary<string, JiraProject> _cachedJiraProjects = new Dictionary<string, JiraProject>();
 
 		/// <summary>
 		/// Jira service constructor
@@ -144,13 +141,11 @@ namespace Horde.Build.Issues.External
 		/// <param name="logger"></param>
 		public JiraService(IOptions<ServerSettings> settings, StreamService streamService, IssueService issueService, ConfigCollection configCollection, IClock clock, ILogger<JiraService> logger)
 		{
-
 			_settings = settings.Value;
 			_logger = logger;
 			_streamService = streamService;
 			_issueService = issueService;
 			_configCollection = configCollection;
-			_clock = clock;
 			_ticker = clock.AddTicker<JiraService>(TimeSpan.FromMinutes(2.0), TickAsync, logger);
 			_jiraUrl = _settings.JiraUrl!;
 
@@ -206,7 +201,6 @@ namespace Horde.Build.Issues.External
 				{
 					projectKeys.Add(workflow.ExternalIssues.ProjectKey);
 				}				
-
 			});
 
 			foreach (string key in projectKeys)
@@ -274,9 +268,7 @@ namespace Horde.Build.Issues.External
 						}
 
 						jiraProject.IssueTypes[issueType.GetProperty("id").GetString()!] = name;
-
 					}
-
 				}
 
 				// Components
@@ -328,10 +320,9 @@ namespace Horde.Build.Issues.External
 			HashSet<string> jiraProjectKeys = new HashSet<string>();
 
 			List<IStream> cachedStreams = await _streamService.StreamCollection.FindAllAsync();
-			StreamConfig? config = null;
 			foreach (IStream cachedStream in cachedStreams)
 			{
-				config = null;
+				StreamConfig? config = null;
 				try
 				{
 					config = await _configCollection.GetConfigAsync<StreamConfig>(cachedStream.ConfigRevision);
@@ -393,7 +384,6 @@ namespace Horde.Build.Issues.External
 				issueDesc = $"\n{issueDesc}\n\nJira Issue created by {user.Name} from Horde\n\n";
 			}
 
-
 			Uri uri = new Uri(_jiraUrl, $"/rest/api/2/issue");
 
 			string bodyJson = JsonSerializer.Serialize(new
@@ -408,7 +398,7 @@ namespace Horde.Build.Issues.External
 					{
 						id = issueType
 					},
-					summary = summary,
+					summary,
 					description = issueDesc,
 					components = new object[]
 					{
@@ -449,7 +439,6 @@ namespace Horde.Build.Issues.External
 				{
 					_logger.LogError(ex, "Unable to add {ExternalUser} as watcher to issue {IssueKey}", externalIssueUser, jiraResponse.Key);
 				}
-				
 			}
 
 			return (jiraResponse.Key, $"{_jiraUrl}browse/{jiraResponse.Key}");
diff --git a/Engine/Source/Programs/Horde/Horde.Build/Issues/Handlers/CompileIssueHandler.cs b/Engine/Source/Programs/Horde/Horde.Build/Issues/Handlers/CompileIssueHandler.cs
--- a/Engine/Source/Programs/Horde/Horde.Build/Issues/Handlers/CompileIssueHandler.cs
+++ b/Engine/Source/Programs/Horde/Horde.Build/Issues/Handlers/CompileIssueHandler.cs
@@ -2,12 +2,10 @@
 
 using System;
 using System.Collections.Generic;
-using System.Diagnostics.CodeAnalysis;
 using System.Linq;
 using EpicGames.Core;
 using Horde.Build.Jobs;
 using Horde.Build.Jobs.Graphs;
-using Horde.Build.Logs;
 using Microsoft.Extensions.Logging;
 using MongoDB.Driver;
 
diff --git a/Engine/Source/Programs/Horde/Horde.Build/Issues/Handlers/ContentIssueHandler.cs b/Engine/Source/Programs/Horde/Horde.Build/Issues/Handlers/ContentIssueHandler.cs
--- a/Engine/Source/Programs/Horde/Horde.Build/Issues/Handlers/ContentIssueHandler.cs
+++ b/Engine/Source/Programs/Horde/Horde.Build/Issues/Handlers/ContentIssueHandler.cs
@@ -2,7 +2,6 @@
 
 using System;
 using System.Collections.Generic;
-using System.Diagnostics.CodeAnalysis;
 using System.Linq;
 using EpicGames.Core;
 using Horde.Build.Jobs;
diff --git a/Engine/Source/Programs/Horde/Horde.Build/Issues/Handlers/CopyrightIssueHandler.cs b/Engine/Source/Programs/Horde/Horde.Build/Issues/Handlers/CopyrightIssueHandler.cs
--- a/Engine/Source/Programs/Horde/Horde.Build/Issues/Handlers/CopyrightIssueHandler.cs
+++ b/Engine/Source/Programs/Horde/Horde.Build/Issues/Handlers/CopyrightIssueHandler.cs
@@ -1,16 +1,11 @@
 // Copyright Epic Games, Inc. All Rights Reserved.
 
-using System;
 using System.Collections.Generic;
-using System.Diagnostics.CodeAnalysis;
 using System.Linq;
 using EpicGames.Core;
 using Horde.Build.Jobs;
 using Horde.Build.Jobs.Graphs;
-using Horde.Build.Logs;
-using Horde.Build.Utilities;
 using Microsoft.Extensions.Logging;
-using MongoDB.Driver;
 
 namespace Horde.Build.Issues.Handlers
 {
diff --git a/Engine/Source/Programs/Horde/Horde.Build/Issues/Handlers/GauntletIssueHandler.cs b/Engine/Source/Programs/Horde/Horde.Build/Issues/Handlers/GauntletIssueHandler.cs
--- a/Engine/Source/Programs/Horde/Horde.Build/Issues/Handlers/GauntletIssueHandler.cs
+++ b/Engine/Source/Programs/Horde/Horde.Build/Issues/Handlers/GauntletIssueHandler.cs
@@ -2,7 +2,6 @@
 
 using System;
 using System.Collections.Generic;
-using System.Diagnostics.CodeAnalysis;
 using System.Linq;
 using System.Text.Json;
 using EpicGames.Core;
diff --git a/Engine/Source/Programs/Horde/Horde.Build/Issues/Handlers/LocalizationIssueHandler.cs b/Engine/Source/Programs/Horde/Horde.Build/Issues/Handlers/LocalizationIssueHandler.cs
--- a/Engine/Source/Programs/Horde/Horde.Build/Issues/Handlers/LocalizationIssueHandler.cs
+++ b/Engine/Source/Programs/Horde/Horde.Build/Issues/Handlers/LocalizationIssueHandler.cs
@@ -2,7 +2,6 @@
 
 using System;
 using System.Collections.Generic;
-using System.Diagnostics.CodeAnalysis;
 using System.Linq;
 using EpicGames.Core;
 using Horde.Build.Jobs;
diff --git a/Engine/Source/Programs/Horde/Horde.Build/Issues/Handlers/PerforceCaseIssueHandler.cs b/Engine/Source/Programs/Horde/Horde.Build/Issues/Handlers/PerforceCaseIssueHandler.cs
--- a/Engine/Source/Programs/Horde/Horde.Build/Issues/Handlers/PerforceCaseIssueHandler.cs
+++ b/Engine/Source/Programs/Horde/Horde.Build/Issues/Handlers/PerforceCaseIssueHandler.cs
@@ -1,7 +1,6 @@
 // Copyright Epic Games, Inc. All Rights Reserved.
 
 using System.Collections.Generic;
-using System.Diagnostics.CodeAnalysis;
 using System.Linq;
 using System.Text.Json;
 using EpicGames.Core;
diff --git a/Engine/Source/Programs/Horde/Horde.Build/Issues/Handlers/ShaderIssueHandler.cs b/Engine/Source/Programs/Horde/Horde.Build/Issues/Handlers/ShaderIssueHandler.cs
--- a/Engine/Source/Programs/Horde/Horde.Build/Issues/Handlers/ShaderIssueHandler.cs
+++ b/Engine/Source/Programs/Horde/Horde.Build/Issues/Handlers/ShaderIssueHandler.cs
@@ -2,14 +2,11 @@
 
 using System;
 using System.Collections.Generic;
-using System.Diagnostics.CodeAnalysis;
 using System.Linq;
 using EpicGames.Core;
 using Horde.Build.Jobs;
 using Horde.Build.Jobs.Graphs;
-using Horde.Build.Logs;
 using Microsoft.Extensions.Logging;
-using MongoDB.Driver;
 
 namespace Horde.Build.Issues.Handlers
 {
diff --git a/Engine/Source/Programs/Horde/Horde.Build/Issues/Handlers/SourceFileIssueHandler.cs b/Engine/Source/Programs/Horde/Horde.Build/Issues/Handlers/SourceFileIssueHandler.cs
--- a/Engine/Source/Programs/Horde/Horde.Build/Issues/Handlers/SourceFileIssueHandler.cs
+++ b/Engine/Source/Programs/Horde/Horde.Build/Issues/Handlers/SourceFileIssueHandler.cs
@@ -2,11 +2,8 @@
 
 using System;
 using System.Collections.Generic;
-using System.Diagnostics.CodeAnalysis;
 using System.Linq;
 using System.Text.Json;
-using Horde.Build.Jobs;
-using Horde.Build.Jobs.Graphs;
 using Horde.Build.Logs;
 using Horde.Build.Utilities;
 using MongoDB.Driver;
diff --git a/Engine/Source/Programs/Horde/Horde.Build/Issues/Handlers/SymbolIssueHandler.cs b/Engine/Source/Programs/Horde/Horde.Build/Issues/Handlers/SymbolIssueHandler.cs
--- a/Engine/Source/Programs/Horde/Horde.Build/Issues/Handlers/SymbolIssueHandler.cs
+++ b/Engine/Source/Programs/Horde/Horde.Build/Issues/Handlers/SymbolIssueHandler.cs
@@ -2,7 +2,6 @@
 
 using System;
 using System.Collections.Generic;
-using System.Diagnostics.CodeAnalysis;
 using System.Linq;
 using System.Text;
 using EpicGames.Core;
@@ -22,7 +21,7 @@ namespace Horde.Build.Issues.Handlers
 		const string NodeName = "Node";
 		const string EventIdName = "EventId";
 
-		static readonly string DuplicateEventIdMetadata = $"{EventIdName}={KnownLogEvents.Linker_DuplicateSymbol.Id}";
+		static readonly string s_duplicateEventIdMetadata = $"{EventIdName}={KnownLogEvents.Linker_DuplicateSymbol.Id}";
 
 		/// <inheritdoc/>
 		public override string Type => "Symbol";
@@ -105,7 +104,7 @@ namespace Horde.Build.Issues.Handlers
 			}
 			else
 			{
-				string problemType = fingerprint.HasMetadataValue(DuplicateEventIdMetadata) ? "Duplicate" : "Undefined";
+				string problemType = fingerprint.HasMetadataValue(s_duplicateEventIdMetadata) ? "Duplicate" : "Undefined";
 				if (symbols.Count == 1)
 				{
 					return $"{problemType} symbol '{symbols.First()}'";
diff --git a/Engine/Source/Programs/Horde/Horde.Build/Issues/Handlers/SystemicIssueHandler.cs b/Engine/Source/Programs/Horde/Horde.Build/Issues/Handlers/SystemicIssueHandler.cs
--- a/Engine/Source/Programs/Horde/Horde.Build/Issues/Handlers/SystemicIssueHandler.cs
+++ b/Engine/Source/Programs/Horde/Horde.Build/Issues/Handlers/SystemicIssueHandler.cs
@@ -1,8 +1,6 @@
 // Copyright Epic Games, Inc. All Rights Reserved.
 
-using System;
 using System.Collections.Generic;
-using System.Diagnostics.CodeAnalysis;
 using System.Linq;
 using EpicGames.Core;
 using Horde.Build.Jobs;
diff --git a/Engine/Source/Programs/Horde/Horde.Build/Issues/Handlers/UnacceptableWordsIssueHandler.cs b/Engine/Source/Programs/Horde/Horde.Build/Issues/Handlers/UnacceptableWordsIssueHandler.cs
--- a/Engine/Source/Programs/Horde/Horde.Build/Issues/Handlers/UnacceptableWordsIssueHandler.cs
+++ b/Engine/Source/Programs/Horde/Horde.Build/Issues/Handlers/UnacceptableWordsIssueHandler.cs
@@ -1,12 +1,10 @@
 // Copyright Epic Games, Inc. All Rights Reserved.
 
 using System.Collections.Generic;
-using System.Diagnostics.CodeAnalysis;
 using System.Linq;
 using EpicGames.Core;
 using Horde.Build.Jobs;
 using Horde.Build.Jobs.Graphs;
-using Horde.Build.Logs;
 using Microsoft.Extensions.Logging;
 
 namespace Horde.Build.Issues.Handlers
diff --git a/Engine/Source/Programs/Horde/Horde.Build/Issues/IIssueStep.cs b/Engine/Source/Programs/Horde/Horde.Build/Issues/IIssueStep.cs
--- a/Engine/Source/Programs/Horde/Horde.Build/Issues/IIssueStep.cs
+++ b/Engine/Source/Programs/Horde/Horde.Build/Issues/IIssueStep.cs
@@ -4,7 +4,6 @@ using System;
 using Horde.Build.Jobs;
 using Horde.Build.Jobs.Graphs;
 using Horde.Build.Logs;
-using Horde.Build.Server;
 using Horde.Build.Utilities;
 using MongoDB.Bson;
 
@@ -12,7 +11,6 @@ namespace Horde.Build.Issues
 {
 	using JobId = ObjectId<IJob>;
 	using LogId = ObjectId<ILogFile>;
-	using WorkflowId = StringId<WorkflowConfig>;
 
 	/// <summary>
 	/// Identifies a particular changelist and job that contributes to a span
diff --git a/Engine/Source/Programs/Horde/Horde.Build/Issues/IssueCollection.cs b/Engine/Source/Programs/Horde/Horde.Build/Issues/IssueCollection.cs
--- a/Engine/Source/Programs/Horde/Horde.Build/Issues/IssueCollection.cs
+++ b/Engine/Source/Programs/Horde/Horde.Build/Issues/IssueCollection.cs
@@ -9,7 +9,6 @@ using EpicGames.Redis.Utility;
 using Horde.Build.Auditing;
 using Horde.Build.Jobs;
 using Horde.Build.Jobs.Graphs;
-using Horde.Build.Issues;
 using Horde.Build.Logs;
 using Horde.Build.Server;
 using Horde.Build.Streams;
@@ -123,7 +122,6 @@ namespace Horde.Build.Issues
 			[BsonIgnoreIfNull]
 			public DateTime? QuarantineTimeUtc { get; set; }
 
-
 			[BsonConstructor]
 			private Issue()
 			{
@@ -606,7 +604,6 @@ namespace Horde.Build.Issues
 				}
 			}
 
-
 			if (newIssue.ExternalIssueKey != oldIssue.ExternalIssueKey)
 			{
 				if (newIssue.ExternalIssueKey != null)
@@ -631,7 +628,6 @@ namespace Horde.Build.Issues
 				}
 			}
 
-
 			string oldFingerprints = oldIssue.FingerprintsDesc;
 			string newFingerprints = newIssue.FingerprintsDesc;
 			if (oldFingerprints != newFingerprints)
diff --git a/Engine/Source/Programs/Horde/Horde.Build/Issues/IssueMessages.cs b/Engine/Source/Programs/Horde/Horde.Build/Issues/IssueMessages.cs
--- a/Engine/Source/Programs/Horde/Horde.Build/Issues/IssueMessages.cs
+++ b/Engine/Source/Programs/Horde/Horde.Build/Issues/IssueMessages.cs
@@ -4,9 +4,7 @@ using System;
 using System.Collections.Generic;
 using System.Linq;
 using EpicGames.Core;
-using Horde.Build.Issues;
 using Horde.Build.Issues.External;
-using Horde.Build.Server;
 using Horde.Build.Streams;
 using Horde.Build.Users;
 using Horde.Build.Utilities;
@@ -531,7 +529,6 @@ namespace Horde.Build.Issues
 		/// </summary>
 		public DateTime? QuarantineTimeUtc { get; set; }
 
-
 		/// <summary>
 		/// Constructs a new issue
 		/// </summary>
@@ -885,10 +882,8 @@ namespace Horde.Build.Issues
 		/// Id of user quarantining issue
 		/// </summary>
 		public string? QuarantinedById { get; set; }
-
 	}
 
-
 	/// <summary>
 	/// External issue project information
 	/// </summary>
@@ -997,7 +992,6 @@ namespace Horde.Build.Issues
 		/// Optional link to issue on Horde
 		/// </summary>
 		public string? HordeIssueLink { get; set; }
-
 	}
 
 	/// <summary>
@@ -1025,7 +1019,6 @@ namespace Horde.Build.Issues
 			Key = key;
 			Link = link;
 		}
-
 	}
 
 	/// <summary>
@@ -1087,8 +1080,5 @@ namespace Horde.Build.Issues
 			AssigneeDisplayName = issue.AssigneeDisplayName;
 			AssigneeEmailAddress = issue.AssigneeEmailAddress;
 		}
-
-
 	}
-
 }
diff --git a/Engine/Source/Programs/Horde/Horde.Build/Issues/IssueService.cs b/Engine/Source/Programs/Horde/Horde.Build/Issues/IssueService.cs
--- a/Engine/Source/Programs/Horde/Horde.Build/Issues/IssueService.cs
+++ b/Engine/Source/Programs/Horde/Horde.Build/Issues/IssueService.cs
@@ -10,13 +10,11 @@ using System.Text.RegularExpressions;
 using System.Threading;
 using System.Threading.Tasks;
 using EpicGames.Core;
-using Horde.Build.Configuration;
 using Horde.Build.Issues.Handlers;
 using Horde.Build.Jobs;
 using Horde.Build.Jobs.Graphs;
 using Horde.Build.Logs;
 using Horde.Build.Perforce;
-using Horde.Build.Server;
 using Horde.Build.Streams;
 using Horde.Build.Users;
 using Horde.Build.Utilities;
@@ -31,7 +29,6 @@ using OpenTracing.Util;
 
 namespace Horde.Build.Issues
 {
-	using LogId = ObjectId<ILogFile>;
 	using StreamId = StringId<IStream>;
 	using TemplateRefId = StringId<TemplateRef>;
 	using UserId = ObjectId<IUser>;
@@ -110,7 +107,6 @@ namespace Horde.Build.Issues
 			ExternalIssueKey = externalIssueKey; 
 			_showDesktopAlerts = showDesktopAlerts;
 
-
 			if (issue.OwnerId == null)
 			{
 				_notifyUsers = new HashSet<UserId>(suspects.Select(x => x.AuthorId));
@@ -151,7 +147,6 @@ namespace Horde.Build.Issues
 		/// </summary>
 		const int MaxChanges = 250;
 
-		readonly ConfigCollection _configCollection;
 		readonly IJobStepRefCollection _jobStepRefs;
 		readonly IIssueCollection _issueCollection;
 		readonly StreamService _streams;
@@ -209,7 +204,7 @@ namespace Horde.Build.Issues
 		/// <summary>
 		/// Constructor
 		/// </summary>
-		public IssueService(ConfigCollection configCollection, IIssueCollection issueCollection, IJobStepRefCollection jobStepRefs, StreamService streams, IUserCollection userCollection, IPerforceService perforce, ILogFileService logFileService, IClock clock, ILogger<IssueService> logger)
+		public IssueService(IIssueCollection issueCollection, IJobStepRefCollection jobStepRefs, StreamService streams, IUserCollection userCollection, IPerforceService perforce, ILogFileService logFileService, IClock clock, ILogger<IssueService> logger)
 		{
 			Type[] issueTypes = Assembly.GetExecutingAssembly().GetTypes().Where(x => !x.IsAbstract && typeof(IIssue).IsAssignableFrom(x)).ToArray();
 			foreach (Type issueType in issueTypes)
@@ -218,7 +213,6 @@ namespace Horde.Build.Issues
 			}
 
 			// Get all the collections
-			_configCollection = configCollection;
 			_issueCollection = issueCollection;
 			_jobStepRefs = jobStepRefs;
 			_streams = streams;
@@ -365,6 +359,8 @@ namespace Horde.Build.Issues
 		/// <inheritdoc/>
 		public bool ShowDesktopAlertsForIssue(IIssue issue, IReadOnlyList<IIssueSpan> spans)
 		{
+			_ = issue;
+
 			bool bShowDesktopAlerts = false;
 
 			HashSet<(StreamId, TemplateRefId)> checkedTemplates = new HashSet<(StreamId, TemplateRefId)>();
@@ -1002,10 +998,9 @@ namespace Horde.Build.Issues
 						{
 							newStream.ContainsFix = stream.ContainsFix;
 						}
-						if (newStream.ContainsFix == null)
-						{
-							newStream.ContainsFix = await ContainsFixChange(newStream.StreamId, issue.FixChange.Value, fixChangeCache);
-						}
+
+						newStream.ContainsFix ??= await ContainsFixChange(newStream.StreamId, issue.FixChange.Value, fixChangeCache);
+
 						if (spans.Any(x => x.StreamId == newStream.StreamId && x.LastFailure.Change > issue.FixChange.Value))
 						{
 							newStream.FixFailed = true;
diff --git a/Engine/Source/Programs/Horde/Horde.Build/Issues/SuspectChange.cs b/Engine/Source/Programs/Horde/Horde.Build/Issues/SuspectChange.cs
--- a/Engine/Source/Programs/Horde/Horde.Build/Issues/SuspectChange.cs
+++ b/Engine/Source/Programs/Horde/Horde.Build/Issues/SuspectChange.cs
@@ -2,11 +2,7 @@
 
 using System;
 using System.Collections.Generic;
-using System.Diagnostics.CodeAnalysis;
 using EpicGames.Core;
-using Horde.Build.Jobs;
-using Horde.Build.Jobs.Graphs;
-using Horde.Build.Logs;
 using Horde.Build.Perforce;
 
 namespace Horde.Build.Issues
diff --git a/Engine/Source/Programs/Horde/Horde.Build/Jobs/Artifacts/ArtifactCollection.cs b/Engine/Source/Programs/Horde/Horde.Build/Jobs/Artifacts/ArtifactCollection.cs
--- a/Engine/Source/Programs/Horde/Horde.Build/Jobs/Artifacts/ArtifactCollection.cs
+++ b/Engine/Source/Programs/Horde/Horde.Build/Jobs/Artifacts/ArtifactCollection.cs
@@ -4,7 +4,6 @@ using System;
 using System.Collections.Generic;
 using System.Linq;
 using System.Threading.Tasks;
-using Horde.Build.Jobs;
 using Horde.Build.Server;
 using Horde.Build.Storage;
 using Horde.Build.Utilities;
diff --git a/Engine/Source/Programs/Horde/Horde.Build/Jobs/Artifacts/ArtifactMessages.cs b/Engine/Source/Programs/Horde/Horde.Build/Jobs/Artifacts/ArtifactMessages.cs
--- a/Engine/Source/Programs/Horde/Horde.Build/Jobs/Artifacts/ArtifactMessages.cs
+++ b/Engine/Source/Programs/Horde/Horde.Build/Jobs/Artifacts/ArtifactMessages.cs
@@ -1,7 +1,6 @@
 // Copyright Epic Games, Inc. All Rights Reserved.
 
 using System.Collections.Generic;
-using Horde.Build.Jobs.Artifacts;
 
 namespace Horde.Build.Jobs.Artifacts
 {
diff --git a/Engine/Source/Programs/Horde/Horde.Build/Jobs/Artifacts/ArtifactsController.cs b/Engine/Source/Programs/Horde/Horde.Build/Jobs/Artifacts/ArtifactsController.cs
--- a/Engine/Source/Programs/Horde/Horde.Build/Jobs/Artifacts/ArtifactsController.cs
+++ b/Engine/Source/Programs/Horde/Horde.Build/Jobs/Artifacts/ArtifactsController.cs
@@ -10,7 +10,6 @@ using System.Net.Mime;
 using System.Security.Claims;
 using System.Threading.Tasks;
 using Horde.Build.Acls;
-using Horde.Build.Jobs;
 using Horde.Build.Jobs.Graphs;
 using Horde.Build.Server;
 using Horde.Build.Utilities;
diff --git a/Engine/Source/Programs/Horde/Horde.Build/Jobs/Graphs/NodeAnnotations.cs b/Engine/Source/Programs/Horde/Horde.Build/Jobs/Graphs/NodeAnnotations.cs
--- a/Engine/Source/Programs/Horde/Horde.Build/Jobs/Graphs/NodeAnnotations.cs
+++ b/Engine/Source/Programs/Horde/Horde.Build/Jobs/Graphs/NodeAnnotations.cs
@@ -3,7 +3,6 @@
 using System;
 using System.Collections.Generic;
 using Horde.Build.Issues;
-using Horde.Build.Server;
 using Horde.Build.Utilities;
 
 namespace Horde.Build.Jobs.Graphs
diff --git a/Engine/Source/Programs/Horde/Horde.Build/Jobs/JobCollection.cs b/Engine/Source/Programs/Horde/Horde.Build/Jobs/JobCollection.cs
--- a/Engine/Source/Programs/Horde/Horde.Build/Jobs/JobCollection.cs
+++ b/Engine/Source/Programs/Horde/Horde.Build/Jobs/JobCollection.cs
@@ -892,10 +892,7 @@ namespace Horde.Build.Jobs
 							// Apply any property updates
 							if (newProperties != null)
 							{
-								if (step.Properties == null)
-								{
-									step.Properties = new Dictionary<string, string>(StringComparer.Ordinal);
-								}
+								step.Properties ??= new Dictionary<string, string>(StringComparer.Ordinal);
 
 								foreach (KeyValuePair<string, string?> pair in newProperties)
 								{
diff --git a/Engine/Source/Programs/Horde/Horde.Build/Jobs/JobMessages.cs b/Engine/Source/Programs/Horde/Horde.Build/Jobs/JobMessages.cs
--- a/Engine/Source/Programs/Horde/Horde.Build/Jobs/JobMessages.cs
+++ b/Engine/Source/Programs/Horde/Horde.Build/Jobs/JobMessages.cs
@@ -7,7 +7,6 @@ using System.Linq;
 using System.Text.Json.Serialization;
 using EpicGames.Core;
 using Horde.Build.Acls;
-using Horde.Build.Jobs;
 using Horde.Build.Jobs.Graphs;
 using Horde.Build.Streams;
 using Horde.Build.Users;
diff --git a/Engine/Source/Programs/Horde/Horde.Build/Jobs/JobTaskSource.cs b/Engine/Source/Programs/Horde/Horde.Build/Jobs/JobTaskSource.cs
--- a/Engine/Source/Programs/Horde/Horde.Build/Jobs/JobTaskSource.cs
+++ b/Engine/Source/Programs/Horde/Horde.Build/Jobs/JobTaskSource.cs
@@ -196,7 +196,6 @@ namespace Horde.Build.Jobs
 		readonly IUgsMetadataCollection _ugsMetadataCollection;
 		readonly PerforceLoadBalancer _perforceLoadBalancer;
 		readonly IOptionsMonitor<ServerSettings> _settings;
-		readonly ICommitService _commitService;
 		readonly ILogger<JobTaskSource> _logger;
 		readonly ITicker _ticker;
 
@@ -237,7 +236,7 @@ namespace Horde.Build.Jobs
 		/// <summary>
 		/// Constructor
 		/// </summary>
-		public JobTaskSource(MongoService mongoService, IAgentCollection agents, IJobCollection jobs, IJobStepRefCollection jobStepRefs, IGraphCollection graphs, IPoolCollection pools, IUgsMetadataCollection ugsMetadataCollection, StreamService streamService, ILogFileService logFileService, PerforceLoadBalancer perforceLoadBalancer, ICommitService commitService, IClock clock, IOptionsMonitor<ServerSettings> settings, ILogger<JobTaskSource> logger)
+		public JobTaskSource(MongoService mongoService, IAgentCollection agents, IJobCollection jobs, IJobStepRefCollection jobStepRefs, IGraphCollection graphs, IPoolCollection pools, IUgsMetadataCollection ugsMetadataCollection, StreamService streamService, ILogFileService logFileService, PerforceLoadBalancer perforceLoadBalancer, IClock clock, IOptionsMonitor<ServerSettings> settings, ILogger<JobTaskSource> logger)
 		{
 			_mongoService = mongoService;
 			_agentsCollection = agents;
@@ -251,7 +250,6 @@ namespace Horde.Build.Jobs
 			_perforceLoadBalancer = perforceLoadBalancer;
 			_ticker = clock.AddTicker<JobTaskSource>(s_refreshInterval, TickAsync, logger);
 			_settings = settings;
-			_commitService = commitService;
 			_logger = logger;
 
 			OnLeaseStartedProperties.Add(nameof(ExecuteJobTask.JobId), x => new JobId(x.JobId)).Add(nameof(ExecuteJobTask.BatchId), x => SubResourceId.Parse(x.BatchId)).Add(nameof(ExecuteJobTask.LogId), x => new LogId(x.LogId));
diff --git a/Engine/Source/Programs/Horde/Horde.Build/Jobs/JobsController.cs b/Engine/Source/Programs/Horde/Horde.Build/Jobs/JobsController.cs
--- a/Engine/Source/Programs/Horde/Horde.Build/Jobs/JobsController.cs
+++ b/Engine/Source/Programs/Horde/Horde.Build/Jobs/JobsController.cs
@@ -5,7 +5,6 @@ using System.Collections.Generic;
 using System.Linq;
 using System.Threading.Tasks;
 using EpicGames.Core;
-using EpicGames.Perforce;
 using Horde.Build.Acls;
 using Horde.Build.Agents;
 using Horde.Build.Jobs.Artifacts;
diff --git a/Engine/Source/Programs/Horde/Horde.Build/Jobs/Schedules/ScheduleService.cs b/Engine/Source/Programs/Horde/Horde.Build/Jobs/Schedules/ScheduleService.cs
--- a/Engine/Source/Programs/Horde/Horde.Build/Jobs/Schedules/ScheduleService.cs
+++ b/Engine/Source/Programs/Horde/Horde.Build/Jobs/Schedules/ScheduleService.cs
@@ -512,7 +512,7 @@ namespace Horde.Build.Jobs.Schedules
 			if (template.SubmitNewChange != null)
 			{
 				int newChange = await _perforce.CreateNewChangeForTemplateAsync(stream, template);
-				int newCodeChange = await _perforce.GetCodeChangeAsync(stream.ClusterName, stream.Name, newChange);
+				int newCodeChange = await _perforce.GetCodeChangeAsync(stream.ClusterName, stream.Name, newChange, cancellationToken);
 				triggerChanges = new List<(int, int)> { (newChange, newCodeChange) };
 			}
 
@@ -587,7 +587,7 @@ namespace Horde.Build.Jobs.Schedules
 						JobStepOutcome outcome = state.Value.Item2;
 						if (outcome == JobStepOutcome.Success || outcome == JobStepOutcome.Warnings)
 						{
-							return await _perforce.GetChangeDetailsAsync(stream.ClusterName, stream.Name, job.Change);
+							return await _perforce.GetChangeDetailsAsync(stream.ClusterName, stream.Name, job.Change, cancellationToken);
 						}
 						_logger.LogInformation("Skipping trigger of {StreamName} template {TemplateId} - last {OtherTemplateRefId} job ({JobId}) ended with errors", stream.Id, templateRefId, gate.TemplateRefId, job.Id);
 					}
diff --git a/Engine/Source/Programs/Horde/Horde.Build/Jobs/Schedules/Schedules.cs b/Engine/Source/Programs/Horde/Horde.Build/Jobs/Schedules/Schedules.cs
--- a/Engine/Source/Programs/Horde/Horde.Build/Jobs/Schedules/Schedules.cs
+++ b/Engine/Source/Programs/Horde/Horde.Build/Jobs/Schedules/Schedules.cs
@@ -4,7 +4,6 @@ using System;
 using System.Collections.Generic;
 using System.ComponentModel.DataAnnotations;
 using Horde.Build.Acls;
-using Horde.Build.Jobs.Schedules;
 using Horde.Build.Perforce;
 using Horde.Build.Streams;
 using Horde.Build.Utilities;
diff --git a/Engine/Source/Programs/Horde/Horde.Build/Jobs/TestData/TestDataController.Messages.cs b/Engine/Source/Programs/Horde/Horde.Build/Jobs/TestData/TestDataController.Messages.cs
--- a/Engine/Source/Programs/Horde/Horde.Build/Jobs/TestData/TestDataController.Messages.cs
+++ b/Engine/Source/Programs/Horde/Horde.Build/Jobs/TestData/TestDataController.Messages.cs
@@ -3,7 +3,6 @@
 using System;
 using System.Collections.Generic;
 using System.ComponentModel.DataAnnotations;
-using Horde.Build.Jobs.TestData;
 using MongoDB.Bson.Serialization;
 
 namespace Horde.Build.Jobs.TestData
diff --git a/Engine/Source/Programs/Horde/Horde.Build/Logs/Data/LogIndexData.cs b/Engine/Source/Programs/Horde/Horde.Build/Logs/Data/LogIndexData.cs
--- a/Engine/Source/Programs/Horde/Horde.Build/Logs/Data/LogIndexData.cs
+++ b/Engine/Source/Programs/Horde/Horde.Build/Logs/Data/LogIndexData.cs
@@ -56,10 +56,7 @@ namespace Horde.Build.Logs.Data
 		/// </summary>
 		public ILogText InflatePlainText()
 		{
-			if (CachedPlainText == null)
-			{
-				CachedPlainText = new ReadOnlyLogText(CompressedPlainText.DecompressBzip2());
-			}
+			CachedPlainText ??= new ReadOnlyLogText(CompressedPlainText.DecompressBzip2());
 			return CachedPlainText;
 		}
 	}
diff --git a/Engine/Source/Programs/Horde/Horde.Build/Logs/Data/LogSubChunkData.cs b/Engine/Source/Programs/Horde/Horde.Build/Logs/Data/LogSubChunkData.cs
--- a/Engine/Source/Programs/Horde/Horde.Build/Logs/Data/LogSubChunkData.cs
+++ b/Engine/Source/Programs/Horde/Horde.Build/Logs/Data/LogSubChunkData.cs
@@ -57,10 +57,7 @@ namespace Horde.Build.Logs.Data
 		/// </summary>
 		public ILogText InflateText()
 		{
-			if (_textInternal == null)
-			{
-				_textInternal = new ReadOnlyLogText(_compressedTextInternal.DecompressBzip2());
-			}
+			_textInternal ??= new ReadOnlyLogText(_compressedTextInternal.DecompressBzip2());
 			return _textInternal;
 		}
 
diff --git a/Engine/Source/Programs/Horde/Horde.Build/Logs/LogFileMessages.cs b/Engine/Source/Programs/Horde/Horde.Build/Logs/LogFileMessages.cs
--- a/Engine/Source/Programs/Horde/Horde.Build/Logs/LogFileMessages.cs
+++ b/Engine/Source/Programs/Horde/Horde.Build/Logs/LogFileMessages.cs
@@ -1,7 +1,6 @@
 // Copyright Epic Games, Inc. All Rights Reserved.
 
 using System.Collections.Generic;
-using Horde.Build.Logs;
 using Horde.Build.Logs.Data;
 
 namespace Horde.Build.Logs
diff --git a/Engine/Source/Programs/Horde/Horde.Build/Logs/LogFileService.cs b/Engine/Source/Programs/Horde/Horde.Build/Logs/LogFileService.cs
--- a/Engine/Source/Programs/Horde/Horde.Build/Logs/LogFileService.cs
+++ b/Engine/Source/Programs/Horde/Horde.Build/Logs/LogFileService.cs
@@ -883,15 +883,15 @@ namespace Horde.Build.Logs
 		/// <summary>
 		/// Helper method for catching exceptions in <see cref="LogText.ConvertToPlainText(ReadOnlySpan{Byte}, Byte[], Int32)"/>
 		/// </summary>
-		public static int GuardedConvertToPlainText(ReadOnlySpan<byte> input, byte[] output, int outputOffset, ILogger Logger)
+		public static int GuardedConvertToPlainText(ReadOnlySpan<byte> input, byte[] output, int outputOffset, ILogger logger)
 		{
 			try
 			{
 				return LogText.ConvertToPlainText(input, output, outputOffset);
 			}
-			catch (Exception Ex)
+			catch (Exception ex)
 			{
-				Logger.LogWarning(Ex, "Unable to convert log line to plain text: {Line}", Encoding.UTF8.GetString(input));
+				logger.LogWarning(ex, "Unable to convert log line to plain text: {Line}", Encoding.UTF8.GetString(input));
 				output[outputOffset] = (byte)'\n';
 				return outputOffset + 1;
 			}
@@ -1326,10 +1326,7 @@ namespace Horde.Build.Logs
 				}
 
 				// Otherwise go directly to the log storage
-				if (chunkData == null)
-				{
-					chunkData = await _storage.ReadChunkAsync(logFile.Id, chunk.Offset, chunk.LineIndex);
-				}
+				chunkData ??= await _storage.ReadChunkAsync(logFile.Id, chunk.Offset, chunk.LineIndex);
 			}
 			catch (Exception ex)
 			{
@@ -1348,9 +1345,9 @@ namespace Horde.Build.Logs
 				{
 					chunkData = await RepairChunkDataAsync(logFile, chunkIdx, chunkData, chunk.Length, logFile.MaxLineIndex.Value - chunk.LineIndex);
 				}
-				else if(chunkData == null)
+				else
 				{
-					chunkData = await RepairChunkDataAsync(logFile, chunkIdx, chunkData, 1024, 1);
+					chunkData ??= await RepairChunkDataAsync(logFile, chunkIdx, chunkData, 1024, 1);
 				}
 			}
 
diff --git a/Engine/Source/Programs/Horde/Horde.Build/Notifications/INotificationTriggerCollection.cs b/Engine/Source/Programs/Horde/Horde.Build/Notifications/INotificationTriggerCollection.cs
--- a/Engine/Source/Programs/Horde/Horde.Build/Notifications/INotificationTriggerCollection.cs
+++ b/Engine/Source/Programs/Horde/Horde.Build/Notifications/INotificationTriggerCollection.cs
@@ -2,7 +2,6 @@
 
 using System.Collections.Generic;
 using System.Threading.Tasks;
-using Horde.Build.Notifications;
 using Horde.Build.Users;
 using Horde.Build.Utilities;
 using MongoDB.Bson;
diff --git a/Engine/Source/Programs/Horde/Horde.Build/Notifications/Sinks/SlackNotificationSink.cs b/Engine/Source/Programs/Horde/Horde.Build/Notifications/Sinks/SlackNotificationSink.cs
--- a/Engine/Source/Programs/Horde/Horde.Build/Notifications/Sinks/SlackNotificationSink.cs
+++ b/Engine/Source/Programs/Horde/Horde.Build/Notifications/Sinks/SlackNotificationSink.cs
@@ -3,7 +3,6 @@
 using System;
 using System.Collections.Generic;
 using System.Diagnostics.CodeAnalysis;
-using System.Drawing;
 using System.Globalization;
 using System.Linq;
 using System.Net.Http;
@@ -19,14 +18,12 @@ using EpicGames.Redis;
 using EpicGames.Slack;
 using EpicGames.Slack.Blocks;
 using EpicGames.Slack.Elements;
-using Horde.Build.Configuration;
 using Horde.Build.Devices;
 using Horde.Build.Issues;
 using Horde.Build.Issues.External;
 using Horde.Build.Jobs;
 using Horde.Build.Jobs.Graphs;
 using Horde.Build.Logs;
-using Horde.Build.Perforce;
 using Horde.Build.Server;
 using Horde.Build.Streams;
 using Horde.Build.Users;
@@ -55,7 +52,7 @@ namespace Horde.Build.Notifications.Sinks
 	/// </summary>
 	public sealed class SlackNotificationSink : BackgroundService, INotificationSink, IAvatarService
 	{
-		const bool defaultAllowMentions = true;
+		const bool DefaultAllowMentions = true;
 
 		const int MaxLineLength = 2048;
 
@@ -739,7 +736,7 @@ namespace Horde.Build.Notifications.Sinks
 			{
 				foreach (string channel in channels)
 				{
-					await SendIssueMessageAsync(channel, issue, details, null, defaultAllowMentions);
+					await SendIssueMessageAsync(channel, issue, details, null, DefaultAllowMentions);
 				}
 			}
 
@@ -1086,7 +1083,7 @@ namespace Horde.Build.Notifications.Sinks
 				return;
 			}
 
-			await SendIssueMessageAsync(slackUserId, issue, details, user.Id, defaultAllowMentions);
+			await SendIssueMessageAsync(slackUserId, issue, details, user.Id, DefaultAllowMentions);
 		}
 
 		Uri GetJobUrl(JobId jobId)
@@ -2312,6 +2309,8 @@ namespace Horde.Build.Notifications.Sinks
 		/// </summary>
 		private async Task<object?> HandleInteractionMessage(EventPayload payload, CancellationToken cancellationToken)
 		{
+			_ = cancellationToken;
+
 			if (payload.User != null && payload.User.Id != null)
 			{
 				if (String.Equals(payload.Type, "block_actions", StringComparison.Ordinal))
@@ -2359,7 +2358,7 @@ namespace Horde.Build.Notifications.Sinks
 								{
 									Dictionary<string, string> errors = new Dictionary<string, string>();
 									errors.Add("fix_cl", $"'{fixChangeStr}' is not a valid fix changelist.");
-									return new { response_action = "errors", errors = errors };
+									return new { response_action = "errors", errors };
 								}
 
 								await _issueService.UpdateIssueAsync(issueId, fixChange: fixChange, resolvedById: userId);
@@ -2420,7 +2419,7 @@ namespace Horde.Build.Notifications.Sinks
 					if (recipient != null)
 					{
 						IIssueDetails details = await _issueService.GetIssueDetailsAsync(newIssue);
-						await SendIssueMessageAsync(recipient, newIssue, details, userId, defaultAllowMentions);
+						await SendIssueMessageAsync(recipient, newIssue, details, userId, DefaultAllowMentions);
 					}
 				}
 			}
diff --git a/Engine/Source/Programs/Horde/Horde.Build/Perforce/ChangeMessages.cs b/Engine/Source/Programs/Horde/Horde.Build/Perforce/ChangeMessages.cs
--- a/Engine/Source/Programs/Horde/Horde.Build/Perforce/ChangeMessages.cs
+++ b/Engine/Source/Programs/Horde/Horde.Build/Perforce/ChangeMessages.cs
@@ -1,7 +1,6 @@
 // Copyright Epic Games, Inc. All Rights Reserved.
 
 using System.Collections.Generic;
-using Horde.Build.Perforce;
 using Horde.Build.Users;
 
 namespace Horde.Build.Perforce
diff --git a/Engine/Source/Programs/Horde/Horde.Build/Perforce/CommitService.cs b/Engine/Source/Programs/Horde/Horde.Build/Perforce/CommitService.cs
--- a/Engine/Source/Programs/Horde/Horde.Build/Perforce/CommitService.cs
+++ b/Engine/Source/Programs/Horde/Horde.Build/Perforce/CommitService.cs
@@ -1,29 +1,20 @@
 // Copyright Epic Games, Inc. All Rights Reserved.
 
 using System;
-using System.Buffers;
 using System.Collections.Generic;
-using System.Diagnostics;
 using System.Globalization;
-using System.IO;
 using System.Linq;
-using System.Text;
 using System.Text.RegularExpressions;
 using System.Threading;
 using System.Threading.Tasks;
 using EpicGames.Core;
-using EpicGames.Horde.Storage;
-using EpicGames.Horde.Storage.Bundles;
-using EpicGames.Horde.Storage.Nodes;
 using EpicGames.Perforce;
 using EpicGames.Redis;
-using EpicGames.Serialization;
 using Horde.Build.Server;
 using Horde.Build.Streams;
 using Horde.Build.Users;
 using Horde.Build.Utilities;
 using HordeCommon;
-using Microsoft.Extensions.Caching.Memory;
 using Microsoft.Extensions.Hosting;
 using Microsoft.Extensions.Logging;
 using Microsoft.Extensions.Options;
@@ -31,7 +22,6 @@ using StackExchange.Redis;
 
 namespace Horde.Build.Perforce
 {
-	using CommitId = ObjectId<ICommit>;
 	using StreamId = StringId<IStream>;
 
 	/// <summary>
diff --git a/Engine/Source/Programs/Horde/Horde.Build/Perforce/ICommitService.cs b/Engine/Source/Programs/Horde/Horde.Build/Perforce/ICommitService.cs
--- a/Engine/Source/Programs/Horde/Horde.Build/Perforce/ICommitService.cs
+++ b/Engine/Source/Programs/Horde/Horde.Build/Perforce/ICommitService.cs
@@ -2,9 +2,6 @@
 
 using System;
 using System.Threading.Tasks;
-using EpicGames.Horde.Storage;
-using Horde.Build.Streams;
-using Horde.Build.Utilities;
 
 namespace Horde.Build.Perforce
 {
diff --git a/Engine/Source/Programs/Horde/Horde.Build/Perforce/IPerforceService.cs b/Engine/Source/Programs/Horde/Horde.Build/Perforce/IPerforceService.cs
--- a/Engine/Source/Programs/Horde/Horde.Build/Perforce/IPerforceService.cs
+++ b/Engine/Source/Programs/Horde/Horde.Build/Perforce/IPerforceService.cs
@@ -2,7 +2,6 @@
 
 using System;
 using System.Collections.Generic;
-using System.Diagnostics.CodeAnalysis;
 using System.Text.RegularExpressions;
 using System.Threading;
 using System.Threading.Tasks;
diff --git a/Engine/Source/Programs/Horde/Horde.Build/Perforce/PerforceLoadBalancer.cs b/Engine/Source/Programs/Horde/Horde.Build/Perforce/PerforceLoadBalancer.cs
--- a/Engine/Source/Programs/Horde/Horde.Build/Perforce/PerforceLoadBalancer.cs
+++ b/Engine/Source/Programs/Horde/Horde.Build/Perforce/PerforceLoadBalancer.cs
@@ -256,7 +256,7 @@ namespace Horde.Build.Perforce
 			}
 			if (properties == null)
 			{
-				throw new ArgumentException(nameof(properties));
+				throw new ArgumentException("Properties argument to EvaluateCondition may not be null", nameof(properties));
 			}
 			return condition.Evaluate(x => GetPropertyValues(properties, x));
 		}
diff --git a/Engine/Source/Programs/Horde/Horde.Build/Perforce/PerforceService.cs b/Engine/Source/Programs/Horde/Horde.Build/Perforce/PerforceService.cs
--- a/Engine/Source/Programs/Horde/Horde.Build/Perforce/PerforceService.cs
+++ b/Engine/Source/Programs/Horde/Horde.Build/Perforce/PerforceService.cs
@@ -2,13 +2,11 @@
 
 using System;
 using System.Collections.Generic;
-using System.Diagnostics;
 using System.Diagnostics.CodeAnalysis;
 using System.Globalization;
 using System.IO;
 using System.Linq;
 using System.Net;
-using System.Reflection;
 using System.Runtime.InteropServices;
 using System.Text;
 using System.Text.RegularExpressions;
@@ -16,7 +14,6 @@ using System.Threading;
 using System.Threading.Tasks;
 using EpicGames.Core;
 using EpicGames.Perforce;
-using Grpc.Core;
 using Horde.Build.Server;
 using Horde.Build.Users;
 using Horde.Build.Utilities;
@@ -25,7 +22,6 @@ using Microsoft.Extensions.Logging;
 using Microsoft.Extensions.Options;
 using OpenTracing;
 using OpenTracing.Util;
-using StackExchange.Redis;
 
 namespace Horde.Build.Perforce
 {
diff --git a/Engine/Source/Programs/Horde/Horde.Build/Perforce/ReplicationService.cs b/Engine/Source/Programs/Horde/Horde.Build/Perforce/ReplicationService.cs
--- a/Engine/Source/Programs/Horde/Horde.Build/Perforce/ReplicationService.cs
+++ b/Engine/Source/Programs/Horde/Horde.Build/Perforce/ReplicationService.cs
@@ -1,15 +1,12 @@
 // Copyright Epic Games, Inc. All Rights Reserved.
 
 using System;
-using System.Buffers;
 using System.Collections.Generic;
 using System.Diagnostics;
-using System.Globalization;
 using System.IO;
 using System.Linq;
 using System.Security.Cryptography;
 using System.Text;
-using System.Text.RegularExpressions;
 using System.Threading;
 using System.Threading.Tasks;
 using EpicGames.Core;
@@ -19,13 +16,9 @@ using EpicGames.Horde.Storage.Git;
 using EpicGames.Horde.Storage.Nodes;
 using EpicGames.Perforce;
 using EpicGames.Redis;
-using EpicGames.Serialization;
 using Horde.Build.Server;
 using Horde.Build.Streams;
-using Horde.Build.Users;
 using Horde.Build.Utilities;
-using HordeCommon;
-using Microsoft.Extensions.Caching.Memory;
 using Microsoft.Extensions.Hosting;
 using Microsoft.Extensions.Logging;
 using Microsoft.Extensions.Options;
@@ -33,7 +26,6 @@ using StackExchange.Redis;
 
 namespace Horde.Build.Perforce
 {
-	using CommitId = ObjectId<ICommit>;
 	using StreamId = StringId<IStream>;
 
 	/// <summary>
@@ -192,7 +184,7 @@ namespace Horde.Build.Perforce
 		/// <summary>
 		/// Constructor
 		/// </summary>
-		public ReplicationService(RedisService redisService, CommitService commitService, IStreamCollection streamCollection, IPerforceService perforceService, ITreeStore<ReplicationService> treeStore, IClock clock, IOptions<ReplicationServiceOptions> options, ILogger<ReplicationService> logger)
+		public ReplicationService(RedisService redisService, CommitService commitService, IStreamCollection streamCollection, IPerforceService perforceService, ITreeStore<ReplicationService> treeStore, IOptions<ReplicationServiceOptions> options, ILogger<ReplicationService> logger)
 		{
 			Options = options.Value;
 
@@ -527,14 +519,14 @@ namespace Horde.Build.Perforce
 		[DebuggerDisplay("{_path}")]
 		class DirectoryToSync
 		{
-			public readonly Utf8String _path;
-			public readonly Dictionary<Utf8String, long> _fileNameToSize;
+			public readonly Utf8String Path;
+			public readonly Dictionary<Utf8String, long> FileNameToSize;
 			public long _size;
 
 			public DirectoryToSync(Utf8String path, Utf8StringComparer comparer)
 			{
-				_path = path;
-				_fileNameToSize = new Dictionary<Utf8String, long>(comparer);
+				Path = path;
+				FileNameToSize = new Dictionary<Utf8String, long>(comparer);
 			}
 		}
 
@@ -692,12 +684,12 @@ namespace Horde.Build.Perforce
 					Utf8String directoryPath = path.Slice(clientRoot.Length, fileIdx - clientRoot.Length);
 
 					DirectoryToSync directory = FindOrAddDirectoryTree(pathToDirectory, directoryPath, clientInfo.ServerInfo.Utf8PathComparer);
-					directory._fileNameToSize.Add(path.Slice(fileIdx).Clone(), response.Data.FileSize);
+					directory.FileNameToSize.Add(path.Slice(fileIdx).Clone(), response.Data.FileSize);
 					directory._size += response.Data.FileSize;
 				}
 
 				// Sort the directories by name
-				List<DirectoryToSync> directories = pathToDirectory.Values.OrderBy(x => x._path, clientInfo.ServerInfo.Utf8PathComparer).ToList();
+				List<DirectoryToSync> directories = pathToDirectory.Values.OrderBy(x => x.Path, clientInfo.ServerInfo.Utf8PathComparer).ToList();
 
 				// Output some stats for the sync
 				long totalSize = directories.Sum(x => x._size);
@@ -740,17 +732,17 @@ namespace Horde.Build.Perforce
 					List<string> syncPaths = new List<string>();
 					for (int idx = dirIdx; idx < directories.Count; idx++)
 					{
-						Utf8String basePath = directories[idx]._path;
+						Utf8String basePath = directories[idx].Path;
 						syncPaths.Add($"//{clientInfo.Client.Name}/{basePath}...@{change}");
 
 						long dirSize = directories[idx]._size;
-						while (idx + 1 < directories.Count && directories[idx + 1]._path.StartsWith(basePath, clientInfo.ServerInfo.Utf8PathComparer))
+						while (idx + 1 < directories.Count && directories[idx + 1].Path.StartsWith(basePath, clientInfo.ServerInfo.Utf8PathComparer))
 						{
 							dirSize += directories[idx + 1]._size;
 							idx++;
 						}
 
-						_logger.LogInformation("  {Directory} ({Size:n1}mb)", directories[idx]._path + "...", dirSize / (1024.0 * 1024.0));
+						_logger.LogInformation("  {Directory} ({Size:n1}mb)", directories[idx].Path + "...", dirSize / (1024.0 * 1024.0));
 					}
 
 					Stopwatch syncTimer = Stopwatch.StartNew();
@@ -790,7 +782,7 @@ namespace Horde.Build.Perforce
 									{
 										throw new ReplicationException($"Unable to find directory for {file}");
 									}
-									if (!directory._fileNameToSize.TryGetValue(file.Substring(offset), out fileSize))
+									if (!directory.FileNameToSize.TryGetValue(file.Substring(offset), out fileSize))
 									{
 										throw new ReplicationException($"Unable to find file entry for {file}");
 									}
@@ -846,7 +838,7 @@ namespace Horde.Build.Perforce
 					// Combine all the existing sync paths together with a new wildcard.
 					while (directories.Count > dirIdx)
 					{
-						Utf8String nextPath = directories[^1]._path;
+						Utf8String nextPath = directories[^1].Path;
 						if (syncNode.Paths.Count > 0)
 						{
 							Utf8String lastPath = syncNode.Paths[^1];
diff --git a/Engine/Source/Programs/Horde/Horde.Build/Projects/IProject.cs b/Engine/Source/Programs/Horde/Horde.Build/Projects/IProject.cs
--- a/Engine/Source/Programs/Horde/Horde.Build/Projects/IProject.cs
+++ b/Engine/Source/Programs/Horde/Horde.Build/Projects/IProject.cs
@@ -5,7 +5,6 @@ using System.Linq;
 using System.Text.RegularExpressions;
 using EpicGames.Core;
 using Horde.Build.Acls;
-using Horde.Build.Projects;
 using Horde.Build.Streams;
 using Horde.Build.Utilities;
 using MongoDB.Bson.Serialization.Attributes;
diff --git a/Engine/Source/Programs/Horde/Horde.Build/Projects/IProjectCollection.cs b/Engine/Source/Programs/Horde/Horde.Build/Projects/IProjectCollection.cs
--- a/Engine/Source/Programs/Horde/Horde.Build/Projects/IProjectCollection.cs
+++ b/Engine/Source/Programs/Horde/Horde.Build/Projects/IProjectCollection.cs
@@ -2,7 +2,6 @@
 
 using System.Collections.Generic;
 using System.Threading.Tasks;
-using Horde.Build.Server;
 using Horde.Build.Utilities;
 
 namespace Horde.Build.Projects
diff --git a/Engine/Source/Programs/Horde/Horde.Build/Projects/ProjectMessages.cs b/Engine/Source/Programs/Horde/Horde.Build/Projects/ProjectMessages.cs
--- a/Engine/Source/Programs/Horde/Horde.Build/Projects/ProjectMessages.cs
+++ b/Engine/Source/Programs/Horde/Horde.Build/Projects/ProjectMessages.cs
@@ -3,7 +3,6 @@
 using System.Collections.Generic;
 using System.ComponentModel.DataAnnotations;
 using Horde.Build.Acls;
-using Horde.Build.Projects;
 
 namespace Horde.Build.Projects
 {
diff --git a/Engine/Source/Programs/Horde/Horde.Build/Server/AdminController.cs b/Engine/Source/Programs/Horde/Horde.Build/Server/AdminController.cs
--- a/Engine/Source/Programs/Horde/Horde.Build/Server/AdminController.cs
+++ b/Engine/Source/Programs/Horde/Horde.Build/Server/AdminController.cs
@@ -6,7 +6,6 @@ using System.Linq;
 using System.Security.Claims;
 using System.Threading.Tasks;
 using Horde.Build.Acls;
-using Horde.Build.Agents.Software;
 using Horde.Build.Utilities;
 using Microsoft.AspNetCore.Authorization;
 using Microsoft.AspNetCore.Mvc;
diff --git a/Engine/Source/Programs/Horde/Horde.Build/Server/DashboardController.Messages.cs b/Engine/Source/Programs/Horde/Horde.Build/Server/DashboardController.Messages.cs
--- a/Engine/Source/Programs/Horde/Horde.Build/Server/DashboardController.Messages.cs
+++ b/Engine/Source/Programs/Horde/Horde.Build/Server/DashboardController.Messages.cs
@@ -31,14 +31,5 @@ namespace Horde.Build.Server
 		/// Help slack channel that users can use for issues
 		/// </summary>
 		public string? HelpSlackChannel { get; set; }
-
-		/// <summary>
-		/// Response constructor
-		/// </summary>
-		public GetDashboardConfigResponse()
-		{
-		}
-
 	}
 }
-
diff --git a/Engine/Source/Programs/Horde/Horde.Build/Server/GlobalConfig.cs b/Engine/Source/Programs/Horde/Horde.Build/Server/GlobalConfig.cs
--- a/Engine/Source/Programs/Horde/Horde.Build/Server/GlobalConfig.cs
+++ b/Engine/Source/Programs/Horde/Horde.Build/Server/GlobalConfig.cs
@@ -3,23 +3,17 @@
 using System;
 using System.Collections.Generic;
 using System.ComponentModel.DataAnnotations;
-using System.Diagnostics.CodeAnalysis;
-using System.Linq;
 using EpicGames.Core;
 using EpicGames.Horde.Common;
 using EpicGames.Serialization;
 using Horde.Build.Acls;
-using Horde.Build.Jobs.Graphs;
 using Horde.Build.Projects;
-using Horde.Build.Streams;
 using Horde.Build.Tools;
 using Horde.Build.Utilities;
 
 namespace Horde.Build.Server
 {
 	using ProjectId = StringId<IProject>;
-	using StreamId = StringId<IStream>;
-	using TemplateRefId = StringId<TemplateRef>;
 	
 	/// <summary>
 	/// Global configuration
diff --git a/Engine/Source/Programs/Horde/Horde.Build/Server/Globals.cs b/Engine/Source/Programs/Horde/Horde.Build/Server/Globals.cs
--- a/Engine/Source/Programs/Horde/Horde.Build/Server/Globals.cs
+++ b/Engine/Source/Programs/Horde/Horde.Build/Server/Globals.cs
@@ -8,7 +8,6 @@ using System.Security.Cryptography;
 using EpicGames.Horde.Common;
 using EpicGames.Perforce;
 using Horde.Build.Acls;
-using Horde.Build.Server;
 using Horde.Build.Utilities;
 using MongoDB.Bson;
 using MongoDB.Bson.Serialization.Attributes;
diff --git a/Engine/Source/Programs/Horde/Horde.Build/Server/HealthService.cs b/Engine/Source/Programs/Horde/Horde.Build/Server/HealthService.cs
--- a/Engine/Source/Programs/Horde/Horde.Build/Server/HealthService.cs
+++ b/Engine/Source/Programs/Horde/Horde.Build/Server/HealthService.cs
@@ -3,7 +3,6 @@
 using System.Threading.Tasks;
 using Grpc.Core;
 using Grpc.Health.V1;
-using Horde.Build.Server;
 using static Grpc.Health.V1.HealthCheckResponse.Types;
 
 namespace Horde.Build.Server
diff --git a/Engine/Source/Programs/Horde/Horde.Build/Server/LifetimeController.cs b/Engine/Source/Programs/Horde/Horde.Build/Server/LifetimeController.cs
--- a/Engine/Source/Programs/Horde/Horde.Build/Server/LifetimeController.cs
+++ b/Engine/Source/Programs/Horde/Horde.Build/Server/LifetimeController.cs
@@ -1,7 +1,6 @@
 // Copyright Epic Games, Inc. All Rights Reserved.
 
 using System.Threading.Tasks;
-using Horde.Build.Server;
 using Microsoft.AspNetCore.Mvc;
 
 namespace Horde.Build.Server
diff --git a/Engine/Source/Programs/Horde/Horde.Build/Server/Notices/NoticeMessages.cs b/Engine/Source/Programs/Horde/Horde.Build/Server/Notices/NoticeMessages.cs
--- a/Engine/Source/Programs/Horde/Horde.Build/Server/Notices/NoticeMessages.cs
+++ b/Engine/Source/Programs/Horde/Horde.Build/Server/Notices/NoticeMessages.cs
@@ -109,14 +109,5 @@ namespace Horde.Build.Server.Notices
 		/// User id who created the notice, otherwise null if a system message
 		/// </summary>
 		public GetThinUserInfoResponse? CreatedByUser { get; set; }
-
-		/// <summary>
-		/// Constructor
-		/// </summary>
-		public GetNoticeResponse()
-		{
-
-		}
-
 	}
-}
+}
\ No newline at end of file
diff --git a/Engine/Source/Programs/Horde/Horde.Build/Server/RedisService.cs b/Engine/Source/Programs/Horde/Horde.Build/Server/RedisService.cs
--- a/Engine/Source/Programs/Horde/Horde.Build/Server/RedisService.cs
+++ b/Engine/Source/Programs/Horde/Horde.Build/Server/RedisService.cs
@@ -10,7 +10,6 @@ using System.Runtime.InteropServices;
 using System.Threading.Tasks;
 using EpicGames.Core;
 using EpicGames.Redis;
-using Horde.Build.Utilities;
 using Microsoft.Extensions.Logging;
 using Microsoft.Extensions.Options;
 using StackExchange.Redis;
@@ -46,7 +45,7 @@ namespace Horde.Build.Server
 		/// The database interface.
 		/// If possible, use ConnectionPool instead. 
 		/// </summary>
-		public IDatabase Database { get; }
+		public IDatabase DatabaseSingleton { get; }
 
 		/// <summary>
 		/// Connection pool
@@ -108,7 +107,7 @@ namespace Horde.Build.Server
 			}
 
 			_multiplexer = ConnectionMultiplexer.Connect(connectionString);
-			Database = _multiplexer.GetDatabase(dbNum);
+			DatabaseSingleton = _multiplexer.GetDatabase(dbNum);
 			ConnectionPool = new RedisConnectionPool(20, connectionString, dbNum);
 		}
 
diff --git a/Engine/Source/Programs/Horde/Horde.Build/Server/RpcService.cs b/Engine/Source/Programs/Horde/Horde.Build/Server/RpcService.cs
--- a/Engine/Source/Programs/Horde/Horde.Build/Server/RpcService.cs
+++ b/Engine/Source/Programs/Horde/Horde.Build/Server/RpcService.cs
@@ -252,12 +252,22 @@ namespace Horde.Build.Server
 			properties = new List<string>();
 			resources = new Dictionary<string, int>();
 
-			if (capabilities == null) return;
+			if (capabilities == null)
+			{
+				return;
+			}
 			properties.AddRange(capabilities.Properties);
 
-			if (capabilities.Devices.Count <= 0) return;
+			if (capabilities.Devices.Count <= 0)
+			{
+				return;
+			}
+
 			RpcDeviceCapabilities device = capabilities.Devices[0];
-			if (device.Properties == null) return;
+			if (device.Properties == null)
+			{
+				return;
+			}
 
 			properties.AddRange(device.Properties);
 			CopyPropertyToResource(KnownPropertyNames.LogicalCores, properties, resources);
@@ -585,10 +595,7 @@ namespace Horde.Build.Server
 			}
 
 			// Create a log file if necessary
-			if (log.Value == null)
-			{
-				log.Value = await _logFileService.CreateLogFileAsync(job.Id, batch.SessionId, LogType.Json);
-			}
+			log.Value ??= await _logFileService.CreateLogFileAsync(job.Id, batch.SessionId, LogType.Json);
 
 			// Get the node for this step
 			IGraph graph = await _jobService.GetGraphAsync(job);
diff --git a/Engine/Source/Programs/Horde/Horde.Build/Server/SchemaController.cs b/Engine/Source/Programs/Horde/Horde.Build/Server/SchemaController.cs
--- a/Engine/Source/Programs/Horde/Horde.Build/Server/SchemaController.cs
+++ b/Engine/Source/Programs/Horde/Horde.Build/Server/SchemaController.cs
@@ -71,10 +71,7 @@ namespace Horde.Build.Server
 				host = hosts.FirstOrDefault();
 			}
 
-			if (host == null)
-			{
-				host = Dns.GetHostName();
-			}
+			host ??= Dns.GetHostName();
 
 			CatalogRoot root = new CatalogRoot();
 			foreach (Type schemaType in ConfigSchemas)
diff --git a/Engine/Source/Programs/Horde/Horde.Build/Server/UpgradeService.cs b/Engine/Source/Programs/Horde/Horde.Build/Server/UpgradeService.cs
--- a/Engine/Source/Programs/Horde/Horde.Build/Server/UpgradeService.cs
+++ b/Engine/Source/Programs/Horde/Horde.Build/Server/UpgradeService.cs
@@ -3,7 +3,6 @@
 using System;
 using System.Threading.Tasks;
 using Horde.Build.Jobs;
-using Horde.Build.Server;
 using Horde.Build.Users;
 using Microsoft.Extensions.DependencyInjection;
 using Microsoft.Extensions.Logging;
diff --git a/Engine/Source/Programs/Horde/Horde.Build/ServerSettings.cs b/Engine/Source/Programs/Horde/Horde.Build/ServerSettings.cs
--- a/Engine/Source/Programs/Horde/Horde.Build/ServerSettings.cs
+++ b/Engine/Source/Programs/Horde/Horde.Build/ServerSettings.cs
@@ -6,7 +6,6 @@ using EpicGames.Horde.Storage;
 using EpicGames.Horde.Storage.Bundles;
 using EpicGames.Horde.Storage.Nodes;
 using Horde.Build.Agents.Fleet;
-using Horde.Build.Storage;
 using Horde.Build.Storage.Backends;
 using Horde.Build.Utilities;
 using TimeZoneConverter;
@@ -584,10 +583,7 @@ namespace Horde.Build
 		{
 			get
 			{
-				if (_cachedTimeZoneInfo == null)
-				{
-					_cachedTimeZoneInfo = (ScheduleTimeZone == null) ? TimeZoneInfo.Local : TZConvert.GetTimeZoneInfo(ScheduleTimeZone);
-				}
+				_cachedTimeZoneInfo ??= (ScheduleTimeZone == null) ? TimeZoneInfo.Local : TZConvert.GetTimeZoneInfo(ScheduleTimeZone);
 				return _cachedTimeZoneInfo;
 			}
 		}
diff --git a/Engine/Source/Programs/Horde/Horde.Build/Startup.cs b/Engine/Source/Programs/Horde/Horde.Build/Startup.cs
--- a/Engine/Source/Programs/Horde/Horde.Build/Startup.cs
+++ b/Engine/Source/Programs/Horde/Horde.Build/Startup.cs
@@ -81,7 +81,6 @@ using OpenTracing.Contrib.Grpc.Interceptors;
 using OpenTracing.Util;
 using Serilog;
 using Serilog.Events;
-using StackExchange.Redis;
 using StatsdClient;
 using Status = Grpc.Core.Status;
 using Horde.Build.Users;
@@ -325,7 +324,7 @@ namespace Horde.Build
 			RedisService redisService = new RedisService(settings);
 #pragma warning restore CA2000 // Dispose objects before losing scope
 			services.AddSingleton<RedisService>(sp => redisService);
-			services.AddDataProtection().PersistKeysToStackExchangeRedis(() => redisService.Database, "aspnet-data-protection");
+			services.AddDataProtection().PersistKeysToStackExchangeRedis(() => redisService.DatabaseSingleton, "aspnet-data-protection");
 
 			if (settings.CorsEnabled)
 			{
diff --git a/Engine/Source/Programs/Horde/Horde.Build/Storage/Backends/TransientStorageBackend.cs b/Engine/Source/Programs/Horde/Horde.Build/Storage/Backends/TransientStorageBackend.cs
--- a/Engine/Source/Programs/Horde/Horde.Build/Storage/Backends/TransientStorageBackend.cs
+++ b/Engine/Source/Programs/Horde/Horde.Build/Storage/Backends/TransientStorageBackend.cs
@@ -36,7 +36,7 @@ namespace Horde.Build.Storage.Backends
 		{
 			using (MemoryStream buffer = new MemoryStream())
 			{
-				await stream.CopyToAsync(buffer);
+				await stream.CopyToAsync(buffer, cancellationToken);
 				_pathToData[path] = buffer.ToArray();
 			}
 		}
diff --git a/Engine/Source/Programs/Horde/Horde.Build/Storage/BasicBlobStore.cs b/Engine/Source/Programs/Horde/Horde.Build/Storage/BasicBlobStore.cs
--- a/Engine/Source/Programs/Horde/Horde.Build/Storage/BasicBlobStore.cs
+++ b/Engine/Source/Programs/Horde/Horde.Build/Storage/BasicBlobStore.cs
@@ -3,8 +3,6 @@
 using System;
 using System.Buffers;
 using System.Collections.Generic;
-using System.IO;
-using System.Text;
 using System.Threading;
 using System.Threading.Tasks;
 using EpicGames.Core;
diff --git a/Engine/Source/Programs/Horde/Horde.Build/Storage/IStorageBackend.cs b/Engine/Source/Programs/Horde/Horde.Build/Storage/IStorageBackend.cs
--- a/Engine/Source/Programs/Horde/Horde.Build/Storage/IStorageBackend.cs
+++ b/Engine/Source/Programs/Horde/Horde.Build/Storage/IStorageBackend.cs
@@ -1,7 +1,6 @@
 // Copyright Epic Games, Inc. All Rights Reserved.
 
 using System;
-using System.Buffers;
 using System.IO;
 using System.Threading;
 using System.Threading.Tasks;
diff --git a/Engine/Source/Programs/Horde/Horde.Build/Storage/Services/BasicStorageClient.cs b/Engine/Source/Programs/Horde/Horde.Build/Storage/Services/BasicStorageClient.cs
--- a/Engine/Source/Programs/Horde/Horde.Build/Storage/Services/BasicStorageClient.cs
+++ b/Engine/Source/Programs/Horde/Horde.Build/Storage/Services/BasicStorageClient.cs
@@ -62,7 +62,7 @@ namespace Horde.Build.Storage.Services
 		{
 			string path = GetFullBlobPath(namespaceId, hash);
 
-			Stream? stream = await _storageBackend.ReadAsync(path);
+			Stream? stream = await _storageBackend.ReadAsync(path, cancellationToken);
 			if(stream == null)
 			{
 				throw new BlobNotFoundException(namespaceId, hash);
@@ -94,7 +94,7 @@ namespace Horde.Build.Storage.Services
 			}
 
 			string path = GetFullBlobPath(namespaceId, hash);
-			await _storageBackend.WriteBytesAsync(path, data);
+			await _storageBackend.WriteBytesAsync(path, data, cancellationToken);
 		}
 
 		/// <inheritdoc/>
@@ -110,7 +110,7 @@ namespace Horde.Build.Storage.Services
 			IoHash dataHash = IoHash.Compute(data);
 
 			string path = GetFullBlobPath(namespaceId, dataHash);
-			await _storageBackend.WriteBytesAsync(path, data);
+			await _storageBackend.WriteBytesAsync(path, data, cancellationToken);
 
 			return dataHash;
 		}
@@ -130,7 +130,7 @@ namespace Horde.Build.Storage.Services
 		/// <inheritdoc/>
 		public async Task<bool> HasBlobAsync(NamespaceId namespaceId, IoHash hash, CancellationToken cancellationToken)
 		{
-			return await _storageBackend.ExistsAsync(GetFullBlobPath(namespaceId, hash));
+			return await _storageBackend.ExistsAsync(GetFullBlobPath(namespaceId, hash), cancellationToken);
 		}
 
 		/// <inheritdoc/>
diff --git a/Engine/Source/Programs/Horde/Horde.Build/Streams/IStream.cs b/Engine/Source/Programs/Horde/Horde.Build/Streams/IStream.cs
--- a/Engine/Source/Programs/Horde/Horde.Build/Streams/IStream.cs
+++ b/Engine/Source/Programs/Horde/Horde.Build/Streams/IStream.cs
@@ -12,7 +12,6 @@ using Horde.Build.Agents.Pools;
 using Horde.Build.Jobs.Schedules;
 using Horde.Build.Jobs.Templates;
 using Horde.Build.Projects;
-using Horde.Build.Server;
 using Horde.Build.Users;
 using Horde.Build.Utilities;
 using HordeCommon;
@@ -310,7 +309,6 @@ namespace Horde.Build.Streams
 		/// </summary>
 		public DateTime? PauseTimeUtc { get; set; }
 
-
 		/// <summary>
 		/// Default constructor for serialization
 		/// </summary>
@@ -327,10 +325,8 @@ namespace Horde.Build.Streams
 			PausedByUserId = pausedByUserId;
 			PauseTimeUtc = pauseTimeUtc;
 		}
-
 	}
 
-
 	/// <summary>
 	/// Reference to a template
 	/// </summary>
diff --git a/Engine/Source/Programs/Horde/Horde.Build/Streams/StreamCollection.cs b/Engine/Source/Programs/Horde/Horde.Build/Streams/StreamCollection.cs
--- a/Engine/Source/Programs/Horde/Horde.Build/Streams/StreamCollection.cs
+++ b/Engine/Source/Programs/Horde/Horde.Build/Streams/StreamCollection.cs
@@ -231,7 +231,6 @@ namespace Horde.Build.Streams
 					{
 						newTemplateRef.StepStates = new List<TemplateStepState>(oldTemplateRef.StepStates);
 					}
-
 				}
 				newTemplateRefs.Add(newTemplateRefId, newTemplateRef);
 			}
@@ -292,7 +291,6 @@ namespace Horde.Build.Streams
 				{
 					updates.Add(updateBuilder.Set(x => x.Templates, newTemplates));
 				}
-				
 			}
 			else if (stepStates != null)
 			{
@@ -306,21 +304,21 @@ namespace Horde.Build.Streams
 				{
 					TemplateStepState? newState = newStepStates.Where(x => x.Name == updateState.Name).FirstOrDefault();
 
-					UserId? PausedByUserId = updateState.PausedByUserId != null ? new UserId(updateState.PausedByUserId) : null;
+					UserId? pausedByUserId = updateState.PausedByUserId != null ? new UserId(updateState.PausedByUserId) : null;
 
 					if (newState == null)
 					{
 						// if this is a new state without anything set, ignore it
-						if (PausedByUserId == null)
+						if (pausedByUserId == null)
 						{
 							continue;
 						}
 
-						newStepStates.Add(new TemplateStepState(updateState.Name, PausedByUserId, PausedByUserId != null ? DateTime.UtcNow : null));
+						newStepStates.Add(new TemplateStepState(updateState.Name, pausedByUserId, DateTime.UtcNow));
 					}
 					else
 					{
-						if (PausedByUserId == null)
+						if (pausedByUserId == null)
 						{
 							newState.PauseTimeUtc = null;
 						}
@@ -329,7 +327,7 @@ namespace Horde.Build.Streams
 							newState.PauseTimeUtc = DateTime.UtcNow;
 						}
 
-						newState.PausedByUserId = PausedByUserId;						
+						newState.PausedByUserId = pausedByUserId;						
 					}
 				}
 
diff --git a/Engine/Source/Programs/Horde/Horde.Build/Streams/StreamMessages.cs b/Engine/Source/Programs/Horde/Horde.Build/Streams/StreamMessages.cs
--- a/Engine/Source/Programs/Horde/Horde.Build/Streams/StreamMessages.cs
+++ b/Engine/Source/Programs/Horde/Horde.Build/Streams/StreamMessages.cs
@@ -7,7 +7,6 @@ using EpicGames.Core;
 using Horde.Build.Acls;
 using Horde.Build.Jobs.Graphs;
 using Horde.Build.Jobs.Templates;
-using Horde.Build.Server;
 using Horde.Build.Jobs.Schedules;
 using Horde.Build.Users;
 using Horde.Build.Utilities;
@@ -278,21 +277,17 @@ namespace Horde.Build.Streams
 		public string? PausedByUserId { get; set; }
 	}
 
-
 	/// <summary>
 	/// Updates an existing stream template ref
 	/// </summary>
 	public class UpdateTemplateRefRequest
 	{
-		
 		/// <summary>
 		/// Step states to update
 		/// </summary>
 		public List<UpdateStepStateRequest>? StepStates { get; set; }
-
 	}
 
-
 	/// <summary>
 	/// Mapping from a BuildGraph agent type to a set of machines on the farm
 	/// </summary>
@@ -472,10 +467,8 @@ namespace Horde.Build.Streams
 			PauseTimeUtc = state.PauseTimeUtc;
 			PausedByUserInfo = pausedByUserInfo;
 		}
-
 	}
 
-
 	/// <summary>
 	/// Information about a template in this stream
 	/// </summary>
diff --git a/Engine/Source/Programs/Horde/Horde.Build/Streams/StreamsController.cs b/Engine/Source/Programs/Horde/Horde.Build/Streams/StreamsController.cs
--- a/Engine/Source/Programs/Horde/Horde.Build/Streams/StreamsController.cs
+++ b/Engine/Source/Programs/Horde/Horde.Build/Streams/StreamsController.cs
@@ -144,14 +144,8 @@ namespace Horde.Build.Streams
 
 								stepStates ??= new List<GetTemplateStepStateResponse>();
 
-								GetThinUserInfoResponse? pausedByUserInfo = null;
-								if (state.PausedByUserId != null)
-								{
-									pausedByUserInfo = new GetThinUserInfoResponse(await _userCollection.GetCachedUserAsync(state.PausedByUserId));
-								}
-
+								GetThinUserInfoResponse? pausedByUserInfo = new GetThinUserInfoResponse(await _userCollection.GetCachedUserAsync(state.PausedByUserId));
 								stepStates.Add(new GetTemplateStepStateResponse(state, pausedByUserInfo));
-
 							}
 						}
 
@@ -347,8 +341,6 @@ namespace Horde.Build.Streams
 			await _streamService.TryUpdateTemplateRefAsync(stream, templateRefId, update.StepStates);
 
 			return Ok();
-
 		}
-
 	}
 }
diff --git a/Engine/Source/Programs/Horde/Horde.Build/Tasks/ITaskSource.cs b/Engine/Source/Programs/Horde/Horde.Build/Tasks/ITaskSource.cs
--- a/Engine/Source/Programs/Horde/Horde.Build/Tasks/ITaskSource.cs
+++ b/Engine/Source/Programs/Horde/Horde.Build/Tasks/ITaskSource.cs
@@ -233,6 +233,7 @@ namespace Horde.Build.Tasks
         /// <returns>Lease task</returns>
         protected static Task<AgentLease?> Skip(CancellationToken token)
         {
+			_ = token;
             return Task.FromResult<AgentLease?>(null);
         }
 
diff --git a/Engine/Source/Programs/Horde/Horde.Build/Tools/Tool.cs b/Engine/Source/Programs/Horde/Horde.Build/Tools/Tool.cs
--- a/Engine/Source/Programs/Horde/Horde.Build/Tools/Tool.cs
+++ b/Engine/Source/Programs/Horde/Horde.Build/Tools/Tool.cs
@@ -88,7 +88,7 @@ namespace Horde.Build.Tools
 		public AclV2? Acl
 		{
 			get => _acl;
-			set => _acl = (value is AclV2 && !value.IsDefault()) ? value : null;
+			set => _acl = (value is not null && !value.IsDefault()) ? value : null;
 		}
 
 		private AclV2? _acl;
diff --git a/Engine/Source/Programs/Horde/Horde.Build/Tools/ToolCollection.cs b/Engine/Source/Programs/Horde/Horde.Build/Tools/ToolCollection.cs
--- a/Engine/Source/Programs/Horde/Horde.Build/Tools/ToolCollection.cs
+++ b/Engine/Source/Programs/Horde/Horde.Build/Tools/ToolCollection.cs
@@ -164,7 +164,7 @@ namespace Horde.Build.Tools
 		/// </summary>
 		/// <param name="id">The tool to remove</param>
 		/// <returns>True if a tool with the given identifier was deleted</returns>
-		private async Task<bool> DeleteAsync(ToolId id)
+		public async Task<bool> DeleteAsync(ToolId id)
 		{
 			bool deleted = await _tools.DeleteAsync(id);
 			if (deleted)
diff --git a/Engine/Source/Programs/Horde/Horde.Build/Users/UserCollectionV2.cs b/Engine/Source/Programs/Horde/Horde.Build/Users/UserCollectionV2.cs
--- a/Engine/Source/Programs/Horde/Horde.Build/Users/UserCollectionV2.cs
+++ b/Engine/Source/Programs/Horde/Horde.Build/Users/UserCollectionV2.cs
@@ -291,10 +291,7 @@ namespace Horde.Build.Users
 		public async Task<IUserClaims> GetClaimsAsync(UserId userId)
 		{
 			IUserClaims? claims = await _userClaims.Find(x => x.Id == userId).FirstOrDefaultAsync();
-			if(claims == null)
-			{
-				claims = new UserClaimsDocument(userId);
-			}
+			claims ??= new UserClaimsDocument(userId);
 			return claims;
 		}
 
@@ -310,10 +307,7 @@ namespace Horde.Build.Users
 		public async Task<IUserSettings> GetSettingsAsync(UserId userId)
 		{
 			IUserSettings? settings = await _userSettings.Find(x => x.Id == userId).FirstOrDefaultAsync();
-			if (settings == null)
-			{
-				settings = new UserSettingsDocument(userId);
-			}
+			settings ??= new UserSettingsDocument(userId);
 			return settings;
 		}
 
diff --git a/Engine/Source/Programs/Horde/Horde.Build/Utilities/ObjectId.cs b/Engine/Source/Programs/Horde/Horde.Build/Utilities/ObjectId.cs
--- a/Engine/Source/Programs/Horde/Horde.Build/Utilities/ObjectId.cs
+++ b/Engine/Source/Programs/Horde/Horde.Build/Utilities/ObjectId.cs
@@ -7,7 +7,6 @@ using System.Globalization;
 using System.Text.Json;
 using System.Text.Json.Serialization;
 using EpicGames.Core;
-using EpicGames.Serialization;
 using MongoDB.Bson;
 using MongoDB.Bson.Serialization;
 using MongoDB.Bson.Serialization.Serializers;
diff --git a/Engine/Source/Programs/Horde/Horde.Build/Utilities/ReadOnlyTrieBuilder.cs b/Engine/Source/Programs/Horde/Horde.Build/Utilities/ReadOnlyTrieBuilder.cs
--- a/Engine/Source/Programs/Horde/Horde.Build/Utilities/ReadOnlyTrieBuilder.cs
+++ b/Engine/Source/Programs/Horde/Horde.Build/Utilities/ReadOnlyTrieBuilder.cs
@@ -41,14 +41,8 @@ namespace Horde.Build.Logs
 			for (int shift = (sizeof(ulong) * 8) - 4; shift >= 0; shift -= 4)
 			{
 				int index = (int)(value >> shift) & 15;
-				if (leaf._children == null)
-				{
-					leaf._children = new Node[16];
-				}
-				if (leaf._children[index] == null)
-				{
-					leaf._children[index] = new Node();
-				}
+				leaf._children ??= new Node[16];
+				leaf._children[index] ??= new Node();
 				leaf = leaf._children[index]!;
 			}
 		}
diff --git a/Engine/Source/Programs/Horde/Horde.Build/Utilities/StringId.cs b/Engine/Source/Programs/Horde/Horde.Build/Utilities/StringId.cs
--- a/Engine/Source/Programs/Horde/Horde.Build/Utilities/StringId.cs
+++ b/Engine/Source/Programs/Horde/Horde.Build/Utilities/StringId.cs
@@ -198,6 +198,26 @@ namespace Horde.Build.Utilities
 		{
 			return !left.Equals(right);
 		}
+
+		/// <summary>
+		/// Compares two stringids
+		/// </summary>
+		public static bool operator <(StringId<T> left, StringId<T> right) => left.CompareTo(right) < 0;
+
+		/// <summary>
+		/// Compares two stringids
+		/// </summary>
+		public static bool operator <=(StringId<T> left, StringId<T> right) => left.CompareTo(right) <= 0;
+
+		/// <summary>
+		/// Compares two stringids
+		/// </summary>
+		public static bool operator >(StringId<T> left, StringId<T> right) => left.CompareTo(right) > 0;
+
+		/// <summary>
+		/// Compares two stringids
+		/// </summary>
+		public static bool operator >=(StringId<T> left, StringId<T> right) => left.CompareTo(right) >= 0;
 	}
 
 	/// <summary>
