Horde: Missing file.

diff --git a/Engine/Source/Programs/Horde/Horde.Build/Issues/SuspectChange.cs b/Engine/Source/Programs/Horde/Horde.Build/Issues/SuspectChange.cs
--- a/Engine/Source/Programs/Horde/Horde.Build/Issues/SuspectChange.cs
+++ b/Engine/Source/Programs/Horde/Horde.Build/Issues/SuspectChange.cs
@@ -49,7 +49,12 @@ namespace Horde.Build.Issues
 		/// <summary>
 		/// The change detials
 		/// </summary>
-		public ChangeDetails Details { get; set; }
+		public ICommit Details { get; set; }
+
+		/// <summary>
+		/// Files in this change
+		/// </summary>
+		public IReadOnlyList<string> Files { get; }
 
 		/// <summary>
 		/// Rank for the change. A value less than or equal to zero indicates a lack of culpability, positive values indicate
@@ -71,16 +76,18 @@ namespace Horde.Build.Issues
 		/// Constructor
 		/// </summary>
 		/// <param name="details">The changelist details</param>
-		public SuspectChange(ChangeDetails details)
+		/// <param name="files">Files modified by the change</param>
+		public SuspectChange(ICommit details, IReadOnlyList<string> files)
 		{
 			Details = details;
+			Files = files;
 
-			foreach (ChangeFile file in details.Files)
+			foreach (string file in files)
 			{
-				int idx = file.Path.LastIndexOf('.');
+				int idx = file.LastIndexOf('.');
 				if (idx != -1)
 				{
-					StringView extension = new StringView(file.Path, idx);
+					StringView extension = new StringView(file, idx);
 					if (s_codeExtensions.Contains(extension))
 					{
 						ContainsCode = true;
@@ -104,9 +111,9 @@ namespace Horde.Build.Issues
 		/// <returns>True if the change modifies the given file</returns>
 		public bool ModifiesFile(string fileToCheck)
 		{
-			foreach (ChangeFile file in Details.Files)
+			foreach (string file in Files)
 			{
-				if (file.Path.EndsWith(fileToCheck, StringComparison.OrdinalIgnoreCase) && (file.Length == fileToCheck.Length || file.Path[file.Path.Length - fileToCheck.Length - 1] == '/'))
+				if (file.EndsWith(fileToCheck, StringComparison.OrdinalIgnoreCase) && (file.Length == fileToCheck.Length || file[file.Length - fileToCheck.Length - 1] == '/'))
 				{
 					return true;
 				}
