Horde: Move config classes for agents and workspaces next to StreamConfig, and remove request/response naming convention now that config is data driven.

diff --git a/Engine/Source/Programs/Horde/Horde.Build.Tests/Fixture.cs b/Engine/Source/Programs/Horde/Horde.Build.Tests/Fixture.cs
--- a/Engine/Source/Programs/Horde/Horde.Build.Tests/Fixture.cs
+++ b/Engine/Source/Programs/Horde/Horde.Build.Tests/Fixture.cs
@@ -72,7 +72,7 @@ namespace Horde.Build.Tests
 			List<CreateStreamTabRequest> tabs = new List<CreateStreamTabRequest>();
 			tabs.Add(new CreateJobsTabRequest { Title = "foo", Templates = new List<TemplateRefId> { TemplateRefId1, TemplateRefId2 } });
 
-			Dictionary<string, CreateAgentTypeRequest> agentTypes = new()
+			Dictionary<string, AgentConfig> agentTypes = new()
 			{
 				{ "Win64", new() { Pool = PoolName} }
 			};
diff --git a/Engine/Source/Programs/Horde/Horde.Build.Tests/Fleet/JobQueueStrategyTest.cs b/Engine/Source/Programs/Horde/Horde.Build.Tests/Fleet/JobQueueStrategyTest.cs
--- a/Engine/Source/Programs/Horde/Horde.Build.Tests/Fleet/JobQueueStrategyTest.cs
+++ b/Engine/Source/Programs/Horde/Horde.Build.Tests/Fleet/JobQueueStrategyTest.cs
@@ -113,7 +113,7 @@ namespace Horde.Build.Tests.Fleet
 			PoolSizeData poolSize = new (pool1, agents, null);
 			
 			string agentTypeName1 = "bogusAgentType1";
-			Dictionary<string, CreateAgentTypeRequest> agentTypes = new() { {agentTypeName1, new() { Pool = pool1.Name} }, };
+			Dictionary<string, AgentConfig> agentTypes = new() { {agentTypeName1, new() { Pool = pool1.Name} }, };
 
 			IStream stream = (await CreateOrReplaceStreamAsync(
 				new StreamId("ue5-main"),
diff --git a/Engine/Source/Programs/Horde/Horde.Build.Tests/JobServiceTests.cs b/Engine/Source/Programs/Horde/Horde.Build.Tests/JobServiceTests.cs
--- a/Engine/Source/Programs/Horde/Horde.Build.Tests/JobServiceTests.cs
+++ b/Engine/Source/Programs/Horde/Horde.Build.Tests/JobServiceTests.cs
@@ -44,7 +44,7 @@ namespace Horde.Build.Tests
 			TemplateRefId templateRefId2 = new TemplateRefId("template2");
 
 			StreamConfig streamConfig = new StreamConfig();
-			streamConfig.Templates.Add(new TemplateRefConfig { Id = templateRefId1, Name = "Test Template", ChainedJobs = new List<CreateChainedJobTemplateRequest> { new CreateChainedJobTemplateRequest { TemplateId = templateRefId2.ToString(), Trigger = "Setup Build" } } });
+			streamConfig.Templates.Add(new TemplateRefConfig { Id = templateRefId1, Name = "Test Template", ChainedJobs = new List<ChainedJobTemplateConfig> { new ChainedJobTemplateConfig { TemplateId = templateRefId2.ToString(), Trigger = "Setup Build" } } });
 			streamConfig.Templates.Add(new TemplateRefConfig { Id = templateRefId2, Name = "Test Template" });
 			streamConfig.Tabs.Add(new CreateJobsTabRequest { Title = "foo", Templates = new List<TemplateRefId> { templateRefId1, templateRefId2 } });
 
diff --git a/Engine/Source/Programs/Horde/Horde.Build.Tests/SchedulerTests.cs b/Engine/Source/Programs/Horde/Horde.Build.Tests/SchedulerTests.cs
--- a/Engine/Source/Programs/Horde/Horde.Build.Tests/SchedulerTests.cs
+++ b/Engine/Source/Programs/Horde/Horde.Build.Tests/SchedulerTests.cs
@@ -54,7 +54,7 @@ namespace Horde.Build.Tests
 			PerforceService.AddChange(StreamId, 102, bob, "", new[] { "content.uasset" });
 		}
 
-		async Task<IStream> SetScheduleAsync(CreateScheduleRequest schedule)
+		async Task<IStream> SetScheduleAsync(ScheduleConfig schedule)
 		{
 			await ScheduleService.ResetAsync();
 
@@ -84,10 +84,10 @@ namespace Horde.Build.Tests
 			DateTime startTime = new DateTime(2021, 1, 1, 12, 0, 0, DateTimeKind.Local); // Friday Jan 1, 2021 
 			Clock.UtcNow = startTime;
 
-			CreateScheduleRequest schedule = new CreateScheduleRequest();
+			ScheduleConfig schedule = new ScheduleConfig();
 			schedule.Enabled = true;
 			schedule.MaxChanges = 10;
-			schedule.Patterns.Add(new CreateSchedulePatternRequest { Interval = 1 });
+			schedule.Patterns.Add(new SchedulePatternConfig { Interval = 1 });
 			schedule.Files = files.ToList();
 			await SetScheduleAsync(schedule);
 
@@ -199,9 +199,9 @@ namespace Horde.Build.Tests
 			DateTime startTime = new DateTime(2021, 1, 1, 12, 0, 0, DateTimeKind.Local); // Friday Jan 1, 2021 
 			Clock.UtcNow = startTime;
 
-			CreateScheduleRequest schedule = new CreateScheduleRequest();
+			ScheduleConfig schedule = new ScheduleConfig();
 			schedule.Enabled = true;
-			schedule.Patterns.Add(new CreateSchedulePatternRequest { MinTime = 13 * 60, MaxTime = 14 * 60, Interval = 15 });
+			schedule.Patterns.Add(new SchedulePatternConfig { MinTime = 13 * 60, MaxTime = 14 * 60, Interval = 15 });
 //			Schedule.LastTriggerTime = StartTime;
 			await SetScheduleAsync(schedule);
 
@@ -227,9 +227,9 @@ namespace Horde.Build.Tests
 			DateTime startTime = new DateTime(2021, 1, 1, 12, 0, 0, DateTimeKind.Local); // Friday Jan 1, 2021 
 			Clock.UtcNow = startTime;
 
-			CreateScheduleRequest schedule = new CreateScheduleRequest();
+			ScheduleConfig schedule = new ScheduleConfig();
 			schedule.Enabled = true;
-			schedule.Patterns.Add(new CreateSchedulePatternRequest { MinTime = 13 * 60, MaxTime = 14 * 60, Interval = 15 });
+			schedule.Patterns.Add(new SchedulePatternConfig { MinTime = 13 * 60, MaxTime = 14 * 60, Interval = 15 });
 			IStream stream = await SetScheduleAsync(schedule);
 
 			// Initial tick
@@ -266,9 +266,9 @@ namespace Horde.Build.Tests
 			DateTime startTime = new DateTime(2021, 1, 1, 12, 0, 0, DateTimeKind.Local); // Friday Jan 1, 2021 
 			Clock.UtcNow = startTime;
 
-			CreateScheduleRequest schedule = new CreateScheduleRequest();
+			ScheduleConfig schedule = new ScheduleConfig();
 			schedule.Enabled = true;
-			schedule.Patterns.Add(new CreateSchedulePatternRequest { MinTime = 13 * 60, MaxTime = 14 * 60, Interval = 15 });
+			schedule.Patterns.Add(new SchedulePatternConfig { MinTime = 13 * 60, MaxTime = 14 * 60, Interval = 15 });
 			schedule.MaxChanges = 2;
 			schedule.Filter = new List<ChangeContentFlags> { ChangeContentFlags.ContainsCode };
 			await SetScheduleAsync(schedule);
@@ -303,9 +303,9 @@ namespace Horde.Build.Tests
 			DateTime startTime = new DateTime(2021, 1, 1, 12, 0, 0, DateTimeKind.Local); // Friday Jan 1, 2021 
 			Clock.UtcNow = startTime;
 
-			CreateScheduleRequest schedule = new CreateScheduleRequest();
+			ScheduleConfig schedule = new ScheduleConfig();
 			schedule.Enabled = true;
-			schedule.Patterns.Add(new CreateSchedulePatternRequest { MinTime = 13 * 60, MaxTime = 14 * 60, Interval = 15 });
+			schedule.Patterns.Add(new SchedulePatternConfig { MinTime = 13 * 60, MaxTime = 14 * 60, Interval = 15 });
 			schedule.MaxChanges = 2;
 			schedule.Filter = new List<ChangeContentFlags> { ChangeContentFlags.ContainsCode };
 			await SetScheduleAsync(schedule);
@@ -340,10 +340,10 @@ namespace Horde.Build.Tests
 			DateTime startTime = new DateTime(2021, 1, 1, 12, 0, 0, DateTimeKind.Local); // Friday Jan 1, 2021 
 			Clock.UtcNow = startTime;
 
-			CreateScheduleRequest schedule = new CreateScheduleRequest();
+			ScheduleConfig schedule = new ScheduleConfig();
 			schedule.Enabled = true;
 			schedule.RequireSubmittedChange = false;
-			schedule.Patterns.Add(new CreateSchedulePatternRequest { MinTime = 13 * 60, MaxTime = 14 * 60, Interval = 15 });
+			schedule.Patterns.Add(new SchedulePatternConfig { MinTime = 13 * 60, MaxTime = 14 * 60, Interval = 15 });
 			schedule.MaxActive = 1;
 			await SetScheduleAsync(schedule);
 
@@ -382,9 +382,9 @@ namespace Horde.Build.Tests
 			DateTime startTime = new DateTime(2021, 1, 1, 12, 0, 0, DateTimeKind.Local); // Friday Jan 1, 2021 
 			Clock.UtcNow = startTime;
 
-			CreateScheduleRequest schedule = new CreateScheduleRequest();
+			ScheduleConfig schedule = new ScheduleConfig();
 			schedule.Enabled = true;
-			schedule.Patterns.Add(new CreateSchedulePatternRequest { MinTime = 13 * 60, MaxTime = 14 * 60, Interval = 15 });
+			schedule.Patterns.Add(new SchedulePatternConfig { MinTime = 13 * 60, MaxTime = 14 * 60, Interval = 15 });
 			/*IStream stream = */await SetScheduleAsync(schedule);
 
 			// Trigger a job
@@ -496,11 +496,11 @@ namespace Horde.Build.Tests
 			TemplateRefConfig newTemplate2 = new TemplateRefConfig();
 			newTemplate2.Id = newTemplateRefId2;
 			newTemplate2.Name = "Test template 2";
-			newTemplate2.Schedule = new CreateScheduleRequest();
+			newTemplate2.Schedule = new ScheduleConfig();
 			newTemplate2.Schedule.MaxChanges = 4;
 			newTemplate2.Schedule.Filter = new List<ChangeContentFlags> { ChangeContentFlags.ContainsCode };
-			newTemplate2.Schedule.Gate = new CreateScheduleGateRequest { TemplateId = newTemplateRefId1.ToString(), Target = "TriggerNext" };
-			newTemplate2.Schedule.Patterns.Add(new CreateSchedulePatternRequest { Interval = 10 });// (null, 0, null, 10));
+			newTemplate2.Schedule.Gate = new ScheduleGateConfig { TemplateId = newTemplateRefId1.ToString(), Target = "TriggerNext" };
+			newTemplate2.Schedule.Patterns.Add(new SchedulePatternConfig { Interval = 10 });// (null, 0, null, 10));
 //			NewTemplate2.Schedule.LastTriggerTime = StartTime;
 
 			IStream? stream = await StreamService.GetStreamAsync(StreamId);
@@ -558,10 +558,10 @@ namespace Horde.Build.Tests
 			DateTime startTime = new DateTime(2021, 1, 1, 12, 0, 0, DateTimeKind.Local); // Friday Jan 1, 2021 
 			Clock.UtcNow = startTime;
 
-			CreateScheduleRequest schedule = new CreateScheduleRequest();
+			ScheduleConfig schedule = new ScheduleConfig();
 			schedule.Enabled = true;
 			schedule.RequireSubmittedChange = false;
-			schedule.Patterns.Add(new CreateSchedulePatternRequest { MinTime = 13 * 60, MaxTime = 14 * 60, Interval = 15 });
+			schedule.Patterns.Add(new SchedulePatternConfig { MinTime = 13 * 60, MaxTime = 14 * 60, Interval = 15 });
 			schedule.MaxActive = 2;
 			await SetScheduleAsync(schedule);
 
@@ -595,9 +595,9 @@ namespace Horde.Build.Tests
 			DateTime startTime = new DateTime(2021, 1, 1, 12, 0, 0, DateTimeKind.Utc); // Friday Jan 1, 2021 
 			Clock.UtcNow = startTime;
 
-			CreateScheduleRequest schedule = new CreateScheduleRequest();
+			ScheduleConfig schedule = new ScheduleConfig();
 			schedule.Enabled = true;
-			schedule.Patterns.Add(new CreateSchedulePatternRequest { MinTime = 13 * 60, MaxTime = 14 * 60, Interval = 15 });
+			schedule.Patterns.Add(new SchedulePatternConfig { MinTime = 13 * 60, MaxTime = 14 * 60, Interval = 15 });
 			IStream stream = await SetScheduleAsync(schedule);
 
 			IStreamCollection streamCollection = ServiceProvider.GetRequiredService<IStreamCollection>();
diff --git a/Engine/Source/Programs/Horde/Horde.Build/Acls/Acl.cs b/Engine/Source/Programs/Horde/Horde.Build/Acls/Acl.cs
--- a/Engine/Source/Programs/Horde/Horde.Build/Acls/Acl.cs
+++ b/Engine/Source/Programs/Horde/Horde.Build/Acls/Acl.cs
@@ -546,7 +546,7 @@ namespace Horde.Build.Acls
 		/// Constructs a claim from a request object
 		/// </summary>
 		/// <param name="request">The request object</param>
-		public AclClaim(CreateAclClaimRequest request)
+		public AclClaim(AclClaimConfig request)
 			: this(request.Type, request.Value)
 		{
 		}
@@ -594,7 +594,7 @@ namespace Horde.Build.Acls
 		/// </summary>
 		/// <param name="request">Request instance</param>
 		/// <returns>New ACL entry</returns>
-		public static AclEntry FromRequest(CreateAclEntryRequest request)
+		public static AclEntry FromRequest(AclEntryConfig request)
 		{
 			return new AclEntry(new AclClaim(request.Claim.Type, request.Claim.Value), ParseActionNames(request.Actions).ToArray());
 		}
@@ -726,7 +726,7 @@ namespace Horde.Build.Acls
 		/// <param name="baseAcl">The current acl</param>
 		/// <param name="update">The update to apply</param>
 		/// <returns>The new ACL value. Null if the ACL has all default settings.</returns>
-		public static Acl? Merge(Acl? baseAcl, UpdateAclRequest? update)
+		public static Acl? Merge(Acl? baseAcl, AclConfig? update)
 		{
 			Acl? newAcl = null;
 			if(update != null)
diff --git a/Engine/Source/Programs/Horde/Horde.Build/Acls/AclMessages.cs b/Engine/Source/Programs/Horde/Horde.Build/Acls/AclMessages.cs
--- a/Engine/Source/Programs/Horde/Horde.Build/Acls/AclMessages.cs
+++ b/Engine/Source/Programs/Horde/Horde.Build/Acls/AclMessages.cs
@@ -5,63 +5,6 @@ using System.ComponentModel.DataAnnotations;
 
 namespace Horde.Build.Acls
 {
-	/// <summary>
-	/// New claim to create
-	/// </summary>
-	public class CreateAclClaimRequest
-	{
-		/// <summary>
-		/// The claim type
-		/// </summary>
-		[Required]
-		public string Type { get; set; } = null!;
-
-		/// <summary>
-		/// The claim value
-		/// </summary>
-		[Required]
-		public string Value { get; set; } = null!;
-	}
-
-	/// <summary>
-	/// Individual entry in an ACL
-	/// </summary>
-	public class CreateAclEntryRequest
-	{
-		/// <summary>
-		/// Name of the user or group
-		/// </summary>
-		[Required]
-		public CreateAclClaimRequest Claim { get; set; } = null!;
-
-		/// <summary>
-		/// Array of actions to allow
-		/// </summary>
-		[Required]
-		public string[] Actions { get; set; } = null!;
-	}
-
-	/// <summary>
-	/// Parameters to update an ACL
-	/// </summary>
-	public class UpdateAclRequest
-	{
-		/// <summary>
-		/// Entries to replace the existing ACL
-		/// </summary>
-		public List<CreateAclEntryRequest>? Entries { get; set; }
-
-		/// <summary>
-		/// Whether to inherit permissions from the parent ACL
-		/// </summary>
-		public bool? Inherit { get; set; }
-
-		/// <summary>
-		/// List of exceptions to the inherited setting
-		/// </summary>
-		public List<string>? Exceptions { get; set; }
-	}
-
 	/// <summary>
 	/// New claim to update
 	/// </summary>
diff --git a/Engine/Source/Programs/Horde/Horde.Build/Agents/AgentMessages.cs b/Engine/Source/Programs/Horde/Horde.Build/Agents/AgentMessages.cs
--- a/Engine/Source/Programs/Horde/Horde.Build/Agents/AgentMessages.cs
+++ b/Engine/Source/Programs/Horde/Horde.Build/Agents/AgentMessages.cs
@@ -112,7 +112,7 @@ namespace Horde.Build.Agents
 		/// <summary>
 		/// New ACL for this agent
 		/// </summary>
-		public UpdateAclRequest? Acl { get; set; }
+		public AclConfig? Acl { get; set; }
 		
 		/// <summary>
 		/// New comment
diff --git a/Engine/Source/Programs/Horde/Horde.Build/Jobs/JobMessages.cs b/Engine/Source/Programs/Horde/Horde.Build/Jobs/JobMessages.cs
--- a/Engine/Source/Programs/Horde/Horde.Build/Jobs/JobMessages.cs
+++ b/Engine/Source/Programs/Horde/Horde.Build/Jobs/JobMessages.cs
@@ -65,7 +65,7 @@ namespace Horde.Build.Jobs
 		/// <summary>
 		/// Parameters to use when selecting the change to execute at.
 		/// </summary>
-		public ChangeQueryRequest? ChangeQuery { get; set; }
+		public ChangeQueryConfig? ChangeQuery { get; set; }
 
 		/// <summary>
 		/// The preflight changelist number
@@ -155,7 +155,7 @@ namespace Horde.Build.Jobs
 		/// <summary>
 		/// Custom permissions for this object
 		/// </summary>
-		public UpdateAclRequest? Acl { get; set; }
+		public AclConfig? Acl { get; set; }
 	}
 
 	/// <summary>
diff --git a/Engine/Source/Programs/Horde/Horde.Build/Jobs/Schedules/Schedules.cs b/Engine/Source/Programs/Horde/Horde.Build/Jobs/Schedules/Schedules.cs
--- a/Engine/Source/Programs/Horde/Horde.Build/Jobs/Schedules/Schedules.cs
+++ b/Engine/Source/Programs/Horde/Horde.Build/Jobs/Schedules/Schedules.cs
@@ -12,134 +12,6 @@ namespace Horde.Build.Jobs.Schedules
 {
 	using TemplateRefId = StringId<TemplateRef>;
 
-	/// <summary>
-	/// Parameters to create a new schedule
-	/// </summary>
-	public class CreateSchedulePatternRequest
-	{
-		/// <summary>
-		/// Days of the week to run this schedule on. If null, the schedule will run every day.
-		/// </summary>
-		public List<string>? DaysOfWeek { get; set; }
-
-		/// <summary>
-		/// Time during the day for the first schedule to trigger. Measured in minutes from midnight.
-		/// </summary>
-		public int MinTime { get; set; }
-
-		/// <summary>
-		/// Time during the day for the last schedule to trigger. Measured in minutes from midnight.
-		/// </summary>
-		public int? MaxTime { get; set; }
-
-		/// <summary>
-		/// Interval between each schedule triggering
-		/// </summary>
-		public int? Interval { get; set; }
-
-		/// <summary>
-		/// Constructs a model object from this request
-		/// </summary>
-		/// <returns>Model object</returns>
-		public SchedulePattern ToModel()
-		{
-			return new SchedulePattern(DaysOfWeek?.ConvertAll(x => Enum.Parse<DayOfWeek>(x)), MinTime, MaxTime, Interval);
-		}
-	}
-
-	/// <summary>
-	/// Gate allowing a schedule to trigger.
-	/// </summary>
-	public class CreateScheduleGateRequest
-	{
-		/// <summary>
-		/// The template containing the dependency
-		/// </summary>
-		[Required]
-		public string TemplateId { get; set; } = String.Empty;
-
-		/// <summary>
-		/// Target to wait for
-		/// </summary>
-		[Required]
-		public string Target { get; set; } = String.Empty;
-
-		/// <summary>
-		/// Constructs a model object
-		/// </summary>
-		/// <returns>New model object.</returns>
-		public ScheduleGate ToModel()
-		{
-			return new ScheduleGate(new TemplateRefId(TemplateId), Target);
-		}
-	}
-
-	/// <summary>
-	/// Parameters to create a new schedule
-	/// </summary>
-	public class CreateScheduleRequest
-	{
-		/// <summary>
-		/// Roles to impersonate for this schedule
-		/// </summary>
-		public List<CreateAclClaimRequest>? Claims { get; set; }
-
-		/// <summary>
-		/// Whether the schedule should be enabled
-		/// </summary>
-		public bool Enabled { get; set; }
-
-		/// <summary>
-		/// Maximum number of builds that can be active at once
-		/// </summary>
-		public int MaxActive { get; set; }
-
-		/// <summary>
-		/// Maximum number of changes the schedule can fall behind head revision. If greater than zero, builds will be triggered for every submitted changelist until the backlog is this size.
-		/// </summary>
-		public int MaxChanges { get; set; }
-
-		/// <summary>
-		/// Whether the build requires a change to be submitted
-		/// </summary>
-		public bool RequireSubmittedChange { get; set; } = true;
-
-		/// <summary>
-		/// Gate allowing the schedule to trigger
-		/// </summary>
-		public CreateScheduleGateRequest? Gate { get; set; }
-
-		/// <summary>
-		/// The types of changes to run for
-		/// </summary>
-		public List<ChangeContentFlags>? Filter { get; set; }
-
-		/// <summary>
-		/// Files that should cause the job to trigger
-		/// </summary>
-		public List<string>? Files { get; set; }
-
-		/// <summary>
-		/// Parameters for the template
-		/// </summary>
-		public Dictionary<string, string> TemplateParameters { get; set; } = new Dictionary<string, string>();
-
-		/// <summary>
-		/// New patterns for the schedule
-		/// </summary>
-		public List<CreateSchedulePatternRequest> Patterns { get; set; } = new List<CreateSchedulePatternRequest>();
-
-		/// <summary>
-		/// Constructs a model object
-		/// </summary>
-		/// <param name="currentTimeUtc">The current time</param>
-		/// <returns>New model object</returns>
-		public Schedule ToModel(DateTime currentTimeUtc)
-		{
-			return new Schedule(currentTimeUtc, Enabled, MaxActive, MaxChanges, RequireSubmittedChange, Gate?.ToModel(), Filter, Files, TemplateParameters, Patterns.ConvertAll(x => x.ToModel()));
-		}
-	}
-
 	/// <summary>
 	/// Information about a schedule pattern
 	/// </summary>
diff --git a/Engine/Source/Programs/Horde/Horde.Build/Projects/ProjectConfig.cs b/Engine/Source/Programs/Horde/Horde.Build/Projects/ProjectConfig.cs
--- a/Engine/Source/Programs/Horde/Horde.Build/Projects/ProjectConfig.cs
+++ b/Engine/Source/Programs/Horde/Horde.Build/Projects/ProjectConfig.cs
@@ -42,7 +42,7 @@ namespace Horde.Build.Projects
 		/// <summary>
 		/// Acl entries
 		/// </summary>
-		public UpdateAclRequest? Acl { get; set; }
+		public AclConfig? Acl { get; set; }
 	}
 
 	/// <summary>
diff --git a/Engine/Source/Programs/Horde/Horde.Build/Secrets/CredentialMessages.cs b/Engine/Source/Programs/Horde/Horde.Build/Secrets/CredentialMessages.cs
--- a/Engine/Source/Programs/Horde/Horde.Build/Secrets/CredentialMessages.cs
+++ b/Engine/Source/Programs/Horde/Horde.Build/Secrets/CredentialMessages.cs
@@ -62,7 +62,7 @@ namespace Horde.Build.Credentials
 		/// <summary>
 		/// Custom permissions for this object
 		/// </summary>
-		public UpdateAclRequest? Acl { get; set; }
+		public AclConfig? Acl { get; set; }
 	}
 
 	/// <summary>
diff --git a/Engine/Source/Programs/Horde/Horde.Build/Server/GlobalConfig.cs b/Engine/Source/Programs/Horde/Horde.Build/Server/GlobalConfig.cs
--- a/Engine/Source/Programs/Horde/Horde.Build/Server/GlobalConfig.cs
+++ b/Engine/Source/Programs/Horde/Horde.Build/Server/GlobalConfig.cs
@@ -65,7 +65,7 @@ namespace Horde.Build.Server
 		/// <summary>
 		/// Access control list
 		/// </summary>
-		public UpdateAclRequest? Acl { get; set; }
+		public AclConfig? Acl { get; set; }
 	}
 
 	/// <summary>
@@ -96,7 +96,7 @@ namespace Horde.Build.Server
 		/// <summary>
 		/// Access control list
 		/// </summary>
-		public UpdateAclRequest? Acl { get; set; }
+		public AclConfig? Acl { get; set; }
 	}
 
 	/// <summary>
@@ -146,7 +146,7 @@ namespace Horde.Build.Server
 		/// <summary>
 		/// Access control for this namespace
 		/// </summary>
-		public UpdateAclRequest? Acl { get; set; }
+		public AclConfig? Acl { get; set; }
 	}
 
 	/// <summary>
diff --git a/Engine/Source/Programs/Horde/Horde.Build/Streams/IStream.cs b/Engine/Source/Programs/Horde/Horde.Build/Streams/IStream.cs
--- a/Engine/Source/Programs/Horde/Horde.Build/Streams/IStream.cs
+++ b/Engine/Source/Programs/Horde/Horde.Build/Streams/IStream.cs
@@ -97,7 +97,7 @@ namespace Horde.Build.Streams
 		/// Constructor
 		/// </summary>
 		/// <param name="request">The object to construct from</param>
-		public AgentType(CreateAgentTypeRequest request)
+		public AgentType(AgentConfig request)
 		{
 			Pool = new PoolId(request.Pool);
 			Workspace = request.Workspace;
@@ -111,7 +111,7 @@ namespace Horde.Build.Streams
 		/// <param name="request">The request object</param>
 		/// <returns>New agent type object</returns>
 		[return: NotNullIfNotNull("request")]
-		public static AgentType? FromRequest(CreateAgentTypeRequest? request)
+		public static AgentType? FromRequest(AgentConfig? request)
 		{
 			return (request != null) ? new AgentType(request) : null;
 		}
@@ -212,7 +212,7 @@ namespace Horde.Build.Streams
 		/// Constructor
 		/// </summary>
 		/// <param name="request">The object to construct from</param>
-		public WorkspaceType(CreateWorkspaceTypeRequest request)
+		public WorkspaceType(WorkspaceConfig request)
 		{
 			Cluster = request.Cluster;
 			ServerAndPort = request.ServerAndPort;
@@ -231,7 +231,7 @@ namespace Horde.Build.Streams
 		/// <param name="request">The request object</param>
 		/// <returns>New agent type object</returns>
 		[return: NotNullIfNotNull("request")]
-		public static WorkspaceType? FromRequest(CreateWorkspaceTypeRequest? request)
+		public static WorkspaceType? FromRequest(WorkspaceConfig? request)
 		{
 			return (request != null) ? new WorkspaceType(request) : null;
 		}
@@ -283,7 +283,7 @@ namespace Horde.Build.Streams
 		/// Constructor
 		/// </summary>
 		/// <param name="request">Request to construct from</param>
-		public ChainedJobTemplate(CreateChainedJobTemplateRequest request)
+		public ChainedJobTemplate(ChainedJobTemplateConfig request)
 			: this(request.Trigger, new TemplateRefId(request.TemplateId))
 		{
 		}
@@ -452,9 +452,9 @@ namespace Horde.Build.Streams
 		/// Convert to a request object
 		/// </summary>
 		/// <returns></returns>
-		public ChangeQueryRequest ToRequest()
+		public ChangeQueryConfig ToRequest()
 		{
-			return new ChangeQueryRequest { TemplateId = TemplateRefId?.ToString(), Target = Target, Outcomes = Outcomes };
+			return new ChangeQueryConfig { TemplateId = TemplateRefId?.ToString(), Target = Target, Outcomes = Outcomes };
 		}
 	}
 
@@ -494,20 +494,20 @@ namespace Horde.Build.Streams
 		/// Convert to a request object
 		/// </summary>
 		/// <returns></returns>
-		public DefaultPreflightRequest ToRequest()
+		public DefaultPreflightConfig ToRequest()
 		{
 #pragma warning disable CS0618 // Type or member is obsolete
-			ChangeQueryRequest? changeRequest = null;
+			ChangeQueryConfig? changeRequest = null;
 			if (Change != null)
 			{
 				changeRequest = Change.ToRequest();
 			}
 			else if (ChangeTemplateRefId != null)
 			{
-				changeRequest = new ChangeQueryRequest { TemplateId = ChangeTemplateRefId.ToString() };
+				changeRequest = new ChangeQueryConfig { TemplateId = ChangeTemplateRefId.ToString() };
 			}
 
-			return new DefaultPreflightRequest { TemplateId = TemplateRefId?.ToString(), Change = changeRequest, ChangeTemplateId = changeRequest?.TemplateId };
+			return new DefaultPreflightConfig { TemplateId = TemplateRefId?.ToString(), Change = changeRequest, ChangeTemplateId = changeRequest?.TemplateId };
 #pragma warning restore CS0618 // Type or member is obsolete
 		}
 	}
diff --git a/Engine/Source/Programs/Horde/Horde.Build/Streams/StreamConfig.cs b/Engine/Source/Programs/Horde/Horde.Build/Streams/StreamConfig.cs
--- a/Engine/Source/Programs/Horde/Horde.Build/Streams/StreamConfig.cs
+++ b/Engine/Source/Programs/Horde/Horde.Build/Streams/StreamConfig.cs
@@ -8,7 +8,12 @@ using System.Linq;
 using EpicGames.Core;
 using Horde.Build.Acls;
 using Horde.Build.Issues;
+using Horde.Build.Jobs.Graphs;
+using Horde.Build.Jobs.Schedules;
+using Horde.Build.Jobs.Templates;
+using Horde.Build.Perforce;
 using Horde.Build.Utilities;
+using HordeCommon;
 
 namespace Horde.Build.Streams
 {
@@ -61,7 +66,7 @@ namespace Horde.Build.Streams
 		/// <summary>
 		/// Default template for running preflights
 		/// </summary>
-		public DefaultPreflightRequest? DefaultPreflight { get; set; }
+		public DefaultPreflightConfig? DefaultPreflight { get; set; }
 
 		/// <summary>
 		/// List of tabs to show for the new stream
@@ -71,12 +76,12 @@ namespace Horde.Build.Streams
 		/// <summary>
 		/// Map of agent name to type
 		/// </summary>
-		public Dictionary<string, CreateAgentTypeRequest> AgentTypes { get; set; } = new Dictionary<string, CreateAgentTypeRequest>();
+		public Dictionary<string, AgentConfig> AgentTypes { get; set; } = new Dictionary<string, AgentConfig>();
 
 		/// <summary>
 		/// Map of workspace name to type
 		/// </summary>
-		public Dictionary<string, CreateWorkspaceTypeRequest> WorkspaceTypes { get; set; } = new Dictionary<string, CreateWorkspaceTypeRequest>();
+		public Dictionary<string, WorkspaceConfig> WorkspaceTypes { get; set; } = new Dictionary<string, WorkspaceConfig>();
 
 		/// <summary>
 		/// List of templates to create
@@ -86,7 +91,7 @@ namespace Horde.Build.Streams
 		/// <summary>
 		/// Custom permissions for this object
 		/// </summary>
-		public UpdateAclRequest? Acl { get; set; }
+		public AclConfig? Acl { get; set; }
 
 		/// <summary>
 		/// Pause stream builds until specified date
@@ -148,6 +153,376 @@ namespace Horde.Build.Streams
 		}
 	}
 
+	/// <summary>
+	/// Mapping from a BuildGraph agent type to a set of machines on the farm
+	/// </summary>
+	public class AgentConfig
+	{
+		/// <summary>
+		/// Pool of agents to use for this agent type
+		/// </summary>
+		[Required]
+		public string Pool { get; set; } = null!;
+
+		/// <summary>
+		/// Name of the workspace to sync
+		/// </summary>
+		public string? Workspace { get; set; }
+
+		/// <summary>
+		/// Path to the temporary storage dir
+		/// </summary>
+		public string? TempStorageDir { get; set; }
+
+		/// <summary>
+		/// Environment variables to be set when executing the job
+		/// </summary>
+		public Dictionary<string, string>? Environment { get; set; }
+	}
+
+	/// <summary>
+	/// Information about a workspace type
+	/// </summary>
+	public class WorkspaceConfig
+	{
+		/// <summary>
+		/// Name of the Perforce server cluster to use
+		/// </summary>
+		public string? Cluster { get; set; }
+
+		/// <summary>
+		/// The Perforce server and port (eg. perforce:1666)
+		/// </summary>
+		public string? ServerAndPort { get; set; }
+
+		/// <summary>
+		/// User to log into Perforce with (defaults to buildmachine)
+		/// </summary>
+		public string? UserName { get; set; }
+
+		/// <summary>
+		/// Password to use to log into the workspace
+		/// </summary>
+		public string? Password { get; set; }
+
+		/// <summary>
+		/// Identifier to distinguish this workspace from other workspaces. Defaults to the workspace type name.
+		/// </summary>
+		public string? Identifier { get; set; }
+
+		/// <summary>
+		/// Override for the stream to sync
+		/// </summary>
+		public string? Stream { get; set; }
+
+		/// <summary>
+		/// Custom view for the workspace
+		/// </summary>
+		public List<string>? View { get; set; }
+
+		/// <summary>
+		/// Whether to use an incrementally synced workspace
+		/// </summary>
+		public bool Incremental { get; set; }
+
+		/// <summary>
+		/// Whether to use the AutoSDK
+		/// </summary>
+		public bool UseAutoSdk { get; set; } = true;
+	}
+
+	/// <summary>
+	/// Query selecting the base changelist to use
+	/// </summary>
+	public class ChangeQueryConfig
+	{
+		/// <summary>
+		/// The template id to query
+		/// </summary>
+		public string? TemplateId { get; set; }
+
+		/// <summary>
+		/// The target to query
+		/// </summary>
+		public string? Target { get; set; }
+
+		/// <summary>
+		/// Whether to match a job that produced warnings
+		/// </summary>
+		public List<JobStepOutcome>? Outcomes { get; set; }
+
+		/// <summary>
+		/// Convert to a model object
+		/// </summary>
+		/// <returns></returns>
+		public ChangeQuery ToModel()
+		{
+			ChangeQuery query = new ChangeQuery();
+			if (TemplateId != null)
+			{
+				query.TemplateRefId = new StringId<TemplateRef>(TemplateId);
+			}
+			query.Target = Target;
+			query.Outcomes = Outcomes;
+			return query;
+		}
+	}
+
+	/// <summary>
+	/// Specifies defaults for running a preflight
+	/// </summary>
+	public class DefaultPreflightConfig
+	{
+		/// <summary>
+		/// The template id to query
+		/// </summary>
+		public string? TemplateId { get; set; }
+
+		/// <summary>
+		/// The last successful job type to use for the base changelist
+		/// </summary>
+		[Obsolete("Use Change.TemplateId instead")]
+		public string? ChangeTemplateId { get; set; }
+
+		/// <summary>
+		/// Query for the change to use
+		/// </summary>
+		public ChangeQueryConfig? Change { get; set; }
+
+		/// <summary>
+		/// Convert to a model object
+		/// </summary>
+		/// <returns></returns>
+		public DefaultPreflight ToModel()
+		{
+#pragma warning disable CS0618 // Type or member is obsolete
+			if (ChangeTemplateId != null)
+			{
+				Change ??= new ChangeQueryConfig();
+				Change.TemplateId = ChangeTemplateId;
+			}
+			return new DefaultPreflight((TemplateId != null) ? (TemplateRefId?)new TemplateRefId(TemplateId) : null, Change?.ToModel());
+#pragma warning restore CS0618 // Type or member is obsolete
+		}
+	}
+
+	/// <summary>
+	/// Trigger for another template
+	/// </summary>
+	public class ChainedJobTemplateConfig
+	{
+		/// <summary>
+		/// Name of the target that needs to complete before starting the other template
+		/// </summary>
+		[Required]
+		public string Trigger { get; set; } = String.Empty;
+
+		/// <summary>
+		/// Id of the template to trigger
+		/// </summary>
+		[Required]
+		public string TemplateId { get; set; } = String.Empty;
+	}
+
+	/// <summary>
+	/// Parameters to create a template within a stream
+	/// </summary>
+	public class TemplateRefConfig : TemplateConfig
+	{
+		TemplateRefId _id;
+
+		/// <summary>
+		/// Optional identifier for this ref. If not specified, an id will be generated from the name.
+		/// </summary>
+		public TemplateRefId Id
+		{
+			get => _id.IsEmpty ? TemplateRefId.Sanitize(Name) : _id;
+			set => _id = value;
+		}
+
+		/// <summary>
+		/// Whether to show badges in UGS for these jobs
+		/// </summary>
+		public bool ShowUgsBadges { get; set; }
+
+		/// <summary>
+		/// Whether to show alerts in UGS for these jobs
+		/// </summary>
+		public bool ShowUgsAlerts { get; set; }
+
+		/// <summary>
+		/// Notification channel for this template. Overrides the stream channel if set.
+		/// </summary>
+		public string? NotificationChannel { get; set; }
+
+		/// <summary>
+		/// Notification channel filter for this template. Can be Success|Failure|Warnings
+		/// </summary>
+		public string? NotificationChannelFilter { get; set; }
+
+		/// <summary>
+		/// Triage channel for this template. Overrides the stream channel if set.
+		/// </summary>
+		public string? TriageChannel { get; set; }
+
+		/// <summary>
+		/// Workflow to user for this stream
+		/// </summary>
+		public WorkflowId? WorkflowId
+		{
+			get => Annotations.WorkflowId;
+			set => Annotations.WorkflowId = value;
+		}
+
+		/// <summary>
+		/// Default annotations to apply to nodes in this template
+		/// </summary>
+		public NodeAnnotations Annotations { get; set; } = new NodeAnnotations();
+
+		/// <summary>
+		/// Schedule to execute this template
+		/// </summary>
+		public ScheduleConfig? Schedule { get; set; }
+
+		/// <summary>
+		/// List of chained job triggers
+		/// </summary>
+		public List<ChainedJobTemplateConfig>? ChainedJobs { get; set; }
+
+		/// <summary>
+		/// The ACL for this template
+		/// </summary>
+		public AclConfig? Acl { get; set; }
+	}
+
+	/// <summary>
+	/// Parameters to create a new schedule
+	/// </summary>
+	public class SchedulePatternConfig
+	{
+		/// <summary>
+		/// Days of the week to run this schedule on. If null, the schedule will run every day.
+		/// </summary>
+		public List<string>? DaysOfWeek { get; set; }
+
+		/// <summary>
+		/// Time during the day for the first schedule to trigger. Measured in minutes from midnight.
+		/// </summary>
+		public int MinTime { get; set; }
+
+		/// <summary>
+		/// Time during the day for the last schedule to trigger. Measured in minutes from midnight.
+		/// </summary>
+		public int? MaxTime { get; set; }
+
+		/// <summary>
+		/// Interval between each schedule triggering
+		/// </summary>
+		public int? Interval { get; set; }
+
+		/// <summary>
+		/// Constructs a model object from this request
+		/// </summary>
+		/// <returns>Model object</returns>
+		public SchedulePattern ToModel()
+		{
+			return new SchedulePattern(DaysOfWeek?.ConvertAll(x => Enum.Parse<DayOfWeek>(x)), MinTime, MaxTime, Interval);
+		}
+	}
+
+	/// <summary>
+	/// Gate allowing a schedule to trigger.
+	/// </summary>
+	public class ScheduleGateConfig
+	{
+		/// <summary>
+		/// The template containing the dependency
+		/// </summary>
+		[Required]
+		public string TemplateId { get; set; } = String.Empty;
+
+		/// <summary>
+		/// Target to wait for
+		/// </summary>
+		[Required]
+		public string Target { get; set; } = String.Empty;
+
+		/// <summary>
+		/// Constructs a model object
+		/// </summary>
+		/// <returns>New model object.</returns>
+		public ScheduleGate ToModel()
+		{
+			return new ScheduleGate(new TemplateRefId(TemplateId), Target);
+		}
+	}
+
+	/// <summary>
+	/// Parameters to create a new schedule
+	/// </summary>
+	public class ScheduleConfig
+	{
+		/// <summary>
+		/// Roles to impersonate for this schedule
+		/// </summary>
+		public List<AclClaimConfig>? Claims { get; set; }
+
+		/// <summary>
+		/// Whether the schedule should be enabled
+		/// </summary>
+		public bool Enabled { get; set; }
+
+		/// <summary>
+		/// Maximum number of builds that can be active at once
+		/// </summary>
+		public int MaxActive { get; set; }
+
+		/// <summary>
+		/// Maximum number of changes the schedule can fall behind head revision. If greater than zero, builds will be triggered for every submitted changelist until the backlog is this size.
+		/// </summary>
+		public int MaxChanges { get; set; }
+
+		/// <summary>
+		/// Whether the build requires a change to be submitted
+		/// </summary>
+		public bool RequireSubmittedChange { get; set; } = true;
+
+		/// <summary>
+		/// Gate allowing the schedule to trigger
+		/// </summary>
+		public ScheduleGateConfig? Gate { get; set; }
+
+		/// <summary>
+		/// The types of changes to run for
+		/// </summary>
+		public List<ChangeContentFlags>? Filter { get; set; }
+
+		/// <summary>
+		/// Files that should cause the job to trigger
+		/// </summary>
+		public List<string>? Files { get; set; }
+
+		/// <summary>
+		/// Parameters for the template
+		/// </summary>
+		public Dictionary<string, string> TemplateParameters { get; set; } = new Dictionary<string, string>();
+
+		/// <summary>
+		/// New patterns for the schedule
+		/// </summary>
+		public List<SchedulePatternConfig> Patterns { get; set; } = new List<SchedulePatternConfig>();
+
+		/// <summary>
+		/// Constructs a model object
+		/// </summary>
+		/// <param name="currentTimeUtc">The current time</param>
+		/// <returns>New model object</returns>
+		public Schedule ToModel(DateTime currentTimeUtc)
+		{
+			return new Schedule(currentTimeUtc, Enabled, MaxActive, MaxChanges, RequireSubmittedChange, Gate?.ToModel(), Filter, Files, TemplateParameters, Patterns.ConvertAll(x => x.ToModel()));
+		}
+	}
+
 	/// <summary>
 	/// Configuration for allocating access tokens for each job
 	/// </summary>
diff --git a/Engine/Source/Programs/Horde/Horde.Build/Streams/StreamMessages.cs b/Engine/Source/Programs/Horde/Horde.Build/Streams/StreamMessages.cs
--- a/Engine/Source/Programs/Horde/Horde.Build/Streams/StreamMessages.cs
+++ b/Engine/Source/Programs/Horde/Horde.Build/Streams/StreamMessages.cs
@@ -18,248 +18,6 @@ namespace Horde.Build.Streams
 	using TemplateRefId = StringId<TemplateRef>;
 	using WorkflowId = StringId<WorkflowConfig>;
 
-	/// <summary>
-	/// Query selecting the base changelist to use
-	/// </summary>
-	public class ChangeQueryRequest
-	{
-		/// <summary>
-		/// The template id to query
-		/// </summary>
-		public string? TemplateId { get; set; }
-
-		/// <summary>
-		/// The target to query
-		/// </summary>
-		public string? Target { get; set; }
-
-		/// <summary>
-		/// Whether to match a job that produced warnings
-		/// </summary>
-		public List<JobStepOutcome>? Outcomes { get; set; }
-
-		/// <summary>
-		/// Convert to a model object
-		/// </summary>
-		/// <returns></returns>
-		public ChangeQuery ToModel()
-		{
-			ChangeQuery query = new ChangeQuery();
-			if (TemplateId != null)
-			{
-				query.TemplateRefId = new StringId<TemplateRef>(TemplateId);
-			}
-			query.Target = Target;
-			query.Outcomes = Outcomes;
-			return query;
-		}
-	}
-
-	/// <summary>
-	/// Specifies defaults for running a preflight
-	/// </summary>
-	public class DefaultPreflightRequest
-	{
-		/// <summary>
-		/// The template id to query
-		/// </summary>
-		public string? TemplateId { get; set; }
-
-		/// <summary>
-		/// The last successful job type to use for the base changelist
-		/// </summary>
-		[Obsolete("Use Change.TemplateId instead")]
-		public string? ChangeTemplateId { get; set; }
-
-		/// <summary>
-		/// Query for the change to use
-		/// </summary>
-		public ChangeQueryRequest? Change { get; set; }
-
-		/// <summary>
-		/// Convert to a model object
-		/// </summary>
-		/// <returns></returns>
-		public DefaultPreflight ToModel()
-		{
-#pragma warning disable CS0618 // Type or member is obsolete
-			if (ChangeTemplateId != null)
-			{
-				Change ??= new ChangeQueryRequest();
-				Change.TemplateId = ChangeTemplateId;
-			}
-			return new DefaultPreflight((TemplateId != null) ? (TemplateRefId?)new TemplateRefId(TemplateId) : null, Change?.ToModel());
-#pragma warning restore CS0618 // Type or member is obsolete
-		}
-	}
-
-	/// <summary>
-	/// Mapping from a BuildGraph agent type to a set of machines on the farm
-	/// </summary>
-	public class CreateAgentTypeRequest
-	{
-		/// <summary>
-		/// Pool of agents to use for this agent type
-		/// </summary>
-		[Required]
-		public string Pool { get; set; } = null!;
-
-		/// <summary>
-		/// Name of the workspace to sync
-		/// </summary>
-		public string? Workspace { get; set; }
-
-		/// <summary>
-		/// Path to the temporary storage dir
-		/// </summary>
-		public string? TempStorageDir { get; set; }
-
-		/// <summary>
-		/// Environment variables to be set when executing the job
-		/// </summary>
-		public Dictionary<string, string>? Environment { get; set; }
-	}
-
-	/// <summary>
-	/// Information about a workspace type
-	/// </summary>
-	public class CreateWorkspaceTypeRequest
-	{
-		/// <summary>
-		/// Name of the Perforce server cluster to use
-		/// </summary>
-		public string? Cluster { get; set; }
-
-		/// <summary>
-		/// The Perforce server and port (eg. perforce:1666)
-		/// </summary>
-		public string? ServerAndPort { get; set; }
-
-		/// <summary>
-		/// User to log into Perforce with (defaults to buildmachine)
-		/// </summary>
-		public string? UserName { get; set; }
-
-		/// <summary>
-		/// Password to use to log into the workspace
-		/// </summary>
-		public string? Password { get; set; }
-
-		/// <summary>
-		/// Identifier to distinguish this workspace from other workspaces. Defaults to the workspace type name.
-		/// </summary>
-		public string? Identifier { get; set; }
-
-		/// <summary>
-		/// Override for the stream to sync
-		/// </summary>
-		public string? Stream { get; set; }
-
-		/// <summary>
-		/// Custom view for the workspace
-		/// </summary>
-		public List<string>? View { get; set; }
-
-		/// <summary>
-		/// Whether to use an incrementally synced workspace
-		/// </summary>
-		public bool Incremental { get; set; }
-
-		/// <summary>
-		/// Whether to use the AutoSDK
-		/// </summary>
-		public bool UseAutoSdk { get; set; } = true;
-	}
-
-	/// <summary>
-	/// Trigger for another template
-	/// </summary>
-	public class CreateChainedJobTemplateRequest
-	{
-		/// <summary>
-		/// Name of the target that needs to complete before starting the other template
-		/// </summary>
-		[Required]
-		public string Trigger { get; set; } = String.Empty;
-
-		/// <summary>
-		/// Id of the template to trigger
-		/// </summary>
-		[Required]
-		public string TemplateId { get; set; } = String.Empty;
-	}
-
-	/// <summary>
-	/// Parameters to create a template within a stream
-	/// </summary>
-	public class TemplateRefConfig : TemplateConfig
-	{
-		TemplateRefId _id;
-
-		/// <summary>
-		/// Optional identifier for this ref. If not specified, an id will be generated from the name.
-		/// </summary>
-		public TemplateRefId Id
-		{
-			get => _id.IsEmpty ? TemplateRefId.Sanitize(Name) : _id;
-			set => _id = value;
-		}
-
-		/// <summary>
-		/// Whether to show badges in UGS for these jobs
-		/// </summary>
-		public bool ShowUgsBadges { get; set; }
-
-		/// <summary>
-		/// Whether to show alerts in UGS for these jobs
-		/// </summary>
-		public bool ShowUgsAlerts { get; set; }
-
-		/// <summary>
-		/// Notification channel for this template. Overrides the stream channel if set.
-		/// </summary>
-		public string? NotificationChannel { get; set; }
-
-		/// <summary>
-		/// Notification channel filter for this template. Can be Success|Failure|Warnings
-		/// </summary>
-		public string? NotificationChannelFilter { get; set; }
-
-		/// <summary>
-		/// Triage channel for this template. Overrides the stream channel if set.
-		/// </summary>
-		public string? TriageChannel { get; set; }
-
-		/// <summary>
-		/// Workflow to user for this stream
-		/// </summary>
-		public WorkflowId? WorkflowId
-		{
-			get => Annotations.WorkflowId;
-			set => Annotations.WorkflowId = value;
-		}
-
-		/// <summary>
-		/// Default annotations to apply to nodes in this template
-		/// </summary>
-		public NodeAnnotations Annotations { get; set; } = new NodeAnnotations();
-
-		/// <summary>
-		/// Schedule to execute this template
-		/// </summary>
-		public CreateScheduleRequest? Schedule { get; set; }
-
-		/// <summary>
-		/// List of chained job triggers
-		/// </summary>
-		public List<CreateChainedJobTemplateRequest>? ChainedJobs { get; set; }
-
-		/// <summary>
-		/// The ACL for this template
-		/// </summary>
-		public UpdateAclRequest? Acl { get; set; }
-	}
-
 	/// <summary>
 	/// Step state update request
 	/// </summary>
@@ -611,7 +369,7 @@ namespace Horde.Build.Streams
 		/// <summary>
 		/// Default template to use for preflights
 		/// </summary>
-		public DefaultPreflightRequest? DefaultPreflight { get; set; }
+		public DefaultPreflightConfig? DefaultPreflight { get; set; }
 
 		/// <summary>
 		/// List of tabs to display for this stream
@@ -673,7 +431,7 @@ namespace Horde.Build.Streams
 		/// <param name="pausedUntil">Stream paused for new builds until this date</param>
 		/// <param name="pauseComment">Reason for stream being paused</param>
 		/// <param name="workflows">Workflows for this stream</param>
-		public GetStreamResponse(string id, string projectId, string name, string configRevision, int order, string? notificationChannel, string? notificationChannelFilter, string? triageChannel, DefaultPreflightRequest? defaultPreflight, List<GetStreamTabResponse> tabs, Dictionary<string, GetAgentTypeResponse> agentTypes, Dictionary<string, GetWorkspaceTypeResponse>? workspaceTypes, List<GetTemplateRefResponse> templates, GetAclResponse? acl, DateTime? pausedUntil, string? pauseComment, List<WorkflowConfig> workflows)
+		public GetStreamResponse(string id, string projectId, string name, string configRevision, int order, string? notificationChannel, string? notificationChannelFilter, string? triageChannel, DefaultPreflightConfig? defaultPreflight, List<GetStreamTabResponse> tabs, Dictionary<string, GetAgentTypeResponse> agentTypes, Dictionary<string, GetWorkspaceTypeResponse>? workspaceTypes, List<GetTemplateRefResponse> templates, GetAclResponse? acl, DateTime? pausedUntil, string? pauseComment, List<WorkflowConfig> workflows)
 		{
 			Id = id;
 			ProjectId = projectId;
