Horde: Missing file.

diff --git a/Engine/Source/Programs/Horde/Horde.Build/Acls/AclConfig.cs b/Engine/Source/Programs/Horde/Horde.Build/Acls/AclConfig.cs
new file mode 100644
--- /dev/null
+++ b/Engine/Source/Programs/Horde/Horde.Build/Acls/AclConfig.cs
@@ -0,0 +1,64 @@
+// Copyright Epic Games, Inc. All Rights Reserved.
+
+using System.Collections.Generic;
+using System.ComponentModel.DataAnnotations;
+
+namespace Horde.Build.Acls
+{
+	/// <summary>
+	/// Parameters to update an ACL
+	/// </summary>
+	public class AclConfig
+	{
+		/// <summary>
+		/// Entries to replace the existing ACL
+		/// </summary>
+		public List<AclEntryConfig>? Entries { get; set; }
+
+		/// <summary>
+		/// Whether to inherit permissions from the parent ACL
+		/// </summary>
+		public bool? Inherit { get; set; }
+
+		/// <summary>
+		/// List of exceptions to the inherited setting
+		/// </summary>
+		public List<string>? Exceptions { get; set; }
+	}
+
+	/// <summary>
+	/// Individual entry in an ACL
+	/// </summary>
+	public class AclEntryConfig
+	{
+		/// <summary>
+		/// Name of the user or group
+		/// </summary>
+		[Required]
+		public AclClaimConfig Claim { get; set; } = null!;
+
+		/// <summary>
+		/// Array of actions to allow
+		/// </summary>
+		[Required]
+		public string[] Actions { get; set; } = null!;
+	}
+
+	/// <summary>
+	/// New claim to create
+	/// </summary>
+	public class AclClaimConfig
+	{
+		/// <summary>
+		/// The claim type
+		/// </summary>
+		[Required]
+		public string Type { get; set; } = null!;
+
+		/// <summary>
+		/// The claim value
+		/// </summary>
+		[Required]
+		public string Value { get; set; } = null!;
+	}
+}
