Horde Storage - Avoid fetching references in the blob index when it is not needed as it can (for some blobs) grow very large.

diff --git a/Engine/Source/Programs/Horde/HordeStorage/Horde.Storage/Implementation/Blob/CachedBlobIndex.cs b/Engine/Source/Programs/Horde/HordeStorage/Horde.Storage/Implementation/Blob/CachedBlobIndex.cs
--- a/Engine/Source/Programs/Horde/HordeStorage/Horde.Storage/Implementation/Blob/CachedBlobIndex.cs
+++ b/Engine/Source/Programs/Horde/HordeStorage/Horde.Storage/Implementation/Blob/CachedBlobIndex.cs
@@ -21,7 +21,7 @@ public class CachedBlobIndex : IBlobIndex
         _fileSystemStore = fileSystemStore;
     }
 
-    public async Task<BlobInfo?> GetBlobInfo(NamespaceId ns, BlobIdentifier id)
+    public async Task<BlobInfo?> GetBlobInfo(NamespaceId ns, BlobIdentifier id, BlobIndexFlags flags = BlobIndexFlags.None)
     {
         if (await BlobExistsInRegion(ns, id))
         {
diff --git a/Engine/Source/Programs/Horde/HordeStorage/Horde.Storage/Implementation/Blob/IBlobIndex.cs b/Engine/Source/Programs/Horde/HordeStorage/Horde.Storage/Implementation/Blob/IBlobIndex.cs
--- a/Engine/Source/Programs/Horde/HordeStorage/Horde.Storage/Implementation/Blob/IBlobIndex.cs
+++ b/Engine/Source/Programs/Horde/HordeStorage/Horde.Storage/Implementation/Blob/IBlobIndex.cs
@@ -12,14 +12,19 @@ namespace Horde.Storage.Implementation.Blob
         public HashSet<string> Regions { get; init; } = new HashSet<string>();
         public NamespaceId Namespace { get; init; }
         public BlobIdentifier BlobIdentifier { get; init; } = null!;
-        public List<(BucketId, IoHashKey)> References { get; init; } = new List<(BucketId, IoHashKey)>();
+        public List<(BucketId, IoHashKey)>? References { get; init; } = new List<(BucketId, IoHashKey)>();
+    }
+
+    public enum BlobIndexFlags {
+        None = 0,
+        IncludeReferences = 1
     }
 
     public interface IBlobIndex
     {
         Task AddBlobToIndex(NamespaceId ns, BlobIdentifier id, string? region = null);
 
-        Task<BlobInfo?> GetBlobInfo(NamespaceId ns, BlobIdentifier id);
+        Task<BlobInfo?> GetBlobInfo(NamespaceId ns, BlobIdentifier id, BlobIndexFlags flags = BlobIndexFlags.None);
 
         Task<bool> RemoveBlobFromIndex(NamespaceId ns, BlobIdentifier id);
         Task RemoveBlobFromRegion(NamespaceId ns, BlobIdentifier id, string? region = null);
diff --git a/Engine/Source/Programs/Horde/HordeStorage/Horde.Storage/Implementation/Blob/MemoryBlobIndex.cs b/Engine/Source/Programs/Horde/HordeStorage/Horde.Storage/Implementation/Blob/MemoryBlobIndex.cs
--- a/Engine/Source/Programs/Horde/HordeStorage/Horde.Storage/Implementation/Blob/MemoryBlobIndex.cs
+++ b/Engine/Source/Programs/Horde/HordeStorage/Horde.Storage/Implementation/Blob/MemoryBlobIndex.cs
@@ -14,7 +14,6 @@ public class MemoryBlobIndex : IBlobIndex
 {
     private class MemoryBlobInfo : BlobInfo
     {
-
     }
 
     private readonly ConcurrentDictionary<NamespaceId, ConcurrentDictionary<BlobIdentifier, MemoryBlobInfo>> _index = new ();
@@ -38,7 +37,7 @@ public class MemoryBlobIndex : IBlobIndex
         return Task.CompletedTask;
     }
 
-    public Task<BlobInfo?> GetBlobInfo(NamespaceId ns, BlobIdentifier id)
+    public Task<BlobInfo?> GetBlobInfo(NamespaceId ns, BlobIdentifier id, BlobIndexFlags flags = BlobIndexFlags.None)
     {
         ConcurrentDictionary<BlobIdentifier, MemoryBlobInfo> index = GetNamespaceContainer(ns);
 
@@ -89,11 +88,11 @@ public class MemoryBlobIndex : IBlobIndex
             index.AddOrUpdate(id, _ =>
             {
                 MemoryBlobInfo info = NewBlobInfo(ns, id, _jupiterSettings.CurrentValue.CurrentSite);
-                info.References.Add((bucket, key));
+                info.References!.Add((bucket, key));
                 return info;
             }, (_, info) =>
             {
-                info.References.Add((bucket, key));
+                info.References!.Add((bucket, key));
                 return info;
             });
         }
@@ -121,6 +120,7 @@ public class MemoryBlobIndex : IBlobIndex
             Regions = new HashSet<string> { region },
             Namespace = ns,
             BlobIdentifier = blob,
+            References = new List<(BucketId, IoHashKey)>()
         };
         return info;
     }
diff --git a/Engine/Source/Programs/Horde/HordeStorage/Horde.Storage/Implementation/Blob/MongoBlobIndex.cs b/Engine/Source/Programs/Horde/HordeStorage/Horde.Storage/Implementation/Blob/MongoBlobIndex.cs
--- a/Engine/Source/Programs/Horde/HordeStorage/Horde.Storage/Implementation/Blob/MongoBlobIndex.cs
+++ b/Engine/Source/Programs/Horde/HordeStorage/Horde.Storage/Implementation/Blob/MongoBlobIndex.cs
@@ -54,7 +54,7 @@ public class MongoBlobIndex : MongoStore, IBlobIndex
         });
     }
 
-    public async Task<BlobInfo?> GetBlobInfo(NamespaceId ns, BlobIdentifier id)
+    public async Task<BlobInfo?> GetBlobInfo(NamespaceId ns, BlobIdentifier id, BlobIndexFlags flags = BlobIndexFlags.None)
     {
         IMongoCollection<MongoBlobIndexModelV0> collection = GetCollection<MongoBlobIndexModelV0>();
         IAsyncCursor<MongoBlobIndexModelV0>? cursor = await collection.FindAsync(m => m.Ns == ns.ToString() && m.BlobId == id.ToString());
@@ -152,8 +152,11 @@ class MongoBlobIndexModelV0
         Ns = blobInfo.Namespace.ToString();
         BlobId = blobInfo.BlobIdentifier.ToString();
         Regions = blobInfo.Regions.ToList();
-        References = blobInfo.References.Select(pair => new Dictionary<string, string>{ {"bucket", pair.Item1.ToString()}, {"key", pair.Item2.ToString()}}).ToList();
-    }
+        if (blobInfo.References != null)
+        {
+            References = blobInfo.References.Select(pair => new Dictionary<string, string> { { "bucket", pair.Item1.ToString() }, { "key", pair.Item2.ToString() } }).ToList();
+        }
+	}
 
     public MongoBlobIndexModelV0(NamespaceId ns, BlobIdentifier blobId)
     {
diff --git a/Engine/Source/Programs/Horde/HordeStorage/Horde.Storage/Implementation/Blob/ScyllaBlobIndex.cs b/Engine/Source/Programs/Horde/HordeStorage/Horde.Storage/Implementation/Blob/ScyllaBlobIndex.cs
--- a/Engine/Source/Programs/Horde/HordeStorage/Horde.Storage/Implementation/Blob/ScyllaBlobIndex.cs
+++ b/Engine/Source/Programs/Horde/HordeStorage/Horde.Storage/Implementation/Blob/ScyllaBlobIndex.cs
@@ -47,21 +47,24 @@ public class ScyllaBlobIndex : IBlobIndex
             new string[] { region }, ns.ToString(), new ScyllaBlobIdentifier(id));
     }
 
-    public async Task<BlobInfo?> GetBlobInfo(NamespaceId ns, BlobIdentifier id)
+    public async Task<BlobInfo?> GetBlobInfo(NamespaceId ns, BlobIdentifier id, BlobIndexFlags flags = BlobIndexFlags.None)
     {
         using IScope scope = Tracer.Instance.StartActive("scylla.fetch_blob_index");
         scope.Span.ResourceName = $"{ns}.{id}";
 
-        ScyllaBlobIndexTable? blobIndex =
-            await _mapper.FirstOrDefaultAsync<ScyllaBlobIndexTable>("WHERE namespace = ? AND blob_id = ?",
-                ns.ToString(), new ScyllaBlobIdentifier(id));
+        bool includeReferences = (flags & BlobIndexFlags.IncludeReferences) != 0;
 
-        if (blobIndex == null)
+        ScyllaBlobIndexTable o;
+        if (includeReferences)
         {
-            return null;
+            o = await _mapper.SingleOrDefaultAsync<ScyllaBlobIndexTable>("SELECT namespace, blob_id, regions, references FROM blob_index WHERE namespace = ? AND blob_id = ?", ns.ToString(), new ScyllaBlobIdentifier(id));
+        }
+        else
+        {
+            o = await _mapper.SingleOrDefaultAsync<ScyllaBlobIndexTable>("SELECT namespace, blob_id, regions FROM blob_index WHERE namespace = ? AND blob_id = ?", ns.ToString(), new ScyllaBlobIdentifier(id));
         }
 
-        return blobIndex.ToBlobInfo();
+        return o?.ToBlobInfo();
     }
 
     public async Task<bool> RemoveBlobFromIndex(NamespaceId ns, BlobIdentifier id)
diff --git a/Engine/Source/Programs/Horde/HordeStorage/Horde.Storage/Implementation/GC/OrphanBlobCleanupRefs.cs b/Engine/Source/Programs/Horde/HordeStorage/Horde.Storage/Implementation/GC/OrphanBlobCleanupRefs.cs
--- a/Engine/Source/Programs/Horde/HordeStorage/Horde.Storage/Implementation/GC/OrphanBlobCleanupRefs.cs
+++ b/Engine/Source/Programs/Horde/HordeStorage/Horde.Storage/Implementation/GC/OrphanBlobCleanupRefs.cs
@@ -134,14 +134,14 @@ namespace Horde.Storage.Implementation
                     break;
                 }
 
-                BlobInfo? blobIndex = await _blobIndex.GetBlobInfo(blobNamespace, blob);
+                BlobInfo? blobIndex = await _blobIndex.GetBlobInfo(blobNamespace, blob, BlobIndexFlags.IncludeReferences);
 
                 if (blobIndex == null)
                 {
                     break;
                 }
 
-                foreach ((BucketId, IoHashKey) tuple in blobIndex.References)
+                foreach ((BucketId, IoHashKey) tuple in blobIndex.References!)
                 {
                     try
                     {
diff --git a/Engine/Source/Programs/Horde/HordeStorage/Horde.Storage/Tests/Functional/Blob/BlobIndexTests.cs b/Engine/Source/Programs/Horde/HordeStorage/Horde.Storage/Tests/Functional/Blob/BlobIndexTests.cs
--- a/Engine/Source/Programs/Horde/HordeStorage/Horde.Storage/Tests/Functional/Blob/BlobIndexTests.cs
+++ b/Engine/Source/Programs/Horde/HordeStorage/Horde.Storage/Tests/Functional/Blob/BlobIndexTests.cs
@@ -112,11 +112,11 @@ namespace Horde.Storage.FunctionalTests.Storage
 
             IBlobIndex? index = _server!.Services.GetService<IBlobIndex>();
             Assert.IsNotNull(index);
-            BlobInfo? blobInfo = await index.GetBlobInfo(_testNamespaceName, objectHash);
+            BlobInfo? blobInfo = await index.GetBlobInfo(_testNamespaceName, objectHash, BlobIndexFlags.IncludeReferences);
 
             Assert.IsNotNull(blobInfo);
             Assert.IsTrue(blobInfo.Regions.Contains("test"));
-            Assert.AreEqual(1, blobInfo.References.Count);
+            Assert.AreEqual(1, blobInfo.References!.Count);
 
             (BucketId bucket, IoHashKey key) = blobInfo.References[0];
             Assert.AreEqual("bucket", bucket.ToString());
