Horde: Move redundant stream configuration properties from MongoDB, and reference the data directly from the StreamConfig class.

diff --git a/Engine/Source/Programs/Horde/Horde.Build.Tests/Fixture.cs b/Engine/Source/Programs/Horde/Horde.Build.Tests/Fixture.cs
--- a/Engine/Source/Programs/Horde/Horde.Build.Tests/Fixture.cs
+++ b/Engine/Source/Programs/Horde/Horde.Build.Tests/Fixture.cs
@@ -6,6 +6,7 @@ using System.Text;
 using System.Threading.Tasks;
 using EpicGames.Core;
 using Horde.Build.Agents;
+using Horde.Build.Agents.Pools;
 using Horde.Build.Configuration;
 using Horde.Build.Jobs;
 using Horde.Build.Jobs.Artifacts;
@@ -20,6 +21,7 @@ using HordeCommon;
 namespace Horde.Build.Tests
 {
 	using JobId = ObjectId<IJob>;
+	using PoolId = StringId<IPool>;
     using ProjectId = StringId<IProject>;
     using StreamId = StringId<IStream>;
     using TemplateRefId = StringId<TemplateRef>;
@@ -74,7 +76,7 @@ namespace Horde.Build.Tests
 
 			Dictionary<string, AgentConfig> agentTypes = new()
 			{
-				{ "Win64", new() { Pool = PoolName} }
+				{ "Win64", new() { Pool = new PoolId(PoolName) } }
 			};
 
 			StreamConfig config = new StreamConfig { Name = "//UE5/Main", Tabs = tabs, Templates = templates, AgentTypes = agentTypes };
diff --git a/Engine/Source/Programs/Horde/Horde.Build.Tests/Fleet/JobQueueStrategyTest.cs b/Engine/Source/Programs/Horde/Horde.Build.Tests/Fleet/JobQueueStrategyTest.cs
--- a/Engine/Source/Programs/Horde/Horde.Build.Tests/Fleet/JobQueueStrategyTest.cs
+++ b/Engine/Source/Programs/Horde/Horde.Build.Tests/Fleet/JobQueueStrategyTest.cs
@@ -113,7 +113,7 @@ namespace Horde.Build.Tests.Fleet
 			PoolSizeData poolSize = new (pool1, agents, null);
 			
 			string agentTypeName1 = "bogusAgentType1";
-			Dictionary<string, AgentConfig> agentTypes = new() { {agentTypeName1, new() { Pool = pool1.Name} }, };
+			Dictionary<string, AgentConfig> agentTypes = new() { {agentTypeName1, new() { Pool = new PoolId(pool1.Name) } }, };
 
 			IStream stream = (await CreateOrReplaceStreamAsync(
 				new StreamId("ue5-main"),
diff --git a/Engine/Source/Programs/Horde/Horde.Build/Agents/Fleet/JobQueueStrategy.cs b/Engine/Source/Programs/Horde/Horde.Build/Agents/Fleet/JobQueueStrategy.cs
--- a/Engine/Source/Programs/Horde/Horde.Build/Agents/Fleet/JobQueueStrategy.cs
+++ b/Engine/Source/Programs/Horde/Horde.Build/Agents/Fleet/JobQueueStrategy.cs
@@ -129,7 +129,7 @@ namespace Horde.Build.Agents.Fleet
 				}
 
 				string batchAgentType = graph.Groups[batch.GroupIdx].AgentType;
-				if (!stream.AgentTypes.TryGetValue(batchAgentType, out AgentType? agentType))
+				if (!stream.Config.AgentTypes.TryGetValue(batchAgentType, out AgentConfig? agentType))
 				{
 					continue;
 				}
diff --git a/Engine/Source/Programs/Horde/Horde.Build/Agents/Pools/PoolUpdateService.cs b/Engine/Source/Programs/Horde/Horde.Build/Agents/Pools/PoolUpdateService.cs
--- a/Engine/Source/Programs/Horde/Horde.Build/Agents/Pools/PoolUpdateService.cs
+++ b/Engine/Source/Programs/Horde/Horde.Build/Agents/Pools/PoolUpdateService.cs
@@ -75,14 +75,14 @@ namespace Horde.Build.Agents.Pools
 				List<IStream> activeStreams = await _streams.FindAllAsync();
 				foreach (IStream activeStream in activeStreams)
 				{
-					foreach (KeyValuePair<string, AgentType> agentTypePair in activeStream.AgentTypes)
+					foreach (KeyValuePair<string, AgentConfig> agentTypePair in activeStream.Config.AgentTypes)
 					{
 						// Create the new agent workspace
 						(AgentWorkspace, bool)? result;
 						if (activeStream.TryGetAgentWorkspace(agentTypePair.Value, out result))
 						{
 							(AgentWorkspace agentWorkspace, bool useAutoSdk) = result.Value;
-							AgentType agentType = agentTypePair.Value;
+							AgentConfig agentType = agentTypePair.Value;
 
 							// Find or add a list of workspaces for this pool
 							List<AgentWorkspace>? agentWorkspaces;
diff --git a/Engine/Source/Programs/Horde/Horde.Build/Commands/Bundles/PerforceCommand.cs b/Engine/Source/Programs/Horde/Horde.Build/Commands/Bundles/PerforceCommand.cs
--- a/Engine/Source/Programs/Horde/Horde.Build/Commands/Bundles/PerforceCommand.cs
+++ b/Engine/Source/Programs/Horde/Horde.Build/Commands/Bundles/PerforceCommand.cs
@@ -87,7 +87,7 @@ namespace Horde.Build.Commands.Bundles
 				baseContents = await replicationService.ReadCommitTreeAsync(stream, BaseChange, Filter, RevisionsOnly, CancellationToken.None);
 			}
 
-			await foreach (NewCommit newCommit in commitService.FindCommitsForClusterAsync(stream.ClusterName, streamToFirstChange).Take(Count))
+			await foreach (NewCommit newCommit in commitService.FindCommitsForClusterAsync(stream.Config.ClusterName, streamToFirstChange).Take(Count))
 			{
 				string briefSummary = newCommit.Description.Replace('\n', ' ');
 				logger.LogInformation("");
diff --git a/Engine/Source/Programs/Horde/Horde.Build/Jobs/JobService.cs b/Engine/Source/Programs/Horde/Horde.Build/Jobs/JobService.cs
--- a/Engine/Source/Programs/Horde/Horde.Build/Jobs/JobService.cs
+++ b/Engine/Source/Programs/Horde/Horde.Build/Jobs/JobService.cs
@@ -226,7 +226,7 @@ namespace Horde.Build.Jobs
 
 			if (preflightChange != null && ShouldClonePreflightChange(stream.Id))
 			{
-				clonedPreflightChange = await CloneShelvedChangeAsync(stream.ClusterName, clonedPreflightChange ?? preflightChange.Value);
+				clonedPreflightChange = await CloneShelvedChangeAsync(stream.Config.ClusterName, clonedPreflightChange ?? preflightChange.Value);
 			}
 
 			_logger.LogInformation("Creating job at CL {Change}, code CL {CodeChange}, preflight CL {PreflightChange}, cloned CL {ClonedPreflightChange}", change, codeChange, preflightChange, clonedPreflightChange);
@@ -990,7 +990,7 @@ namespace Horde.Build.Jobs
 								IStream? stream = await _streamService.GetCachedStream(job.StreamId);
 								if (stream != null)
 								{
-									await DeleteShelvedChangeAsync(stream.ClusterName, job.ClonedPreflightChange);
+									await DeleteShelvedChangeAsync(stream.Config.ClusterName, job.ClonedPreflightChange);
 								}
 							}
 						}
@@ -1085,7 +1085,7 @@ namespace Horde.Build.Jobs
 					{
 						if (ShouldClonePreflightChange(job.StreamId))
 						{
-							clonedPreflightChange = await CloneShelvedChangeAsync(stream.ClusterName, job.PreflightChange);
+							clonedPreflightChange = await CloneShelvedChangeAsync(stream.Config.ClusterName, job.PreflightChange);
 						}
 						else
 						{
@@ -1096,10 +1096,10 @@ namespace Horde.Build.Jobs
 					_logger.LogInformation("Updating description for {ClonedPreflightChange}", clonedPreflightChange);
 
 					ICommit details = await _perforceService.GetChangeDetailsAsync(stream, clonedPreflightChange);
-					await _perforceService.UpdateChangelistDescription(stream.ClusterName, clonedPreflightChange, details.Description.TrimEnd() + $"\n#preflight {job.Id}");
+					await _perforceService.UpdateChangelistDescription(stream.Config.ClusterName, clonedPreflightChange, details.Description.TrimEnd() + $"\n#preflight {job.Id}");
 
 					_logger.LogInformation("Submitting change {Change} (through {ChangeCopy}) after successful completion of {JobId}", job.PreflightChange, clonedPreflightChange, job.Id);
-					(change, message) = await _perforceService.SubmitShelvedChangeAsync(stream.ClusterName, clonedPreflightChange, job.PreflightChange);
+					(change, message) = await _perforceService.SubmitShelvedChangeAsync(stream.Config.ClusterName, clonedPreflightChange, job.PreflightChange);
 
 					_logger.LogInformation("Attempt to submit {Change} (through {ChangeCopy}): {Message}", job.PreflightChange, clonedPreflightChange, message);
 
@@ -1112,14 +1112,14 @@ namespace Horde.Build.Jobs
 					{
 						if (change != null && job.ClonedPreflightChange != 0)
 						{
-							await DeleteShelvedChangeAsync(stream.ClusterName, job.PreflightChange);
+							await DeleteShelvedChangeAsync(stream.Config.ClusterName, job.PreflightChange);
 						}
 					}
 					else
 					{
 						if (change != null && job.PreflightChange != 0)
 						{
-							await DeleteShelvedChangeAsync(stream.ClusterName, job.PreflightChange);
+							await DeleteShelvedChangeAsync(stream.Config.ClusterName, job.PreflightChange);
 						}
 					}
 				}
diff --git a/Engine/Source/Programs/Horde/Horde.Build/Jobs/JobTaskSource.cs b/Engine/Source/Programs/Horde/Horde.Build/Jobs/JobTaskSource.cs
--- a/Engine/Source/Programs/Horde/Horde.Build/Jobs/JobTaskSource.cs
+++ b/Engine/Source/Programs/Horde/Horde.Build/Jobs/JobTaskSource.cs
@@ -416,7 +416,7 @@ namespace Horde.Build.Jobs
 
 					// Validate the agent type and workspace settings
 					IJobStepBatch batch = newJob.Batches[batchIdx];
-					if (!stream.AgentTypes.TryGetValue(graph.Groups[batch.GroupIdx].AgentType, out AgentType? agentType))
+					if (!stream.Config.AgentTypes.TryGetValue(graph.Groups[batch.GroupIdx].AgentType, out AgentConfig? agentType))
 					{
 						newJob = await SkipBatchAsync(newJob, batch.Id, graph, JobStepBatchError.UnknownAgentType);
 					}
@@ -667,8 +667,8 @@ namespace Horde.Build.Jobs
 						INodeGroup group = graph.Groups[batch.GroupIdx];
 
 						// Get the requirements for the new queue item
-						AgentType? agentType;
-						if (stream.AgentTypes.TryGetValue(group.AgentType, out agentType))
+						AgentConfig? agentType;
+						if (stream.Config.AgentTypes.TryGetValue(group.AgentType, out agentType))
 						{
 							(AgentWorkspace, bool)? result;
 							if (stream.TryGetAgentWorkspace(agentType, out result))
diff --git a/Engine/Source/Programs/Horde/Horde.Build/Jobs/JobsController.cs b/Engine/Source/Programs/Horde/Horde.Build/Jobs/JobsController.cs
--- a/Engine/Source/Programs/Horde/Horde.Build/Jobs/JobsController.cs
+++ b/Engine/Source/Programs/Horde/Horde.Build/Jobs/JobsController.cs
@@ -161,7 +161,7 @@ namespace Horde.Build.Jobs
 			string? preflightDescription = null;
 			if (create.PreflightChange != null)
 			{
-				(CheckShelfResult result, preflightDescription) = await _perforce.CheckShelfAsync(stream.ClusterName, stream.Name, create.PreflightChange.Value);
+				(CheckShelfResult result, preflightDescription) = await _perforce.CheckShelfAsync(stream.Config.ClusterName, stream.Name, create.PreflightChange.Value);
 				switch (result)
 				{
 					case CheckShelfResult.Ok:
diff --git a/Engine/Source/Programs/Horde/Horde.Build/Notifications/Sinks/SlackNotificationSink.cs b/Engine/Source/Programs/Horde/Horde.Build/Notifications/Sinks/SlackNotificationSink.cs
--- a/Engine/Source/Programs/Horde/Horde.Build/Notifications/Sinks/SlackNotificationSink.cs
+++ b/Engine/Source/Programs/Horde/Horde.Build/Notifications/Sinks/SlackNotificationSink.cs
@@ -375,9 +375,9 @@ namespace Horde.Build.Notifications.Sinks
 			{
 				await SendJobCompleteNotificationToChannelAsync(job.NotificationChannel, job.NotificationChannelFilter, jobStream, job, graph, outcome);
 			}
-			if (jobStream.NotificationChannel != null)
+			if (jobStream.Config.NotificationChannel != null)
 			{
-				await SendJobCompleteNotificationToChannelAsync(jobStream.NotificationChannel, jobStream.NotificationChannelFilter, jobStream, job, graph, outcome);
+				await SendJobCompleteNotificationToChannelAsync(jobStream.Config.NotificationChannel, jobStream.Config.NotificationChannelFilter, jobStream, job, graph, outcome);
 			}
 		}
 
@@ -709,9 +709,9 @@ namespace Horde.Build.Notifications.Sinks
 					{
 						channels.Add(templateRef.TriageChannel);
 					}
-					else if (stream.TriageChannel != null)
+					else if (stream.Config.TriageChannel != null)
 					{
-						channels.Add(stream.TriageChannel);
+						channels.Add(stream.Config.TriageChannel);
 					}
 				}
 			}
diff --git a/Engine/Source/Programs/Horde/Horde.Build/Perforce/CommitService.cs b/Engine/Source/Programs/Horde/Horde.Build/Perforce/CommitService.cs
--- a/Engine/Source/Programs/Horde/Horde.Build/Perforce/CommitService.cs
+++ b/Engine/Source/Programs/Horde/Horde.Build/Perforce/CommitService.cs
@@ -195,7 +195,7 @@ namespace Horde.Build.Perforce
 		async ValueTask UpdateMetadataAsync(CancellationToken cancellationToken)
 		{
 			List<IStream> streams = await _streamCollection.FindAllAsync();
-			foreach (IGrouping<string, IStream> group in streams.GroupBy(x => x.ClusterName, StringComparer.OrdinalIgnoreCase))
+			foreach (IGrouping<string, IStream> group in streams.GroupBy(x => x.Config.ClusterName, StringComparer.OrdinalIgnoreCase))
 			{
 				try
 				{
diff --git a/Engine/Source/Programs/Horde/Horde.Build/Perforce/IPerforceService.cs b/Engine/Source/Programs/Horde/Horde.Build/Perforce/IPerforceService.cs
--- a/Engine/Source/Programs/Horde/Horde.Build/Perforce/IPerforceService.cs
+++ b/Engine/Source/Programs/Horde/Horde.Build/Perforce/IPerforceService.cs
@@ -183,11 +183,11 @@ namespace Horde.Build.Perforce
 			Match match = Regex.Match(template.SubmitNewChange!, @"^(//[^/]+/[^/]+)/(.+)$");
 			if (match.Success)
 			{
-				return perforce.CreateNewChangeAsync(stream.ClusterName, match.Groups[1].Value, match.Groups[2].Value, description);
+				return perforce.CreateNewChangeAsync(stream.Config.ClusterName, match.Groups[1].Value, match.Groups[2].Value, description);
 			}
 			else
 			{
-				return perforce.CreateNewChangeAsync(stream.ClusterName, stream.Name, template.SubmitNewChange!, description);
+				return perforce.CreateNewChangeAsync(stream.Config.ClusterName, stream.Name, template.SubmitNewChange!, description);
 			}
 		}
 
diff --git a/Engine/Source/Programs/Horde/Horde.Build/Perforce/PerforceService.cs b/Engine/Source/Programs/Horde/Horde.Build/Perforce/PerforceService.cs
--- a/Engine/Source/Programs/Horde/Horde.Build/Perforce/PerforceService.cs
+++ b/Engine/Source/Programs/Horde/Horde.Build/Perforce/PerforceService.cs
@@ -457,7 +457,7 @@ namespace Horde.Build.Perforce
 
 		Task<PooledConnectionHandle> ConnectWithStreamClientAsync(IStream stream, string? userName, CancellationToken cancellationToken)
 		{
-			return ConnectWithStreamClientAsync(stream.ClusterName, userName, stream.Name, cancellationToken);
+			return ConnectWithStreamClientAsync(stream.Config.ClusterName, userName, stream.Name, cancellationToken);
 		}
 
 		async Task<Credentials> GetTicketAsync(PerforceCluster cluster, string userName, CancellationToken cancellationToken)
@@ -598,7 +598,7 @@ namespace Horde.Build.Perforce
 		public async Task<List<ICommit>> GetChangesAsync(IStream stream, int? minChange, int? maxChange, int maxResults, CancellationToken cancellationToken)
 		{
 			using IScope scope = GlobalTracer.Instance.BuildSpan("PerforceService.GetChangesAsync").StartActive();
-			scope.Span.SetTag("ClusterName", stream.ClusterName);
+			scope.Span.SetTag("ClusterName", stream.Config.ClusterName);
 			scope.Span.SetTag("MinChange", minChange ?? -1);
 			scope.Span.SetTag("MaxChange", maxChange ?? -1);
 			scope.Span.SetTag("MaxResults", maxResults);
@@ -623,7 +623,7 @@ namespace Horde.Build.Perforce
 
 		async ValueTask<Commit> CreateCommitInternalAsync(IStream stream, int number, string author, string description, string? basePath, DateTime dateUtc, CancellationToken cancellationToken)
 		{
-			IUser authorUser = await FindOrAddUserAsync(stream.ClusterName, author, cancellationToken);
+			IUser authorUser = await FindOrAddUserAsync(stream.Config.ClusterName, author, cancellationToken);
 			IUser ownerUser = authorUser;
 
 			int originalChange = ParseRobomergeSource(description) ?? number;
@@ -631,10 +631,10 @@ namespace Horde.Build.Perforce
 			string? owner = ParseRobomergeOwner(description);
 			if (owner != null)
 			{
-				ownerUser = await FindOrAddUserAsync(stream.ClusterName, owner, cancellationToken);
+				ownerUser = await FindOrAddUserAsync(stream.Config.ClusterName, owner, cancellationToken);
 			}
 
-			return new Commit(this, stream.ClusterName, stream.Name, stream.Id, number, originalChange, authorUser.Id, ownerUser.Id, description, basePath ?? String.Empty, dateUtc);
+			return new Commit(this, stream.Config.ClusterName, stream.Name, stream.Id, number, originalChange, authorUser.Id, ownerUser.Id, description, basePath ?? String.Empty, dateUtc);
 		}
 
 		async ValueTask<Commit> CreateCommitAsync(IStream stream, DescribeRecord describeRecord, InfoRecord serverInfo, CancellationToken cancellationToken)
@@ -698,11 +698,11 @@ namespace Horde.Build.Perforce
 		public async Task<ICommit> GetChangeDetailsAsync(IStream stream, int changeNumber, CancellationToken cancellationToken)
 		{
 			using IScope scope = GlobalTracer.Instance.BuildSpan("PerforceService.GetChangeDetailsAsync").StartActive();
-			scope.Span.SetTag("ClusterName", stream.ClusterName);
+			scope.Span.SetTag("ClusterName", stream.Config.ClusterName);
 			scope.Span.SetTag("StreamName", stream.Name);
 			scope.Span.SetTag("ChangeNumber", changeNumber);
 
-			PerforceCluster cluster = await GetClusterAsync(stream.ClusterName);
+			PerforceCluster cluster = await GetClusterAsync(stream.Config.ClusterName);
 
 			using (PooledConnectionHandle perforce = await ConnectAsync(cluster, null, cancellationToken))
 			{
@@ -779,7 +779,7 @@ namespace Horde.Build.Perforce
 		public async Task<List<ICommit>> GetChangeDetailsAsync(IStream stream, IReadOnlyList<int> changeNumbers, CancellationToken cancellationToken)
 		{
 			using IScope scope = GlobalTracer.Instance.BuildSpan("PerforceService.GetChangeDetailsAsync").StartActive();
-			scope.Span.SetTag("ClusterName", stream.ClusterName);
+			scope.Span.SetTag("ClusterName", stream.Config.ClusterName);
 			scope.Span.SetTag("StreamName", stream.Name);
 			scope.Span.SetTag("ChangeNumbers.Count", changeNumbers.Count);
 
@@ -1043,7 +1043,7 @@ namespace Horde.Build.Perforce
 		public async Task<int> GetCodeChangeAsync(IStream stream, int change, CancellationToken cancellationToken)
 		{
 			using IScope scope = GlobalTracer.Instance.BuildSpan("PerforceService.GetCodeChangeAsync").StartActive();
-			scope.Span.SetTag("ClusterName", stream.ClusterName);
+			scope.Span.SetTag("ClusterName", stream.Config.ClusterName);
 			scope.Span.SetTag("StreamName", stream.Name);
 			scope.Span.SetTag("Change", change);
 			
diff --git a/Engine/Source/Programs/Horde/Horde.Build/Perforce/ReplicationService.cs b/Engine/Source/Programs/Horde/Horde.Build/Perforce/ReplicationService.cs
--- a/Engine/Source/Programs/Horde/Horde.Build/Perforce/ReplicationService.cs
+++ b/Engine/Source/Programs/Horde/Horde.Build/Perforce/ReplicationService.cs
@@ -44,7 +44,7 @@ namespace Horde.Build.Perforce
 	public class ReplicationServiceOptions
 	{
 		/// <summary>
-		/// Whether to enable replication. Must also be enabled on a per-stream basis (see <see cref="IStream.ReplicationMode"/>).
+		/// Whether to enable replication. Must also be enabled on a per-stream basis (see <see cref="StreamConfig.ReplicationMode"/>).
 		/// </summary>
 		public bool Enable { get; set; } = true;
 
@@ -357,7 +357,7 @@ namespace Horde.Build.Perforce
 			for (; ; )
 			{
 				IStream? stream = await _streamCollection.GetAsync(streamId);
-				if (stream == null || stream.ReplicationMode == ContentReplicationMode.None)
+				if (stream == null || stream.Config.ReplicationMode == ContentReplicationMode.None)
 				{
 					// Remove all but the last item
 					await streamChanges.TrimAsync(0, -2);
@@ -483,7 +483,7 @@ namespace Horde.Build.Perforce
 		/// <returns></returns>
 		static RefName GetRefName(IStream stream, int change)
 		{
-			return GetRefName(stream.Id, change, stream.ReplicationFilter, stream.ReplicationMode == ContentReplicationMode.RevisionsOnly);
+			return GetRefName(stream.Id, change, stream.Config.ReplicationFilter, stream.Config.ReplicationMode == ContentReplicationMode.RevisionsOnly);
 		}
 
 		/// <summary>
@@ -512,8 +512,8 @@ namespace Horde.Build.Perforce
 		/// <returns>Root tree object</returns>
 		public async Task<ReplicationNode> WriteCommitTreeAsync(IStream stream, int change, int baseChange, ReplicationNode baseContents, CancellationToken cancellationToken)
 		{
-			bool revisionsOnly = stream.ReplicationMode == ContentReplicationMode.RevisionsOnly;
-			return await WriteCommitTreeAsync(stream, change, baseChange, baseContents, stream.ReplicationFilter, revisionsOnly, cancellationToken);
+			bool revisionsOnly = stream.Config.ReplicationMode == ContentReplicationMode.RevisionsOnly;
+			return await WriteCommitTreeAsync(stream, change, baseChange, baseContents, stream.Config.ReplicationFilter, revisionsOnly, cancellationToken);
 		}
 
 		[DebuggerDisplay("{_path}")]
@@ -994,7 +994,7 @@ namespace Horde.Build.Perforce
 			ReplicationClient? clientInfo;
 			if (_cachedPerforceClients.TryGetValue(stream.Id, out clientInfo))
 			{
-				if (!String.Equals(clientInfo.ClusterName, stream.ClusterName, StringComparison.Ordinal) && String.Equals(clientInfo.Client.Stream, stream.Name, StringComparison.Ordinal))
+				if (!String.Equals(clientInfo.ClusterName, stream.Config.ClusterName, StringComparison.Ordinal) && String.Equals(clientInfo.Client.Stream, stream.Name, StringComparison.Ordinal))
 				{
 					PerforceSettings serverSettings = new PerforceSettings(clientInfo.Settings);
 					serverSettings.ClientName = null;
@@ -1014,7 +1014,7 @@ namespace Horde.Build.Perforce
 			ReplicationClient? clientInfo = await FindReplicationClientAsync(stream);
 			if (clientInfo == null)
 			{
-				using IPerforceConnection? perforce = await _perforceService.ConnectAsync(stream.ClusterName);
+				using IPerforceConnection? perforce = await _perforceService.ConnectAsync(stream.Config.ClusterName);
 				if (perforce == null)
 				{
 					throw new PerforceException($"Unable to create connection to Perforce server");
@@ -1035,7 +1035,7 @@ namespace Horde.Build.Perforce
 				settings.ClientName = newClient.Name;
 				settings.PreferNativeClient = true;
 
-				clientInfo = new ReplicationClient(settings, stream.ClusterName, serverInfo, newClient, -1);
+				clientInfo = new ReplicationClient(settings, stream.Config.ClusterName, serverInfo, newClient, -1);
 				_cachedPerforceClients.Add(stream.Id, clientInfo);
 			}
 			return clientInfo;
diff --git a/Engine/Source/Programs/Horde/Horde.Build/Streams/IStream.cs b/Engine/Source/Programs/Horde/Horde.Build/Streams/IStream.cs
--- a/Engine/Source/Programs/Horde/Horde.Build/Streams/IStream.cs
+++ b/Engine/Source/Programs/Horde/Horde.Build/Streams/IStream.cs
@@ -19,7 +19,6 @@ using MongoDB.Bson.Serialization.Attributes;
 
 namespace Horde.Build.Streams
 {
-	using PoolId = StringId<IPool>;
 	using ProjectId = StringId<IProject>;
 	using StreamId = StringId<IStream>;
 	using TemplateRefId = StringId<TemplateRef>;
@@ -46,206 +45,6 @@ namespace Horde.Build.Streams
 		}
 	}
 
-	/// <summary>
-	/// Mapping from a BuildGraph agent type to a set of machines on the farm
-	/// </summary>
-	public class AgentType
-	{
-		/// <summary>
-		/// Name of the pool of agents to use
-		/// </summary>
-		[BsonRequired]
-		public PoolId Pool { get; set; }
-
-		/// <summary>
-		/// Name of the workspace to execute on
-		/// </summary>
-		public string? Workspace { get; set; }
-
-		/// <summary>
-		/// Path to the temporary storage dir
-		/// </summary>
-		public string? TempStorageDir { get; set; }
-
-		/// <summary>
-		/// Environment variables to be set when executing the job
-		/// </summary>
-		public Dictionary<string, string>? Environment { get; set; }
-
-		/// <summary>
-		/// Default constructor
-		/// </summary>
-		[BsonConstructor]
-		private AgentType()
-		{
-		}
-
-		/// <summary>
-		/// Constructor
-		/// </summary>
-		/// <param name="pool">The pool for this agent</param>
-		/// <param name="workspace">Name of the workspace to use</param>
-		/// <param name="tempStorageDir">Path to the temp storage directory</param>
-		public AgentType(PoolId pool, string workspace, string? tempStorageDir)
-		{
-			Pool = pool;
-			Workspace = workspace;
-			TempStorageDir = tempStorageDir;
-		}
-
-		/// <summary>
-		/// Constructor
-		/// </summary>
-		/// <param name="request">The object to construct from</param>
-		public AgentType(AgentConfig request)
-		{
-			Pool = new PoolId(request.Pool);
-			Workspace = request.Workspace;
-			TempStorageDir = request.TempStorageDir;
-			Environment = request.Environment;
-		}
-
-		/// <summary>
-		/// Constructs an AgentType object from an optional request
-		/// </summary>
-		/// <param name="request">The request object</param>
-		/// <returns>New agent type object</returns>
-		[return: NotNullIfNotNull("request")]
-		public static AgentType? FromRequest(AgentConfig? request)
-		{
-			return (request != null) ? new AgentType(request) : null;
-		}
-
-		/// <summary>
-		/// Creates an API response object from this stream
-		/// </summary>
-		/// <returns>The response object</returns>
-		public GetAgentTypeResponse ToApiResponse()
-		{
-			return new GetAgentTypeResponse(Pool.ToString(), Workspace, TempStorageDir, Environment);
-		}
-
-		/// <summary>
-		/// Creates an API response object from this stream
-		/// </summary>
-		/// <returns>The response object</returns>
-		public HordeCommon.Rpc.GetAgentTypeResponse ToRpcResponse()
-		{
-			HordeCommon.Rpc.GetAgentTypeResponse response = new HordeCommon.Rpc.GetAgentTypeResponse();
-			if (TempStorageDir != null)
-			{
-				response.TempStorageDir = TempStorageDir;
-			}
-			if (Environment != null)
-			{
-				response.Environment.Add(Environment);
-			}
-			return response;
-		}
-	}
-
-	/// <summary>
-	/// Information about a workspace type
-	/// </summary>
-	public class WorkspaceType
-	{
-		/// <summary>
-		/// Name of the Perforce cluster to use
-		/// </summary>
-		[BsonIgnoreIfNull]
-		public string? Cluster { get; set; }
-
-		/// <summary>
-		/// The Perforce server and port
-		/// </summary>
-		[BsonIgnoreIfNull]
-		public string? ServerAndPort { get; set; }
-
-		/// <summary>
-		/// The Perforce username for syncing this workspace
-		/// </summary>
-		[BsonIgnoreIfNull]
-		public string? UserName { get; set; }
-
-		/// <summary>
-		/// The Perforce password for syncing this workspace
-		/// </summary>
-		[BsonIgnoreIfNull]
-		public string? Password { get; set; }
-
-		/// <summary>
-		/// Identifier to distinguish this workspace from other workspaces. Defaults to the workspace type name.
-		/// </summary>
-		[BsonIgnoreIfNull]
-		public string? Identifier { get; set; }
-
-		/// <summary>
-		/// Override for the stream to sync
-		/// </summary>
-		[BsonIgnoreIfNull]
-		public string? Stream { get; set; }
-
-		/// <summary>
-		/// Custom view for the workspace
-		/// </summary>
-		[BsonIgnoreIfNull]
-		public List<string>? View { get; set; }
-
-		/// <summary>
-		/// Whether to use an incrementally synced workspace
-		/// </summary>
-		public bool Incremental { get; set; }
-
-		/// <summary>
-		/// Whether to use the AutoSDK
-		/// </summary>
-		public bool UseAutoSdk { get; set; } = true;
-
-		/// <summary>
-		/// Default constructor
-		/// </summary>
-		public WorkspaceType()
-		{
-		}
-
-		/// <summary>
-		/// Constructor
-		/// </summary>
-		/// <param name="request">The object to construct from</param>
-		public WorkspaceType(WorkspaceConfig request)
-		{
-			Cluster = request.Cluster;
-			ServerAndPort = request.ServerAndPort;
-			UserName = request.UserName;
-			Password = request.Password;
-			Identifier = request.Identifier;
-			Stream = request.Stream;
-			View = request.View;
-			Incremental = request.Incremental;
-			UseAutoSdk = request.UseAutoSdk;
-		}
-
-		/// <summary>
-		/// Constructs an AgentType object from an optional request
-		/// </summary>
-		/// <param name="request">The request object</param>
-		/// <returns>New agent type object</returns>
-		[return: NotNullIfNotNull("request")]
-		public static WorkspaceType? FromRequest(WorkspaceConfig? request)
-		{
-			return (request != null) ? new WorkspaceType(request) : null;
-		}
-
-		/// <summary>
-		/// Creates an API response object from this stream
-		/// </summary>
-		/// <returns>The response object</returns>
-		public GetWorkspaceTypeResponse ToApiResponse()
-		{
-			return new GetWorkspaceTypeResponse(Cluster, ServerAndPort, UserName, Identifier, Stream, View, Incremental, UseAutoSdk);
-		}
-	}
-
 	/// <summary>
 	/// Allows triggering another downstream job on succesful completion of a step or aggregate
 	/// </summary>
@@ -528,27 +327,6 @@ namespace Horde.Build.Streams
 		}
 	}
 
-	/// <summary>
-	/// How to replicate data for this stream
-	/// </summary>
-	public enum ContentReplicationMode
-	{
-		/// <summary>
-		/// No content will be replicated for this stream
-		/// </summary>
-		None,
-
-		/// <summary>
-		/// Only replicate depot path and revision data for each file
-		/// </summary>
-		RevisionsOnly,
-
-		/// <summary>
-		/// Replicate full stream contents to storage
-		/// </summary>
-		Full,
-	}
-
 	/// <summary>
 	/// Information about a stream
 	/// </summary>
@@ -569,11 +347,6 @@ namespace Horde.Build.Streams
 		/// </summary>
 		public string Name { get; }
 
-		/// <summary>
-		/// Name of the perforce cluster
-		/// </summary>
-		public string ClusterName { get; }
-
 		/// <summary>
 		/// The revision of config file used for this stream
 		/// </summary>
@@ -589,26 +362,6 @@ namespace Horde.Build.Streams
 		/// </summary>
 		public bool Deleted { get; }
 
-		/// <summary>
-		/// Order to display on the dashboard's drop-down list
-		/// </summary>
-		public int Order { get; }
-
-		/// <summary>
-		/// Notification channel for all jobs in this stream
-		/// </summary>
-		public string? NotificationChannel { get; }
-
-		/// <summary>
-		/// Notification channel filter for all jobs in this stream. Errors|Warnings|Success
-		/// </summary>
-		public string? NotificationChannelFilter { get; }
-
-		/// <summary>
-		/// Channel to post issue triage notifications
-		/// </summary>
-		public string? TriageChannel { get; }
-
 		/// <summary>
 		/// Default template to use for preflights
 		/// </summary>
@@ -619,16 +372,6 @@ namespace Horde.Build.Streams
 		/// </summary>
 		public IReadOnlyList<StreamTab> Tabs { get; }
 
-		/// <summary>
-		/// Dictionary of agent types
-		/// </summary>
-		public IReadOnlyDictionary<string, AgentType> AgentTypes { get; }
-
-		/// <summary>
-		/// Dictionary of workspace types
-		/// </summary>
-		public IReadOnlyDictionary<string, WorkspaceType> WorkspaceTypes { get; }
-
 		/// <summary>
 		/// List of templates available for this stream
 		/// </summary>
@@ -644,16 +387,6 @@ namespace Horde.Build.Streams
 		/// </summary>
 		public string? PauseComment { get; }
 
-		/// <summary>
-		/// Whether to replicate data from Perforce to Horde Storage.
-		/// </summary>
-		public ContentReplicationMode ReplicationMode { get; }
-
-		/// <summary>
-		/// Filter for paths to be replicated. May be removed in future.
-		/// </summary>
-		public string? ReplicationFilter { get; }
-
 		/// <summary>
 		/// The ACL for this object
 		/// </summary>
@@ -672,7 +405,7 @@ namespace Horde.Build.Streams
 		/// <param name="agentType">The agent type</param>
 		/// <param name="workspace">Receives the agent workspace definition</param>
 		/// <returns>True if the agent type was valid, and an agent workspace could be created</returns>
-		public static bool TryGetAgentWorkspace(this IStream stream, AgentType agentType, [NotNullWhen(true)] out (AgentWorkspace, bool)? workspace)
+		public static bool TryGetAgentWorkspace(this IStream stream, AgentConfig agentType, [NotNullWhen(true)] out (AgentWorkspace, bool)? workspace)
 		{
 			// Get the workspace settings
 			if (agentType.Workspace == null)
@@ -684,8 +417,8 @@ namespace Horde.Build.Streams
 			else
 			{
 				// Try to get the matching workspace type
-				WorkspaceType? workspaceType;
-				if (!stream.WorkspaceTypes.TryGetValue(agentType.Workspace, out workspaceType))
+				WorkspaceConfig? workspaceType;
+				if (!stream.Config.WorkspaceTypes.TryGetValue(agentType.Workspace, out workspaceType))
 				{
 					workspace = null;
 					return false;
@@ -742,10 +475,10 @@ namespace Horde.Build.Streams
 		public static GetStreamResponse ToApiResponse(this IStream stream, bool bIncludeAcl, List<GetTemplateRefResponse> apiTemplateRefs)
 		{
 			List<GetStreamTabResponse> apiTabs = stream.Tabs.ConvertAll(x => x.ToResponse());
-			Dictionary<string, GetAgentTypeResponse> apiAgentTypes = stream.AgentTypes.ToDictionary(x => x.Key, x => x.Value.ToApiResponse());
-			Dictionary<string, GetWorkspaceTypeResponse> apiWorkspaceTypes = stream.WorkspaceTypes.ToDictionary(x => x.Key, x => x.Value.ToApiResponse());
+			Dictionary<string, AgentConfig> apiAgentTypes = stream.Config.AgentTypes.ToDictionary(x => x.Key, x => x.Value);
+			Dictionary<string, WorkspaceConfig> apiWorkspaceTypes = stream.Config.WorkspaceTypes.ToDictionary(x => x.Key, x => x.Value);
 			GetAclResponse? apiAcl = (bIncludeAcl && stream.Acl != null)? new GetAclResponse(stream.Acl) : null;
-			return new GetStreamResponse(stream.Id.ToString(), stream.ProjectId.ToString(), stream.Name, stream.ConfigRevision, stream.Order, stream.NotificationChannel, stream.NotificationChannelFilter, stream.TriageChannel, stream.DefaultPreflight?.ToRequest(), apiTabs, apiAgentTypes, apiWorkspaceTypes, apiTemplateRefs, apiAcl, stream.PausedUntil, stream.PauseComment, stream.Config.Workflows);
+			return new GetStreamResponse(stream.Id.ToString(), stream.ProjectId.ToString(), stream.Name, stream.ConfigRevision, stream.Config.Order, stream.Config.NotificationChannel, stream.Config.NotificationChannelFilter, stream.Config.TriageChannel, stream.DefaultPreflight?.ToRequest(), apiTabs, apiAgentTypes, apiWorkspaceTypes, apiTemplateRefs, apiAcl, stream.PausedUntil, stream.PauseComment, stream.Config.Workflows);
 		}
 
 		/// <summary>
@@ -757,7 +490,7 @@ namespace Horde.Build.Streams
 		{
 			HordeCommon.Rpc.GetStreamResponse response = new HordeCommon.Rpc.GetStreamResponse();
 			response.Name = stream.Name;
-			response.AgentTypes.Add(stream.AgentTypes.ToDictionary(x => x.Key, x => x.Value.ToRpcResponse()));
+			response.AgentTypes.Add(stream.Config.AgentTypes.ToDictionary(x => x.Key, x => x.Value.ToRpcResponse()));
 			return response;
 		}
 
diff --git a/Engine/Source/Programs/Horde/Horde.Build/Streams/StreamCollection.cs b/Engine/Source/Programs/Horde/Horde.Build/Streams/StreamCollection.cs
--- a/Engine/Source/Programs/Horde/Horde.Build/Streams/StreamCollection.cs
+++ b/Engine/Source/Programs/Horde/Horde.Build/Streams/StreamCollection.cs
@@ -38,63 +38,41 @@ namespace Horde.Build.Streams
 		/// </summary>
 		class StreamDocument : IStream
 		{
-			public const int DefaultOrder = 128;
-
 			[BsonRequired, BsonId]
 			public StreamId Id { get; set; }
 
 			[BsonRequired]
 			public ProjectId ProjectId { get; set; }
 
-			[BsonRequired]
-			public string Name { get; set; }
-
-			public string? ClusterName { get; set; }
 			public string ConfigRevision { get; set; } = String.Empty;
 
 			[BsonIgnore]
 			public StreamConfig? Config { get; set; }
 
-			public int Order { get; set; } = DefaultOrder;
-			public string? NotificationChannel { get; set; }
-			public string? NotificationChannelFilter { get; set; }
-			public string? TriageChannel { get; set; }
 			public DefaultPreflight? DefaultPreflight { get; set; }
 			public List<StreamTab> Tabs { get; set; } = new List<StreamTab>();
-			public Dictionary<string, AgentType> AgentTypes { get; set; } = new Dictionary<string, AgentType>(StringComparer.Ordinal);
-			public Dictionary<string, WorkspaceType> WorkspaceTypes { get; set; } = new Dictionary<string, WorkspaceType>(StringComparer.Ordinal);
 			public Dictionary<TemplateRefId, TemplateRef> Templates { get; set; } = new Dictionary<TemplateRefId, TemplateRef>();
 			public DateTime? PausedUntil { get; set; }
 			public string? PauseComment { get; set; }
 
-			[BsonIgnoreIfDefault]
-			[BsonDefaultValue(ContentReplicationMode.None)]
-			public ContentReplicationMode ReplicationMode { get; set; }
-
-			[BsonIgnoreIfNull]
-			public string? ReplicationFilter { get; set; }
-
 			public Acl? Acl { get; set; }
 			public int UpdateIndex { get; set; }
 			public bool Deleted { get; set; }
 
 			StreamConfig IStream.Config => Config!;
-			string IStream.ClusterName => ClusterName ?? PerforceCluster.DefaultName;
 			IReadOnlyList<StreamTab> IStream.Tabs => Tabs;
-			IReadOnlyDictionary<string, AgentType> IStream.AgentTypes => AgentTypes;
-			IReadOnlyDictionary<string, WorkspaceType> IStream.WorkspaceTypes => WorkspaceTypes;
 			IReadOnlyDictionary<TemplateRefId, TemplateRef> IStream.Templates => Templates;
 
+			string IStream.Name => Config!.Name;
+
 			[BsonConstructor]
 			private StreamDocument()
 			{
-				Name = null!;
 			}
 
-			public StreamDocument(StreamId id, string name, ProjectId projectId)
+			public StreamDocument(StreamId id, ProjectId projectId)
 			{
 				Id = id;
-				Name = name;
 				ProjectId = projectId;
 			}
 		}
@@ -162,34 +140,22 @@ namespace Horde.Build.Streams
 			List<StreamTab> tabs = config.Tabs.ConvertAll(x => StreamTab.FromRequest(x));
 			Dictionary<TemplateRefId, TemplateRef> templateRefs = await CreateTemplateRefsAsync(config.Templates, stream, _templateCollection);
 
-			Dictionary<string, AgentType> agentTypes = new Dictionary<string, AgentType>();
-			if (config.AgentTypes != null)
-			{
-				agentTypes = config.AgentTypes.Where(x => x.Value != null).ToDictionary(x => x.Key, x => new AgentType(x.Value!));
-			}
-
-			Dictionary<string, WorkspaceType> workspaceTypes = new Dictionary<string, WorkspaceType>();
-			if (config.WorkspaceTypes != null)
-			{
-				workspaceTypes = config.WorkspaceTypes.Where(x => x.Value != null).ToDictionary(x => x.Key, x => new WorkspaceType(x.Value!));
-			}
-
 			DefaultPreflight? defaultPreflight = config.DefaultPreflight?.ToModel();
 			if (defaultPreflight == null && config.DefaultPreflightTemplate != null)
 			{
 				defaultPreflight = new DefaultPreflight(new TemplateRefId(config.DefaultPreflightTemplate), null);
 			}
 
-			Validate(id, defaultPreflight, templateRefs, tabs, agentTypes, workspaceTypes);
+			Validate(id, defaultPreflight, templateRefs, tabs, config);
 
 			Acl? acl = Acl.Merge(new Acl(), config.Acl);
 			if (stream == null)
 			{
-				return await TryCreateAsync(id, projectId, revision, config, defaultPreflight, tabs, agentTypes, workspaceTypes, templateRefs, acl);
+				return await TryCreateAsync(id, projectId, revision, config, defaultPreflight, tabs, templateRefs, acl);
 			}
 			else
 			{
-				return await TryReplaceAsync(stream, projectId, revision, config, defaultPreflight, tabs, agentTypes, workspaceTypes, templateRefs, acl);
+				return await TryReplaceAsync(stream, projectId, revision, config, defaultPreflight, tabs, templateRefs, acl);
 			}
 		}
 
@@ -352,23 +318,14 @@ namespace Horde.Build.Streams
 		}
 
 		/// <inheritdoc/>
-		async Task<IStream?> TryCreateAsync(StreamId id, ProjectId projectId, string configRevision, StreamConfig config, DefaultPreflight? defaultPreflight, List<StreamTab> tabs, Dictionary<string, AgentType> agentTypes, Dictionary<string, WorkspaceType> workspaceTypes, Dictionary<TemplateRefId, TemplateRef> templateRefs, Acl? acl)
+		async Task<IStream?> TryCreateAsync(StreamId id, ProjectId projectId, string configRevision, StreamConfig config, DefaultPreflight? defaultPreflight, List<StreamTab> tabs, Dictionary<TemplateRefId, TemplateRef> templateRefs, Acl? acl)
 		{
-			StreamDocument newStream = new StreamDocument(id, config.Name, projectId);
-			newStream.ClusterName = config.ClusterName;
+			StreamDocument newStream = new StreamDocument(id, projectId);
 			newStream.ConfigRevision = configRevision;
 			newStream.Config = config;
-			newStream.Order = config.Order ?? StreamDocument.DefaultOrder;
-			newStream.NotificationChannel = config.NotificationChannel;
-			newStream.NotificationChannelFilter = config.NotificationChannelFilter;
-			newStream.TriageChannel = config.TriageChannel;
 			newStream.DefaultPreflight = defaultPreflight;
 			newStream.Tabs = tabs;
-			newStream.AgentTypes = agentTypes;
-			newStream.WorkspaceTypes = workspaceTypes;
 			newStream.Templates = templateRefs;
-			newStream.ReplicationMode = config.ReplicationMode;
-			newStream.ReplicationFilter = config.ReplicationFilter;
 			newStream.Acl = acl;
 
 			try
@@ -390,30 +347,20 @@ namespace Horde.Build.Streams
 		}
 
 		/// <inheritdoc/>
-		async Task<IStream?> TryReplaceAsync(IStream streamInterface, ProjectId projectId, string configRevision, StreamConfig config, DefaultPreflight? defaultPreflight, List<StreamTab> tabs, Dictionary<string, AgentType>? agentTypes, Dictionary<string, WorkspaceType>? workspaceTypes, Dictionary<TemplateRefId, TemplateRef>? templateRefs, Acl? acl)
+		async Task<IStream?> TryReplaceAsync(IStream streamInterface, ProjectId projectId, string configRevision, StreamConfig config, DefaultPreflight? defaultPreflight, List<StreamTab> tabs, Dictionary<TemplateRefId, TemplateRef>? templateRefs, Acl? acl)
 		{
-			int order = config.Order ?? StreamDocument.DefaultOrder;
+			int order = config.Order;
 
 			StreamDocument stream = (StreamDocument)streamInterface;
 
 			UpdateDefinitionBuilder<StreamDocument> updateBuilder = Builders<StreamDocument>.Update;
 
 			List<UpdateDefinition<StreamDocument>> updates = new List<UpdateDefinition<StreamDocument>>();
-			updates.Add(updateBuilder.Set(x => x.Name, config.Name));
 			updates.Add(updateBuilder.Set(x => x.ProjectId, projectId));
-			updates.Add(updateBuilder.Set(x => x.ClusterName, config.ClusterName));
 			updates.Add(updateBuilder.Set(x => x.ConfigRevision, configRevision));
-			updates.Add(updateBuilder.Set(x => x.Order, order));
-			updates.Add(updateBuilder.Set(x => x.NotificationChannel, config.NotificationChannel));
-			updates.Add(updateBuilder.Set(x => x.NotificationChannelFilter, config.NotificationChannelFilter));
-			updates.Add(updateBuilder.Set(x => x.TriageChannel, config.TriageChannel));
 			updates.Add(updateBuilder.Set(x => x.DefaultPreflight, defaultPreflight));
 			updates.Add(updateBuilder.Set(x => x.Tabs, tabs ?? new List<StreamTab>()));
-			updates.Add(updateBuilder.Set(x => x.AgentTypes, agentTypes ?? new Dictionary<string, AgentType>()));
-			updates.Add(updateBuilder.Set(x => x.WorkspaceTypes, workspaceTypes ?? new Dictionary<string, WorkspaceType>()));
 			updates.Add(updateBuilder.Set(x => x.Templates, templateRefs ?? new Dictionary<TemplateRefId, TemplateRef>()));
-			updates.Add(updateBuilder.Set(x => x.ReplicationMode, config.ReplicationMode));
-			updates.Add(updateBuilder.SetOrUnsetNullRef(x => x.ReplicationFilter, config.ReplicationFilter));
 			updates.Add(updateBuilder.SetOrUnsetNullRef(x => x.Acl, acl));
 			updates.Add(updateBuilder.Unset(x => x.Deleted));
 
@@ -538,7 +485,7 @@ namespace Horde.Build.Streams
 		/// <summary>
 		/// Checks the stream definition for consistency
 		/// </summary>
-		public static void Validate(StreamId streamId, DefaultPreflight? defaultPreflight, IReadOnlyDictionary<TemplateRefId, TemplateRef> templates, IReadOnlyList<StreamTab> tabs, IReadOnlyDictionary<string, AgentType> agentTypes, IReadOnlyDictionary<string, WorkspaceType> workspaceTypes)
+		public static void Validate(StreamId streamId, DefaultPreflight? defaultPreflight, IReadOnlyDictionary<TemplateRefId, TemplateRef> templates, IReadOnlyList<StreamTab> tabs, StreamConfig config)
 		{
 			// Check the default preflight template is valid
 			if (defaultPreflight != null)
@@ -564,10 +511,10 @@ namespace Horde.Build.Streams
 			}
 
 			// Check that all the agent types reference valid workspace names
-			foreach (KeyValuePair<string, AgentType> pair in agentTypes)
+			foreach (KeyValuePair<string, AgentConfig> pair in config.AgentTypes)
 			{
 				string? workspaceTypeName = pair.Value.Workspace;
-				if (workspaceTypeName != null && !workspaceTypes.ContainsKey(workspaceTypeName))
+				if (workspaceTypeName != null && !config.WorkspaceTypes.ContainsKey(workspaceTypeName))
 				{
 					throw new InvalidStreamException($"Agent type '{pair.Key}' references undefined workspace type '{pair.Value.Workspace}' in {streamId}");
 				}
diff --git a/Engine/Source/Programs/Horde/Horde.Build/Streams/StreamConfig.cs b/Engine/Source/Programs/Horde/Horde.Build/Streams/StreamConfig.cs
--- a/Engine/Source/Programs/Horde/Horde.Build/Streams/StreamConfig.cs
+++ b/Engine/Source/Programs/Horde/Horde.Build/Streams/StreamConfig.cs
@@ -7,19 +7,43 @@ using System.Diagnostics.CodeAnalysis;
 using System.Linq;
 using EpicGames.Core;
 using Horde.Build.Acls;
+using Horde.Build.Agents.Pools;
 using Horde.Build.Issues;
 using Horde.Build.Jobs.Graphs;
 using Horde.Build.Jobs.Schedules;
 using Horde.Build.Jobs.Templates;
 using Horde.Build.Perforce;
+using Horde.Build.Server;
 using Horde.Build.Utilities;
 using HordeCommon;
 
 namespace Horde.Build.Streams
 {
+	using PoolId = StringId<IPool>;
 	using TemplateRefId = StringId<TemplateRef>;
 	using WorkflowId = StringId<WorkflowConfig>;
 
+	/// <summary>
+	/// How to replicate data for this stream
+	/// </summary>
+	public enum ContentReplicationMode
+	{
+		/// <summary>
+		/// No content will be replicated for this stream
+		/// </summary>
+		None,
+
+		/// <summary>
+		/// Only replicate depot path and revision data for each file
+		/// </summary>
+		RevisionsOnly,
+
+		/// <summary>
+		/// Replicate full stream contents to storage
+		/// </summary>
+		Full,
+	}
+
 	/// <summary>
 	/// Config for a stream
 	/// </summary>
@@ -36,12 +60,12 @@ namespace Horde.Build.Streams
 		/// <summary>
 		/// The perforce cluster containing the stream
 		/// </summary>
-		public string? ClusterName { get; set; }
+		public string ClusterName { get; set; } = PerforceCluster.DefaultName;
 
 		/// <summary>
 		/// Order for this stream
 		/// </summary>
-		public int? Order { get; set; }
+		public int Order { get; set; } = 128;
 
 		/// <summary>
 		/// Notification channel for all jobs in this stream
@@ -162,7 +186,7 @@ namespace Horde.Build.Streams
 		/// Pool of agents to use for this agent type
 		/// </summary>
 		[Required]
-		public string Pool { get; set; } = null!;
+		public PoolId Pool { get; set; }
 
 		/// <summary>
 		/// Name of the workspace to sync
@@ -178,6 +202,24 @@ namespace Horde.Build.Streams
 		/// Environment variables to be set when executing the job
 		/// </summary>
 		public Dictionary<string, string>? Environment { get; set; }
+
+		/// <summary>
+		/// Creates an API response object from this stream
+		/// </summary>
+		/// <returns>The response object</returns>
+		public HordeCommon.Rpc.GetAgentTypeResponse ToRpcResponse()
+		{
+			HordeCommon.Rpc.GetAgentTypeResponse response = new HordeCommon.Rpc.GetAgentTypeResponse();
+			if (TempStorageDir != null)
+			{
+				response.TempStorageDir = TempStorageDir;
+			}
+			if (Environment != null)
+			{
+				response.Environment.Add(Environment);
+			}
+			return response;
+		}
 	}
 
 	/// <summary>
diff --git a/Engine/Source/Programs/Horde/Horde.Build/Streams/StreamMessages.cs b/Engine/Source/Programs/Horde/Horde.Build/Streams/StreamMessages.cs
--- a/Engine/Source/Programs/Horde/Horde.Build/Streams/StreamMessages.cs
+++ b/Engine/Source/Programs/Horde/Horde.Build/Streams/StreamMessages.cs
@@ -379,12 +379,12 @@ namespace Horde.Build.Streams
 		/// <summary>
 		/// Map of agent name to type
 		/// </summary>
-		public Dictionary<string, GetAgentTypeResponse> AgentTypes { get; set; }
+		public Dictionary<string, AgentConfig> AgentTypes { get; set; }
 
 		/// <summary>
 		/// Map of workspace name to type
 		/// </summary>
-		public Dictionary<string, GetWorkspaceTypeResponse>? WorkspaceTypes { get; set; }
+		public Dictionary<string, WorkspaceConfig>? WorkspaceTypes { get; set; }
 
 		/// <summary>
 		/// Templates for jobs in this stream
@@ -431,7 +431,7 @@ namespace Horde.Build.Streams
 		/// <param name="pausedUntil">Stream paused for new builds until this date</param>
 		/// <param name="pauseComment">Reason for stream being paused</param>
 		/// <param name="workflows">Workflows for this stream</param>
-		public GetStreamResponse(string id, string projectId, string name, string configRevision, int order, string? notificationChannel, string? notificationChannelFilter, string? triageChannel, DefaultPreflightConfig? defaultPreflight, List<GetStreamTabResponse> tabs, Dictionary<string, GetAgentTypeResponse> agentTypes, Dictionary<string, GetWorkspaceTypeResponse>? workspaceTypes, List<GetTemplateRefResponse> templates, GetAclResponse? acl, DateTime? pausedUntil, string? pauseComment, List<WorkflowConfig> workflows)
+		public GetStreamResponse(string id, string projectId, string name, string configRevision, int order, string? notificationChannel, string? notificationChannelFilter, string? triageChannel, DefaultPreflightConfig? defaultPreflight, List<GetStreamTabResponse> tabs, Dictionary<string, AgentConfig> agentTypes, Dictionary<string, WorkspaceConfig>? workspaceTypes, List<GetTemplateRefResponse> templates, GetAclResponse? acl, DateTime? pausedUntil, string? pauseComment, List<WorkflowConfig> workflows)
 		{
 			Id = id;
 			ProjectId = projectId;
