Horde: Remove redundant copy of default preflight information from stream document.

diff --git a/Engine/Source/Programs/Horde/Horde.Build/Jobs/JobMessages.cs b/Engine/Source/Programs/Horde/Horde.Build/Jobs/JobMessages.cs
--- a/Engine/Source/Programs/Horde/Horde.Build/Jobs/JobMessages.cs
+++ b/Engine/Source/Programs/Horde/Horde.Build/Jobs/JobMessages.cs
@@ -10,10 +10,13 @@ using Horde.Build.Acls;
 using Horde.Build.Jobs.Graphs;
 using Horde.Build.Streams;
 using Horde.Build.Users;
+using Horde.Build.Utilities;
 using HordeCommon;
 
 namespace Horde.Build.Jobs
 {
+	using TemplateRefId = StringId<TemplateRef>;
+
 	/// <summary>
 	/// State of the job
 	/// </summary>
@@ -50,7 +53,7 @@ namespace Horde.Build.Jobs
 		/// The template for this job
 		/// </summary>
 		[Required]
-		public string TemplateId { get; set; }
+		public TemplateRefId TemplateId { get; set; }
 
 		/// <summary>
 		/// Name of the job
@@ -95,7 +98,7 @@ namespace Horde.Build.Jobs
 		/// <summary>
 		/// Private constructor for serialization
 		/// </summary>
-		public CreateJobRequest(string streamId, string templateId)
+		public CreateJobRequest(string streamId, TemplateRefId templateId)
 		{
 			StreamId = streamId;
 			TemplateId = templateId;
diff --git a/Engine/Source/Programs/Horde/Horde.Build/Jobs/JobsController.cs b/Engine/Source/Programs/Horde/Horde.Build/Jobs/JobsController.cs
--- a/Engine/Source/Programs/Horde/Horde.Build/Jobs/JobsController.cs
+++ b/Engine/Source/Programs/Horde/Horde.Build/Jobs/JobsController.cs
@@ -88,7 +88,7 @@ namespace Horde.Build.Jobs
 			}
 
 			// Get the name of the template ref
-			TemplateRefId templateRefId = new TemplateRefId(create.TemplateId);
+			TemplateRefId templateRefId = create.TemplateId;
 
 			// Augment the request with template properties
 			TemplateRef? templateRef;
@@ -113,7 +113,7 @@ namespace Horde.Build.Jobs
 
 			// Get the name of the new job
 			string name = create.Name ?? template.Name;
-			if (create.TemplateId.Equals("stage-to-marketplace", StringComparison.Ordinal) && create.Arguments != null)
+			if (create.TemplateId == new TemplateRefId("stage-to-marketplace") && create.Arguments != null)
 			{
 				foreach (string argument in create.Arguments)
 				{
@@ -140,7 +140,7 @@ namespace Horde.Build.Jobs
 			}
 			else if (create.ChangeQuery != null)
 			{
-				change = await ExecuteChangeQueryAsync(stream, new TemplateRefId(create.ChangeQuery.TemplateId ?? create.TemplateId), create.ChangeQuery.Target, create.ChangeQuery.Outcomes ?? new List<JobStepOutcome> { JobStepOutcome.Success });
+				change = await ExecuteChangeQueryAsync(stream, create.ChangeQuery.TemplateId ?? create.TemplateId, create.ChangeQuery.Target, create.ChangeQuery.Outcomes ?? new List<JobStepOutcome> { JobStepOutcome.Success });
 			}
 			else if (create.PreflightChange == null && template.SubmitNewChange != null)
 			{
diff --git a/Engine/Source/Programs/Horde/Horde.Build/Streams/IStream.cs b/Engine/Source/Programs/Horde/Horde.Build/Streams/IStream.cs
--- a/Engine/Source/Programs/Horde/Horde.Build/Streams/IStream.cs
+++ b/Engine/Source/Programs/Horde/Horde.Build/Streams/IStream.cs
@@ -227,90 +227,6 @@ namespace Horde.Build.Streams
 		}
 	}
 
-	/// <summary>
-	/// Query used to identify a base changelist for a preflight
-	/// </summary>
-	public class ChangeQuery
-	{
-		/// <summary>
-		/// Template to search for
-		/// </summary>
-		public TemplateRefId? TemplateRefId { get; set; }
-
-		/// <summary>
-		/// The target to look at the status for
-		/// </summary>
-		public string? Target { get; set; }
-
-		/// <summary>
-		/// Whether to match a job that contains warnings
-		/// </summary>
-		public List<JobStepOutcome>? Outcomes { get; set; }
-
-		/// <summary>
-		/// Convert to a request object
-		/// </summary>
-		/// <returns></returns>
-		public ChangeQueryConfig ToRequest()
-		{
-			return new ChangeQueryConfig { TemplateId = TemplateRefId?.ToString(), Target = Target, Outcomes = Outcomes };
-		}
-	}
-
-	/// <summary>
-	/// Definition of a query to execute to find the changelist to run a build at
-	/// </summary>
-	public class DefaultPreflight
-	{
-		/// <summary>
-		/// The template id to execute
-		/// </summary>
-		public TemplateRefId? TemplateRefId { get; set; }
-
-		/// <summary>
-		/// Query specifying a changelist to use
-		/// </summary>
-		public ChangeQuery? Change { get; set; }
-
-		/// <summary>
-		/// The job type to query for the change to use
-		/// </summary>
-		[Obsolete("Use Change.TemplateRefId instead")]
-		public TemplateRefId? ChangeTemplateRefId { get; set; }
-
-		/// <summary>
-		/// Constructor
-		/// </summary>
-		/// <param name="templateRefId"></param>
-		/// <param name="change">The job type to query for the change to use</param>
-		public DefaultPreflight(TemplateRefId? templateRefId, ChangeQuery? change)
-		{
-			TemplateRefId = templateRefId;
-			Change = change;
-		}
-
-		/// <summary>
-		/// Convert to a request object
-		/// </summary>
-		/// <returns></returns>
-		public DefaultPreflightConfig ToRequest()
-		{
-#pragma warning disable CS0618 // Type or member is obsolete
-			ChangeQueryConfig? changeRequest = null;
-			if (Change != null)
-			{
-				changeRequest = Change.ToRequest();
-			}
-			else if (ChangeTemplateRefId != null)
-			{
-				changeRequest = new ChangeQueryConfig { TemplateId = ChangeTemplateRefId.ToString() };
-			}
-
-			return new DefaultPreflightConfig { TemplateId = TemplateRefId?.ToString(), Change = changeRequest, ChangeTemplateId = changeRequest?.TemplateId };
-#pragma warning restore CS0618 // Type or member is obsolete
-		}
-	}
-
 	/// <summary>
 	/// Extension methods for template refs
 	/// </summary>
@@ -362,11 +278,6 @@ namespace Horde.Build.Streams
 		/// </summary>
 		public bool Deleted { get; }
 
-		/// <summary>
-		/// Default template to use for preflights
-		/// </summary>
-		public DefaultPreflight? DefaultPreflight { get; }
-
 		/// <summary>
 		/// List of pages to display in the dashboard
 		/// </summary>
@@ -478,7 +389,7 @@ namespace Horde.Build.Streams
 			Dictionary<string, AgentConfig> apiAgentTypes = stream.Config.AgentTypes.ToDictionary(x => x.Key, x => x.Value);
 			Dictionary<string, WorkspaceConfig> apiWorkspaceTypes = stream.Config.WorkspaceTypes.ToDictionary(x => x.Key, x => x.Value);
 			GetAclResponse? apiAcl = (bIncludeAcl && stream.Acl != null)? new GetAclResponse(stream.Acl) : null;
-			return new GetStreamResponse(stream.Id.ToString(), stream.ProjectId.ToString(), stream.Name, stream.ConfigRevision, stream.Config.Order, stream.Config.NotificationChannel, stream.Config.NotificationChannelFilter, stream.Config.TriageChannel, stream.DefaultPreflight?.ToRequest(), apiTabs, apiAgentTypes, apiWorkspaceTypes, apiTemplateRefs, apiAcl, stream.PausedUntil, stream.PauseComment, stream.Config.Workflows);
+			return new GetStreamResponse(stream.Id.ToString(), stream.ProjectId.ToString(), stream.Name, stream.ConfigRevision, stream.Config.Order, stream.Config.NotificationChannel, stream.Config.NotificationChannelFilter, stream.Config.TriageChannel, stream.Config.DefaultPreflight, apiTabs, apiAgentTypes, apiWorkspaceTypes, apiTemplateRefs, apiAcl, stream.PausedUntil, stream.PauseComment, stream.Config.Workflows);
 		}
 
 		/// <summary>
diff --git a/Engine/Source/Programs/Horde/Horde.Build/Streams/StreamCollection.cs b/Engine/Source/Programs/Horde/Horde.Build/Streams/StreamCollection.cs
--- a/Engine/Source/Programs/Horde/Horde.Build/Streams/StreamCollection.cs
+++ b/Engine/Source/Programs/Horde/Horde.Build/Streams/StreamCollection.cs
@@ -49,7 +49,6 @@ namespace Horde.Build.Streams
 			[BsonIgnore]
 			public StreamConfig? Config { get; set; }
 
-			public DefaultPreflight? DefaultPreflight { get; set; }
 			public List<StreamTab> Tabs { get; set; } = new List<StreamTab>();
 			public Dictionary<TemplateRefId, TemplateRef> Templates { get; set; } = new Dictionary<TemplateRefId, TemplateRef>();
 			public DateTime? PausedUntil { get; set; }
@@ -140,22 +139,16 @@ namespace Horde.Build.Streams
 			List<StreamTab> tabs = config.Tabs.ConvertAll(x => StreamTab.FromRequest(x));
 			Dictionary<TemplateRefId, TemplateRef> templateRefs = await CreateTemplateRefsAsync(config.Templates, stream, _templateCollection);
 
-			DefaultPreflight? defaultPreflight = config.DefaultPreflight?.ToModel();
-			if (defaultPreflight == null && config.DefaultPreflightTemplate != null)
-			{
-				defaultPreflight = new DefaultPreflight(new TemplateRefId(config.DefaultPreflightTemplate), null);
-			}
-
-			Validate(id, defaultPreflight, templateRefs, tabs, config);
+			Validate(id, templateRefs, tabs, config);
 
 			Acl? acl = Acl.Merge(new Acl(), config.Acl);
 			if (stream == null)
 			{
-				return await TryCreateAsync(id, projectId, revision, config, defaultPreflight, tabs, templateRefs, acl);
+				return await TryCreateAsync(id, projectId, revision, config, tabs, templateRefs, acl);
 			}
 			else
 			{
-				return await TryReplaceAsync(stream, projectId, revision, config, defaultPreflight, tabs, templateRefs, acl);
+				return await TryReplaceAsync(stream, projectId, revision, config, tabs, templateRefs, acl);
 			}
 		}
 
@@ -318,12 +311,11 @@ namespace Horde.Build.Streams
 		}
 
 		/// <inheritdoc/>
-		async Task<IStream?> TryCreateAsync(StreamId id, ProjectId projectId, string configRevision, StreamConfig config, DefaultPreflight? defaultPreflight, List<StreamTab> tabs, Dictionary<TemplateRefId, TemplateRef> templateRefs, Acl? acl)
+		async Task<IStream?> TryCreateAsync(StreamId id, ProjectId projectId, string configRevision, StreamConfig config, List<StreamTab> tabs, Dictionary<TemplateRefId, TemplateRef> templateRefs, Acl? acl)
 		{
 			StreamDocument newStream = new StreamDocument(id, projectId);
 			newStream.ConfigRevision = configRevision;
 			newStream.Config = config;
-			newStream.DefaultPreflight = defaultPreflight;
 			newStream.Tabs = tabs;
 			newStream.Templates = templateRefs;
 			newStream.Acl = acl;
@@ -347,7 +339,7 @@ namespace Horde.Build.Streams
 		}
 
 		/// <inheritdoc/>
-		async Task<IStream?> TryReplaceAsync(IStream streamInterface, ProjectId projectId, string configRevision, StreamConfig config, DefaultPreflight? defaultPreflight, List<StreamTab> tabs, Dictionary<TemplateRefId, TemplateRef>? templateRefs, Acl? acl)
+		async Task<IStream?> TryReplaceAsync(IStream streamInterface, ProjectId projectId, string configRevision, StreamConfig config, List<StreamTab> tabs, Dictionary<TemplateRefId, TemplateRef>? templateRefs, Acl? acl)
 		{
 			int order = config.Order;
 
@@ -358,7 +350,6 @@ namespace Horde.Build.Streams
 			List<UpdateDefinition<StreamDocument>> updates = new List<UpdateDefinition<StreamDocument>>();
 			updates.Add(updateBuilder.Set(x => x.ProjectId, projectId));
 			updates.Add(updateBuilder.Set(x => x.ConfigRevision, configRevision));
-			updates.Add(updateBuilder.Set(x => x.DefaultPreflight, defaultPreflight));
 			updates.Add(updateBuilder.Set(x => x.Tabs, tabs ?? new List<StreamTab>()));
 			updates.Add(updateBuilder.Set(x => x.Templates, templateRefs ?? new Dictionary<TemplateRefId, TemplateRef>()));
 			updates.Add(updateBuilder.SetOrUnsetNullRef(x => x.Acl, acl));
@@ -485,14 +476,14 @@ namespace Horde.Build.Streams
 		/// <summary>
 		/// Checks the stream definition for consistency
 		/// </summary>
-		public static void Validate(StreamId streamId, DefaultPreflight? defaultPreflight, IReadOnlyDictionary<TemplateRefId, TemplateRef> templates, IReadOnlyList<StreamTab> tabs, StreamConfig config)
+		public static void Validate(StreamId streamId, IReadOnlyDictionary<TemplateRefId, TemplateRef> templates, IReadOnlyList<StreamTab> tabs, StreamConfig config)
 		{
 			// Check the default preflight template is valid
-			if (defaultPreflight != null)
+			if (config.DefaultPreflight != null)
 			{
-				if (defaultPreflight.TemplateRefId != null && !templates.ContainsKey(defaultPreflight.TemplateRefId.Value))
+				if (config.DefaultPreflight.TemplateId != null && !templates.ContainsKey(config.DefaultPreflight.TemplateId.Value))
 				{
-					throw new InvalidStreamException($"Default preflight template was listed as '{defaultPreflight.TemplateRefId.Value}', but no template was found by that name");
+					throw new InvalidStreamException($"Default preflight template was listed as '{config.DefaultPreflight.TemplateId.Value}', but no template was found by that name");
 				}
 			}
 
diff --git a/Engine/Source/Programs/Horde/Horde.Build/Streams/StreamConfig.cs b/Engine/Source/Programs/Horde/Horde.Build/Streams/StreamConfig.cs
--- a/Engine/Source/Programs/Horde/Horde.Build/Streams/StreamConfig.cs
+++ b/Engine/Source/Programs/Horde/Horde.Build/Streams/StreamConfig.cs
@@ -83,9 +83,14 @@ namespace Horde.Build.Streams
 		public string? TriageChannel { get; set; }
 
 		/// <summary>
-		/// Default template for running preflights
+		/// Legacy name for the default preflight template
 		/// </summary>
-		public string? DefaultPreflightTemplate { get; set; }
+		[Obsolete("Use DefaultPreflight instead")]
+		public string? DefaultPreflightTemplate
+		{
+			get => DefaultPreflight?.TemplateId?.ToString();
+			set => DefaultPreflight = (value == null)? null : new DefaultPreflightConfig { TemplateId = new TemplateRefId(value) };
+		}
 
 		/// <summary>
 		/// Default template for running preflights
@@ -281,7 +286,7 @@ namespace Horde.Build.Streams
 		/// <summary>
 		/// The template id to query
 		/// </summary>
-		public string? TemplateId { get; set; }
+		public TemplateRefId? TemplateId { get; set; }
 
 		/// <summary>
 		/// The target to query
@@ -292,22 +297,6 @@ namespace Horde.Build.Streams
 		/// Whether to match a job that produced warnings
 		/// </summary>
 		public List<JobStepOutcome>? Outcomes { get; set; }
-
-		/// <summary>
-		/// Convert to a model object
-		/// </summary>
-		/// <returns></returns>
-		public ChangeQuery ToModel()
-		{
-			ChangeQuery query = new ChangeQuery();
-			if (TemplateId != null)
-			{
-				query.TemplateRefId = new StringId<TemplateRef>(TemplateId);
-			}
-			query.Target = Target;
-			query.Outcomes = Outcomes;
-			return query;
-		}
 	}
 
 	/// <summary>
@@ -318,34 +307,12 @@ namespace Horde.Build.Streams
 		/// <summary>
 		/// The template id to query
 		/// </summary>
-		public string? TemplateId { get; set; }
-
-		/// <summary>
-		/// The last successful job type to use for the base changelist
-		/// </summary>
-		[Obsolete("Use Change.TemplateId instead")]
-		public string? ChangeTemplateId { get; set; }
+		public TemplateRefId? TemplateId { get; set; }
 
 		/// <summary>
 		/// Query for the change to use
 		/// </summary>
 		public ChangeQueryConfig? Change { get; set; }
-
-		/// <summary>
-		/// Convert to a model object
-		/// </summary>
-		/// <returns></returns>
-		public DefaultPreflight ToModel()
-		{
-#pragma warning disable CS0618 // Type or member is obsolete
-			if (ChangeTemplateId != null)
-			{
-				Change ??= new ChangeQueryConfig();
-				Change.TemplateId = ChangeTemplateId;
-			}
-			return new DefaultPreflight((TemplateId != null) ? (TemplateRefId?)new TemplateRefId(TemplateId) : null, Change?.ToModel());
-#pragma warning restore CS0618 // Type or member is obsolete
-		}
 	}
 
 	/// <summary>
diff --git a/Engine/Source/Programs/Horde/Horde.Build/Streams/StreamMessages.cs b/Engine/Source/Programs/Horde/Horde.Build/Streams/StreamMessages.cs
--- a/Engine/Source/Programs/Horde/Horde.Build/Streams/StreamMessages.cs
+++ b/Engine/Source/Programs/Horde/Horde.Build/Streams/StreamMessages.cs
@@ -441,7 +441,7 @@ namespace Horde.Build.Streams
 			NotificationChannel = notificationChannel;
 			NotificationChannelFilter = notificationChannelFilter;
 			TriageChannel = triageChannel;
-			DefaultPreflightTemplate = defaultPreflight?.TemplateId;
+			DefaultPreflightTemplate = defaultPreflight?.TemplateId?.ToString();
 			DefaultPreflight = defaultPreflight;
 			Tabs = tabs;
 			AgentTypes = agentTypes;
