Horde: Use config classes to configure stream tabs.

diff --git a/Engine/Source/Programs/Horde/Horde.Build.Tests/Fixture.cs b/Engine/Source/Programs/Horde/Horde.Build.Tests/Fixture.cs
--- a/Engine/Source/Programs/Horde/Horde.Build.Tests/Fixture.cs
+++ b/Engine/Source/Programs/Horde/Horde.Build.Tests/Fixture.cs
@@ -71,8 +71,8 @@ namespace Horde.Build.Tests
 			templates.Add(new TemplateRefConfig { Id = TemplateRefId1, Name = "Test Template" });
 			templates.Add(new TemplateRefConfig { Id = TemplateRefId2, Name = "Test Template" });
 
-			List<CreateStreamTabRequest> tabs = new List<CreateStreamTabRequest>();
-			tabs.Add(new CreateJobsTabRequest { Title = "foo", Templates = new List<TemplateRefId> { TemplateRefId1, TemplateRefId2 } });
+			List<TabConfig> tabs = new List<TabConfig>();
+			tabs.Add(new JobsTabConfig { Title = "foo", Templates = new List<TemplateRefId> { TemplateRefId1, TemplateRefId2 } });
 
 			Dictionary<string, AgentConfig> agentTypes = new()
 			{
diff --git a/Engine/Source/Programs/Horde/Horde.Build.Tests/IssueServiceTests.cs b/Engine/Source/Programs/Horde/Horde.Build.Tests/IssueServiceTests.cs
--- a/Engine/Source/Programs/Horde/Horde.Build.Tests/IssueServiceTests.cs
+++ b/Engine/Source/Programs/Horde/Horde.Build.Tests/IssueServiceTests.cs
@@ -122,7 +122,7 @@ namespace Horde.Build.Tests
 			string revision = $"config:{streamId}";
 
 			StreamConfig streamConfig = new StreamConfig { Name = streamName };
-			streamConfig.Tabs.Add(new CreateJobsTabRequest { Title = "General", Templates = new List<TemplateRefId> { new TemplateRefId("test-template") } });
+			streamConfig.Tabs.Add(new JobsTabConfig { Title = "General", Templates = new List<TemplateRefId> { new TemplateRefId("test-template") } });
 			streamConfig.Templates.Add(new TemplateRefConfig { Id = new TemplateRefId("test-template") });
 			await ConfigCollection.AddConfigAsync(revision, streamConfig);
 
diff --git a/Engine/Source/Programs/Horde/Horde.Build.Tests/JobServiceTests.cs b/Engine/Source/Programs/Horde/Horde.Build.Tests/JobServiceTests.cs
--- a/Engine/Source/Programs/Horde/Horde.Build.Tests/JobServiceTests.cs
+++ b/Engine/Source/Programs/Horde/Horde.Build.Tests/JobServiceTests.cs
@@ -46,7 +46,7 @@ namespace Horde.Build.Tests
 			StreamConfig streamConfig = new StreamConfig();
 			streamConfig.Templates.Add(new TemplateRefConfig { Id = templateRefId1, Name = "Test Template", ChainedJobs = new List<ChainedJobTemplateConfig> { new ChainedJobTemplateConfig { TemplateId = templateRefId2.ToString(), Trigger = "Setup Build" } } });
 			streamConfig.Templates.Add(new TemplateRefConfig { Id = templateRefId2, Name = "Test Template" });
-			streamConfig.Tabs.Add(new CreateJobsTabRequest { Title = "foo", Templates = new List<TemplateRefId> { templateRefId1, templateRefId2 } });
+			streamConfig.Tabs.Add(new JobsTabConfig { Title = "foo", Templates = new List<TemplateRefId> { templateRefId1, templateRefId2 } });
 
 			StreamId streamId = new StreamId("ue5-main");
 			IStream? stream = await StreamService.GetStreamAsync(streamId);
diff --git a/Engine/Source/Programs/Horde/Horde.Build.Tests/SchedulerTests.cs b/Engine/Source/Programs/Horde/Horde.Build.Tests/SchedulerTests.cs
--- a/Engine/Source/Programs/Horde/Horde.Build.Tests/SchedulerTests.cs
+++ b/Engine/Source/Programs/Horde/Horde.Build.Tests/SchedulerTests.cs
@@ -62,7 +62,7 @@ namespace Horde.Build.Tests
 
 			StreamConfig config = new StreamConfig();
 			config.Name = "//UE5/Main";
-			config.Tabs.Add(new CreateJobsTabRequest { Title = "foo", Templates = new List<TemplateRefId> { TemplateRefId } });
+			config.Tabs.Add(new JobsTabConfig { Title = "foo", Templates = new List<TemplateRefId> { TemplateRefId } });
 			config.Templates.Add(new TemplateRefConfig { Id = TemplateRefId, Name = "Test", Schedule = schedule });
 
 			return (await CreateOrReplaceStreamAsync(StreamId, stream, ProjectId, config))!; 
@@ -427,7 +427,7 @@ namespace Horde.Build.Tests
 
 			StreamConfig config = new StreamConfig();
 			config.Name = "//UE5/Main";
-			config.Tabs.Add(new CreateJobsTabRequest { Title = "foo", Templates = new List<TemplateRefId> { newTemplateRefId1, newTemplateRefId2 } });
+			config.Tabs.Add(new JobsTabConfig { Title = "foo", Templates = new List<TemplateRefId> { newTemplateRefId1, newTemplateRefId2 } });
 
 			stream = (await CreateOrReplaceStreamAsync(StreamId, stream, ProjectId, config))!;
 
@@ -507,7 +507,7 @@ namespace Horde.Build.Tests
 
 			StreamConfig config = new StreamConfig();
 			config.Name = "//UE5/Main";
-			config.Tabs.Add(new CreateJobsTabRequest { Title = "foo", Templates = new List<TemplateRefId> { newTemplateRefId1, newTemplateRefId2 } });
+			config.Tabs.Add(new JobsTabConfig { Title = "foo", Templates = new List<TemplateRefId> { newTemplateRefId1, newTemplateRefId2 } });
 			config.Templates.Add(newTemplate1);
 			config.Templates.Add(newTemplate2);
 
diff --git a/Engine/Source/Programs/Horde/Horde.Build/Jobs/Templates/TemplateCollection.cs b/Engine/Source/Programs/Horde/Horde.Build/Jobs/Templates/TemplateCollection.cs
--- a/Engine/Source/Programs/Horde/Horde.Build/Jobs/Templates/TemplateCollection.cs
+++ b/Engine/Source/Programs/Horde/Horde.Build/Jobs/Templates/TemplateCollection.cs
@@ -91,12 +91,6 @@ namespace Horde.Build.Jobs.Templates
 		/// <param name="mongoService">The database service singleton</param>
 		public TemplateCollection(MongoService mongoService)
 		{
-			// Ensure discriminator cannot be registered twice (throws exception). Can otherwise happen during unit tests.
-			if (BsonSerializer.LookupDiscriminatorConvention(typeof(JobsTabColumn)) == null)
-			{
-				BsonSerializer.RegisterDiscriminatorConvention(typeof(JobsTabColumn), new DefaultDiscriminatorConvention(typeof(JobsTabColumn), typeof(JobsTabLabelColumn)));	
-			}
-			
 			_templates = mongoService.GetCollection<TemplateDocument>("Templates");
 
 			MemoryCacheOptions options = new MemoryCacheOptions();
diff --git a/Engine/Source/Programs/Horde/Horde.Build/Notifications/Sinks/SlackNotificationSink.cs b/Engine/Source/Programs/Horde/Horde.Build/Notifications/Sinks/SlackNotificationSink.cs
--- a/Engine/Source/Programs/Horde/Horde.Build/Notifications/Sinks/SlackNotificationSink.cs
+++ b/Engine/Source/Programs/Horde/Horde.Build/Notifications/Sinks/SlackNotificationSink.cs
@@ -1601,7 +1601,7 @@ namespace Horde.Build.Notifications.Sinks
 				TemplateRefConfig? templateConfig;
 				if (stream.Config.TryGetTemplate(templateId, out templateConfig))
 				{
-					CreateJobsTabRequest? tab = stream.Config.Tabs.OfType<CreateJobsTabRequest>().FirstOrDefault(x => x.Templates != null && x.Templates.Contains(templateId));
+					JobsTabConfig? tab = stream.Config.Tabs.OfType<JobsTabConfig>().FirstOrDefault(x => x.Templates != null && x.Templates.Contains(templateId));
 					if (tab != null)
 					{
 						Uri templateUrl = new Uri(_settings.DashboardUrl, $"stream/{stream.Id}?tab={tab.Title}&template={templateId}");
diff --git a/Engine/Source/Programs/Horde/Horde.Build/Streams/IStream.cs b/Engine/Source/Programs/Horde/Horde.Build/Streams/IStream.cs
--- a/Engine/Source/Programs/Horde/Horde.Build/Streams/IStream.cs
+++ b/Engine/Source/Programs/Horde/Horde.Build/Streams/IStream.cs
@@ -278,11 +278,6 @@ namespace Horde.Build.Streams
 		/// </summary>
 		public bool Deleted { get; }
 
-		/// <summary>
-		/// List of pages to display in the dashboard
-		/// </summary>
-		public IReadOnlyList<StreamTab> Tabs { get; }
-
 		/// <summary>
 		/// List of templates available for this stream
 		/// </summary>
@@ -385,7 +380,7 @@ namespace Horde.Build.Streams
 		/// <returns>New response instance</returns>
 		public static GetStreamResponse ToApiResponse(this IStream stream, bool bIncludeAcl, List<GetTemplateRefResponse> apiTemplateRefs)
 		{
-			List<GetStreamTabResponse> apiTabs = stream.Tabs.ConvertAll(x => x.ToResponse());
+			List<TabConfig> apiTabs = stream.Config.Tabs;
 			Dictionary<string, AgentConfig> apiAgentTypes = stream.Config.AgentTypes.ToDictionary(x => x.Key, x => x.Value);
 			Dictionary<string, WorkspaceConfig> apiWorkspaceTypes = stream.Config.WorkspaceTypes.ToDictionary(x => x.Key, x => x.Value);
 			GetAclResponse? apiAcl = (bIncludeAcl && stream.Acl != null)? new GetAclResponse(stream.Acl) : null;
diff --git a/Engine/Source/Programs/Horde/Horde.Build/Streams/StreamCollection.cs b/Engine/Source/Programs/Horde/Horde.Build/Streams/StreamCollection.cs
--- a/Engine/Source/Programs/Horde/Horde.Build/Streams/StreamCollection.cs
+++ b/Engine/Source/Programs/Horde/Horde.Build/Streams/StreamCollection.cs
@@ -49,7 +49,6 @@ namespace Horde.Build.Streams
 			[BsonIgnore]
 			public StreamConfig? Config { get; set; }
 
-			public List<StreamTab> Tabs { get; set; } = new List<StreamTab>();
 			public Dictionary<TemplateRefId, TemplateRef> Templates { get; set; } = new Dictionary<TemplateRefId, TemplateRef>();
 			public DateTime? PausedUntil { get; set; }
 			public string? PauseComment { get; set; }
@@ -59,7 +58,6 @@ namespace Horde.Build.Streams
 			public bool Deleted { get; set; }
 
 			StreamConfig IStream.Config => Config!;
-			IReadOnlyList<StreamTab> IStream.Tabs => Tabs;
 			IReadOnlyDictionary<TemplateRefId, TemplateRef> IStream.Templates => Templates;
 
 			string IStream.Name => Config!.Name;
@@ -136,19 +134,18 @@ namespace Horde.Build.Streams
 		{
 			StreamConfig config = await _configCollection.GetConfigAsync<StreamConfig>(revision);
 
-			List<StreamTab> tabs = config.Tabs.ConvertAll(x => StreamTab.FromRequest(x));
 			Dictionary<TemplateRefId, TemplateRef> templateRefs = await CreateTemplateRefsAsync(config.Templates, stream, _templateCollection);
 
-			Validate(id, templateRefs, tabs, config);
+			Validate(id, templateRefs, config);
 
 			Acl? acl = Acl.Merge(new Acl(), config.Acl);
 			if (stream == null)
 			{
-				return await TryCreateAsync(id, projectId, revision, config, tabs, templateRefs, acl);
+				return await TryCreateAsync(id, projectId, revision, config, templateRefs, acl);
 			}
 			else
 			{
-				return await TryReplaceAsync(stream, projectId, revision, config, tabs, templateRefs, acl);
+				return await TryReplaceAsync(stream, projectId, revision, config, templateRefs, acl);
 			}
 		}
 
@@ -311,12 +308,11 @@ namespace Horde.Build.Streams
 		}
 
 		/// <inheritdoc/>
-		async Task<IStream?> TryCreateAsync(StreamId id, ProjectId projectId, string configRevision, StreamConfig config, List<StreamTab> tabs, Dictionary<TemplateRefId, TemplateRef> templateRefs, Acl? acl)
+		async Task<IStream?> TryCreateAsync(StreamId id, ProjectId projectId, string configRevision, StreamConfig config, Dictionary<TemplateRefId, TemplateRef> templateRefs, Acl? acl)
 		{
 			StreamDocument newStream = new StreamDocument(id, projectId);
 			newStream.ConfigRevision = configRevision;
 			newStream.Config = config;
-			newStream.Tabs = tabs;
 			newStream.Templates = templateRefs;
 			newStream.Acl = acl;
 
@@ -339,7 +335,7 @@ namespace Horde.Build.Streams
 		}
 
 		/// <inheritdoc/>
-		async Task<IStream?> TryReplaceAsync(IStream streamInterface, ProjectId projectId, string configRevision, StreamConfig config, List<StreamTab> tabs, Dictionary<TemplateRefId, TemplateRef>? templateRefs, Acl? acl)
+		async Task<IStream?> TryReplaceAsync(IStream streamInterface, ProjectId projectId, string configRevision, StreamConfig config, Dictionary<TemplateRefId, TemplateRef>? templateRefs, Acl? acl)
 		{
 			int order = config.Order;
 
@@ -350,7 +346,6 @@ namespace Horde.Build.Streams
 			List<UpdateDefinition<StreamDocument>> updates = new List<UpdateDefinition<StreamDocument>>();
 			updates.Add(updateBuilder.Set(x => x.ProjectId, projectId));
 			updates.Add(updateBuilder.Set(x => x.ConfigRevision, configRevision));
-			updates.Add(updateBuilder.Set(x => x.Tabs, tabs ?? new List<StreamTab>()));
 			updates.Add(updateBuilder.Set(x => x.Templates, templateRefs ?? new Dictionary<TemplateRefId, TemplateRef>()));
 			updates.Add(updateBuilder.SetOrUnsetNullRef(x => x.Acl, acl));
 			updates.Add(updateBuilder.Unset(x => x.Deleted));
@@ -476,7 +471,7 @@ namespace Horde.Build.Streams
 		/// <summary>
 		/// Checks the stream definition for consistency
 		/// </summary>
-		public static void Validate(StreamId streamId, IReadOnlyDictionary<TemplateRefId, TemplateRef> templates, IReadOnlyList<StreamTab> tabs, StreamConfig config)
+		public static void Validate(StreamId streamId, IReadOnlyDictionary<TemplateRefId, TemplateRef> templates, StreamConfig config)
 		{
 			// Check the default preflight template is valid
 			if (config.DefaultPreflight != null)
@@ -489,7 +484,7 @@ namespace Horde.Build.Streams
 
 			// Check that all the templates are referenced by a tab
 			HashSet<TemplateRefId> remainingTemplates = new HashSet<TemplateRefId>(templates.Keys);
-			foreach (JobsTab jobsTab in tabs.OfType<JobsTab>())
+			foreach (JobsTabConfig jobsTab in config.Tabs.OfType<JobsTabConfig>())
 			{
 				if (jobsTab.Templates != null)
 				{
diff --git a/Engine/Source/Programs/Horde/Horde.Build/Streams/StreamConfig.cs b/Engine/Source/Programs/Horde/Horde.Build/Streams/StreamConfig.cs
--- a/Engine/Source/Programs/Horde/Horde.Build/Streams/StreamConfig.cs
+++ b/Engine/Source/Programs/Horde/Horde.Build/Streams/StreamConfig.cs
@@ -100,7 +100,7 @@ namespace Horde.Build.Streams
 		/// <summary>
 		/// List of tabs to show for the new stream
 		/// </summary>
-		public List<CreateStreamTabRequest> Tabs { get; set; } = new List<CreateStreamTabRequest>();
+		public List<TabConfig> Tabs { get; set; } = new List<TabConfig>();
 
 		/// <summary>
 		/// Map of agent name to type
@@ -182,6 +182,95 @@ namespace Horde.Build.Streams
 		}
 	}
 
+
+	/// <summary>
+	/// Information about a page to display in the dashboard for a stream
+	/// </summary>
+	[JsonKnownTypes(typeof(JobsTabConfig))]
+	public abstract class TabConfig
+	{
+		/// <summary>
+		/// Title of this page
+		/// </summary>
+		[Required]
+		public string Title { get; set; } = null!;
+	}
+
+	/// <summary>
+	/// Describes a job page
+	/// </summary>
+	[JsonDiscriminator("Jobs")]
+	public class JobsTabConfig : TabConfig
+	{
+		/// <summary>
+		/// Whether to show job names on this page
+		/// </summary>
+		public bool ShowNames { get; set; }
+
+		/// <summary>
+		/// Names of jobs to include on this page. If there is only one name specified, the name column does not need to be displayed.
+		/// </summary>
+		public List<string>? JobNames { get; set; }
+
+		/// <summary>
+		/// List of job template names to show on this page.
+		/// </summary>
+		public List<TemplateRefId>? Templates { get; set; }
+
+		/// <summary>
+		/// Columns to display for different types of aggregates
+		/// </summary>
+		public List<JobsTabColumnConfig>? Columns { get; set; }
+	}
+
+	/// <summary>
+	/// Type of a column in a jobs tab
+	/// </summary>
+	public enum JobsTabColumnType
+	{
+		/// <summary>
+		/// Contains labels
+		/// </summary>
+		Labels,
+
+		/// <summary>
+		/// Contains parameters
+		/// </summary>
+		Parameter
+	}
+
+	/// <summary>
+	/// Describes a column to display on the jobs page
+	/// </summary>
+	public class JobsTabColumnConfig
+	{
+		/// <summary>
+		/// The type of column
+		/// </summary>
+		public JobsTabColumnType Type { get; set; } = JobsTabColumnType.Labels;
+
+		/// <summary>
+		/// Heading for this column
+		/// </summary>
+		[Required]
+		public string Heading { get; set; } = null!;
+
+		/// <summary>
+		/// Category of aggregates to display in this column. If null, includes any aggregate not matched by another column.
+		/// </summary>
+		public string? Category { get; set; }
+
+		/// <summary>
+		/// Parameter to show in this column
+		/// </summary>
+		public string? Parameter { get; set; }
+
+		/// <summary>
+		/// Relative width of this column.
+		/// </summary>
+		public int? RelativeWidth { get; set; }
+	}
+
 	/// <summary>
 	/// Mapping from a BuildGraph agent type to a set of machines on the farm
 	/// </summary>
diff --git a/Engine/Source/Programs/Horde/Horde.Build/Streams/StreamMessages.cs b/Engine/Source/Programs/Horde/Horde.Build/Streams/StreamMessages.cs
--- a/Engine/Source/Programs/Horde/Horde.Build/Streams/StreamMessages.cs
+++ b/Engine/Source/Programs/Horde/Horde.Build/Streams/StreamMessages.cs
@@ -374,7 +374,7 @@ namespace Horde.Build.Streams
 		/// <summary>
 		/// List of tabs to display for this stream
 		/// </summary>
-		public List<GetStreamTabResponse> Tabs { get; set; }
+		public List<TabConfig> Tabs { get; set; }
 
 		/// <summary>
 		/// Map of agent name to type
@@ -431,7 +431,7 @@ namespace Horde.Build.Streams
 		/// <param name="pausedUntil">Stream paused for new builds until this date</param>
 		/// <param name="pauseComment">Reason for stream being paused</param>
 		/// <param name="workflows">Workflows for this stream</param>
-		public GetStreamResponse(string id, string projectId, string name, string configRevision, int order, string? notificationChannel, string? notificationChannelFilter, string? triageChannel, DefaultPreflightConfig? defaultPreflight, List<GetStreamTabResponse> tabs, Dictionary<string, AgentConfig> agentTypes, Dictionary<string, WorkspaceConfig>? workspaceTypes, List<GetTemplateRefResponse> templates, GetAclResponse? acl, DateTime? pausedUntil, string? pauseComment, List<WorkflowConfig> workflows)
+		public GetStreamResponse(string id, string projectId, string name, string configRevision, int order, string? notificationChannel, string? notificationChannelFilter, string? triageChannel, DefaultPreflightConfig? defaultPreflight, List<TabConfig> tabs, Dictionary<string, AgentConfig> agentTypes, Dictionary<string, WorkspaceConfig>? workspaceTypes, List<GetTemplateRefResponse> templates, GetAclResponse? acl, DateTime? pausedUntil, string? pauseComment, List<WorkflowConfig> workflows)
 		{
 			Id = id;
 			ProjectId = projectId;
@@ -453,250 +453,4 @@ namespace Horde.Build.Streams
 			Workflows = workflows;
 		}
 	}
-
-	/// <summary>
-	/// Information about a page to display in the dashboard for a stream
-	/// </summary>
-	[JsonKnownTypes(typeof(CreateJobsTabRequest))]
-	public abstract class CreateStreamTabRequest
-	{
-		/// <summary>
-		/// Title of this page
-		/// </summary>
-		[Required]
-		public string Title { get; set; } = null!;
-	}
-
-	/// <summary>
-	/// Type of a column in a jobs tab
-	/// </summary>
-	public enum JobsTabColumnType
-	{
-		/// <summary>
-		/// Contains labels
-		/// </summary>
-		Labels,
-
-		/// <summary>
-		/// Contains parameters
-		/// </summary>
-		Parameter
-	}
-
-	/// <summary>
-	/// Describes a column to display on the jobs page
-	/// </summary>
-	public class CreateJobsTabColumnRequest
-	{
-		/// <summary>
-		/// The type of column
-		/// </summary>
-		public JobsTabColumnType Type { get; set; } = JobsTabColumnType.Labels;
-
-		/// <summary>
-		/// Heading for this column
-		/// </summary>
-		[Required]
-		public string Heading { get; set; } = null!;
-
-		/// <summary>
-		/// Category of aggregates to display in this column. If null, includes any aggregate not matched by another column.
-		/// </summary>
-		public string? Category { get; set; }
-
-		/// <summary>
-		/// Parameter to show in this column
-		/// </summary>
-		public string? Parameter { get; set; }
-
-		/// <summary>
-		/// Relative width of this column.
-		/// </summary>
-		public int? RelativeWidth { get; set; }
-
-		/// <summary>
-		/// Construct a JobsTabColumn object from this request
-		/// </summary>
-		/// <returns>Column object</returns>
-		public JobsTabColumn ToModel()
-		{
-			switch (Type)
-			{
-				case JobsTabColumnType.Labels:
-					return new JobsTabLabelColumn(Heading, Category, RelativeWidth);
-				case JobsTabColumnType.Parameter:
-					return new JobsTabParameterColumn(Heading, Parameter ?? "Undefined", RelativeWidth);
-				default:
-					return new JobsTabLabelColumn(Heading, "Undefined", RelativeWidth);
-			}
-		}
-	}
-
-	/// <summary>
-	/// Describes a job page
-	/// </summary>
-	[JsonDiscriminator("Jobs")]
-	public class CreateJobsTabRequest : CreateStreamTabRequest
-	{
-		/// <summary>
-		/// Whether to show job names on this page
-		/// </summary>
-		public bool ShowNames { get; set; }
-
-		/// <summary>
-		/// Names of jobs to include on this page. If there is only one name specified, the name column does not need to be displayed.
-		/// </summary>
-		public List<string>? JobNames { get; set; }
-
-		/// <summary>
-		/// List of job template names to show on this page.
-		/// </summary>
-		public List<TemplateRefId>? Templates { get; set; }
-
-		/// <summary>
-		/// Columns to display for different types of aggregates
-		/// </summary>
-		public List<CreateJobsTabColumnRequest>? Columns { get; set; }
-	}
-
-	/// <summary>
-	/// Information about a page to display in the dashboard for a stream
-	/// </summary>
-	[JsonKnownTypes(typeof(GetJobsTabResponse))]
-	public abstract class GetStreamTabResponse
-	{
-		/// <summary>
-		/// Title of this page
-		/// </summary>
-		public string Title { get; set; }
-
-		/// <summary>
-		/// Constructor
-		/// </summary>
-		/// <param name="title">Title of this page</param>
-		protected GetStreamTabResponse(string title)//, StreamPageType Type)
-		{
-			Title = title;
-		}
-	}
-
-	/// <summary>
-	/// Describes a column to display on the jobs page
-	/// </summary>
-	[JsonKnownTypes(typeof(GetJobsTabLabelColumnResponse), typeof(GetJobsTabParameterColumnResponse))]
-	public abstract class GetJobsTabColumnResponse
-	{
-		/// <summary>
-		/// Heading for this column
-		/// </summary>
-		public string Heading { get; set; } = null!;
-
-		/// <summary>
-		/// Relative width of this column.
-		/// </summary>
-		public int? RelativeWidth { get; set; }
-
-		/// <summary>
-		/// Constructor
-		/// </summary>
-		/// <param name="heading">Heading for this column</param>
-		/// <param name="relativeWidth">Relative width of this column</param>
-		protected GetJobsTabColumnResponse(string heading, int? relativeWidth)
-		{
-			Heading = heading;
-			RelativeWidth = relativeWidth;
-		}
-	}
-
-	/// <summary>
-	/// Describes a label column to display on the jobs page
-	/// </summary>
-	[JsonDiscriminator(nameof(JobsTabColumnType.Labels))]
-	public class GetJobsTabLabelColumnResponse : GetJobsTabColumnResponse
-	{
-		/// <summary>
-		/// Category of aggregates to display in this column. If null, includes any aggregate not matched by another column.
-		/// </summary>
-		public string? Category { get; set; }
-
-		/// <summary>
-		/// Constructor
-		/// </summary>
-		/// <param name="heading">Heading for this column</param>
-		/// <param name="category">Category of aggregates to display in this column. If null, includes any aggregate not matched by another column.</param>
-		/// <param name="relativeWidth">Relative width of this column</param>
-		public GetJobsTabLabelColumnResponse(string heading, string? category, int? relativeWidth)
-			: base(heading, relativeWidth)
-		{
-			Category = category;
-		}
-	}
-
-	/// <summary>
-	/// Describes a parameter column to display on the jobs page
-	/// </summary>
-	[JsonDiscriminator(nameof(JobsTabColumnType.Parameter))]
-	public class GetJobsTabParameterColumnResponse : GetJobsTabColumnResponse
-	{
-		/// <summary>
-		/// Parameter to show in this column
-		/// </summary>
-		public string? Parameter { get; set; }
-
-		/// <summary>
-		/// Constructor
-		/// </summary>
-		/// <param name="heading">Heading for this column</param>
-		/// <param name="parameter">Name of the parameter to display</param>
-		/// <param name="relativeWidth">Relative width of this column</param>
-		public GetJobsTabParameterColumnResponse(string heading, string? parameter, int? relativeWidth)
-			: base(heading, relativeWidth)
-		{
-			Parameter = parameter;
-		}
-	}
-
-	/// <summary>
-	/// Describes a job page
-	/// </summary>
-	[JsonDiscriminator("Jobs")]
-	public class GetJobsTabResponse : GetStreamTabResponse
-	{
-		/// <summary>
-		/// Whether to show names on the page
-		/// </summary>
-		public bool ShowNames { get; set; }
-
-		/// <summary>
-		/// List of templates to show on the page
-		/// </summary>
-		public List<TemplateRefId>? Templates { get; set; }
-
-		/// <summary>
-		/// Names of jobs to include on this page. If there is only one name specified, the name column does not need to be displayed.
-		/// </summary>
-		public List<string>? JobNames { get; set; }
-
-		/// <summary>
-		/// Columns to display for different types of aggregates
-		/// </summary>
-		public List<GetJobsTabColumnResponse>? Columns { get; set; }
-
-		/// <summary>
-		/// Constructor
-		/// </summary>
-		/// <param name="title">Title for this page</param>
-		/// <param name="bShowNames">Whether to show names on the page</param>
-		/// <param name="templates">Templates to include on this page</param>
-		/// <param name="jobNames">List of job names to include on this page</param>
-		/// <param name="columns">List of columns to display</param>
-		public GetJobsTabResponse(string title, bool bShowNames, List<TemplateRefId>? templates, List<string>? jobNames, List<GetJobsTabColumnResponse>? columns)
-			: base(title)
-		{
-			ShowNames = bShowNames;
-			Templates = templates;
-			JobNames = jobNames;
-			Columns = columns;
-		}
-	}
 }
diff --git a/Engine/Source/Programs/Horde/Horde.Build/Streams/StreamTab.cs b/Engine/Source/Programs/Horde/Horde.Build/Streams/StreamTab.cs
deleted file mode 100644
