Horde: Fix exception when starting schedules.

diff --git a/Engine/Source/Programs/Horde/Horde.Build/Perforce/PerforceService.cs b/Engine/Source/Programs/Horde/Horde.Build/Perforce/PerforceService.cs
--- a/Engine/Source/Programs/Horde/Horde.Build/Perforce/PerforceService.cs
+++ b/Engine/Source/Programs/Horde/Horde.Build/Perforce/PerforceService.cs
@@ -172,7 +172,7 @@ namespace Horde.Build.Perforce
 						_files = GetStreamFiles(describeRecord, _streamName);
 					}
 				}
-				throw new NotImplementedException();
+				return _files;
 			}
 		}
 
