Horde: Restore missing copy of Redis binaries to output dir.

diff --git a/Engine/Source/Programs/Horde/Horde.Build/Horde.Build.csproj b/Engine/Source/Programs/Horde/Horde.Build/Horde.Build.csproj
--- a/Engine/Source/Programs/Horde/Horde.Build/Horde.Build.csproj
+++ b/Engine/Source/Programs/Horde/Horde.Build/Horde.Build.csproj
@@ -112,6 +112,13 @@
     <ProjectReference Include="..\..\Shared\EpicGames.Slack\EpicGames.Slack.csproj" />
   </ItemGroup>
 
+  <ItemGroup>
+    <Content Include="..\ThirdParty\**" Condition="$([MSBuild]::IsOSPlatform('Windows'))">
+      <CopyToOutputDirectory>PreserveNewest</CopyToOutputDirectory>
+      <Link>ThirdParty/%(RecursiveDir)/%(Filename)%(Extension)</Link>
+    </Content>
+  </ItemGroup>
+
   <ItemGroup>
     <Content Include="Defaults\ue5.project.png">
       <CopyToOutputDirectory>PreserveNewest</CopyToOutputDirectory>
