Horde: Update device service platforms to use config
#rnx

diff --git a/Engine/Source/Programs/Horde/Horde.Build.Tests/DeviceServiceTests.cs b/Engine/Source/Programs/Horde/Horde.Build.Tests/DeviceServiceTests.cs
--- a/Engine/Source/Programs/Horde/Horde.Build.Tests/DeviceServiceTests.cs
+++ b/Engine/Source/Programs/Horde/Horde.Build.Tests/DeviceServiceTests.cs
@@ -4,7 +4,6 @@ using System.Collections.Generic;
 using System.Threading.Tasks;
 using Horde.Build.Devices;
 using Horde.Build.Jobs;
-using Horde.Build.Server;
 using Microsoft.AspNetCore.Mvc;
 using Microsoft.VisualStudio.TestTools.UnitTesting;
 using Horde.Build.Projects;
@@ -12,12 +11,14 @@ using Horde.Build.Utilities;
 using Horde.Build.Users;
 using System.Security.Claims;
 using Microsoft.AspNetCore.Http;
+using Horde.Build.Server;
 
 namespace Horde.Build.Tests
 {
 
 	using ProjectId = StringId<IProject>;
 	using DeviceId = StringId<IDevice>;
+	using DevicePlatformId = StringId<IDevicePlatform>;
 
 	/// <summary>
 	/// Tests for the device service
@@ -50,9 +51,9 @@ namespace Horde.Build.Tests
 			}
 		}
 
-
 		private async Task PopulateDevices()
 		{
+			DeviceConfig Devices = new DeviceConfig();
 
 			await ProjectCollection.AddOrUpdateAsync(new ProjectId("ue5"), "", "", 0, new ProjectConfig { Name = "UE5" });
 
@@ -63,50 +64,67 @@ namespace Horde.Build.Tests
 			}
 
 			Dictionary<string, string> platformMap = new Dictionary<string, string>();
-			Dictionary<string, string> platformReverseMap = new Dictionary<string, string>();
 
 			// create 3 platforms
 			for (int i = 1; i < 4; i++)
 			{
 				string platformName = "TestDevicePlatform" + i;
 				await DevicesController.CreatePlatformAsync(new CreateDevicePlatformRequest() { Name = platformName });
+
+				List<string> modelIds = new List<string>();
+				for (int j = 2; j < 5; j++)
+				{
+					modelIds.Add(platformName + "_Model" + j);
+				}
+
+				// add models
+				await DevicesController.UpdatePlatformAsync(DevicePlatformId.Sanitize(platformName).ToString(), new UpdateDevicePlatformRequest() { ModelIds = modelIds.ToArray() });
+
 				platformMap[platformName] = platformName;
-				platformReverseMap[platformName] = platformName;
+
+				// platform 3 has some name aliases
+				if (i == 3)
+				{
+					DevicePlatformConfig config = new DevicePlatformConfig();
+					config.Id = DevicePlatformId.Sanitize(platformName).ToString();
+					config.Names.Add("TestDevicePlatform3Alias");
+					config.LegacyPerfSpecHighModel = "TestDevicePlatform3_Model4";
+					Devices.Platforms.Add(config);
+				}
 			}
 
+			Globals globals = await MongoService.GetGlobalsAsync();
+			globals.Devices = Devices;
+			Assert.IsTrue(await MongoService.TryUpdateSingletonAsync(globals));
+
 			// add 4 devices to each platform, split between 2 pools
 			for (int i = 1; i < 4; i++)
 			{
 				for (int j = 1; j < 5; j++)
 				{
-					// @todo: add model tests
 					// one base model, and 3 other models 
-					/*
 					string? ModelId = null;
 					if (j > 1)
 					{
 						ModelId = "TestDevicePlatform" + i + "_Model" + j;
 					}
-					*/
 
 					string poolId = (j & 1) != 0 ? "testdevicepool1" : "testdevicepool2";
 
-					await DeviceController.CreateDeviceAsync(new CreateDeviceRequest() { Name = "TestDevice" + j + "_Platform" + i + "_" + poolId, Address = "10.0.0.1", Enabled = true, PlatformId = "testdeviceplatform" + i/*, ModelId = ModelId*/, PoolId = poolId });
+					await DeviceController.CreateDeviceAsync(new CreateDeviceRequest() { Name = "TestDevice" + j + "_Platform" + i + "_" + poolId, Address = "10.0.0.1", Enabled = true, PlatformId = "testdeviceplatform" + i, ModelId = ModelId, PoolId = poolId });
 				}
 			}
 
-			// populate the legacy platform map, this exists to avoid having platform names in server code, though needs to be cleaned up.
-			// Gauntlet in various streams can't/hasn't been updated and asks for non-platform-id device types
-			await DeviceController.UpdatePlatformMapV1(new UpdatePlatformMapRequest() { PlatformMap = platformMap, PlatformReverseMap = platformReverseMap, PerfSpecHighMap = new Dictionary<string, string>() });
-
+			// tick to pick up config change
+			await DeviceService.TickForTestingAsync();
 		}
 
 		static T ResultToValue<T>(ActionResult<T> result) where T: class
 		{
 			return ((result.Result! as JsonResult)!.Value! as T)!;
 		}
-	
-		async Task<LegacyCreateReservationRequest> SetupReservationTest()
+
+		async Task<LegacyCreateReservationRequest> SetupReservationTest(string poolId = "TestDevicePool1", string deviceType = "TestDevicePlatform1")
 		{
 			await CreateFixtureAsync();
 			await PopulateDevices();
@@ -117,8 +135,8 @@ namespace Horde.Build.Tests
 
 			// Gauntlet uses the legacy v1 API
 			LegacyCreateReservationRequest request = new LegacyCreateReservationRequest();
-			request.PoolId = "TestDevicePool1";
-			request.DeviceTypes = new string[] { "TestDevicePlatform1" };
+			request.PoolId = poolId;
+			request.DeviceTypes = new string[] { deviceType };
 			request.Hostname = "localhost";
 			request.Duration = "00:10:00";
 			request.JobId = job.Id;
@@ -127,7 +145,6 @@ namespace Horde.Build.Tests
 			return request;
 		}
 
-
 		[TestMethod]
 		public async Task TestReservation()
 		{
@@ -160,6 +177,39 @@ namespace Horde.Build.Tests
 			Assert.AreEqual(telemetry[0].Telemetry[0].JobName, "hello2");
 		}
 
+		[TestMethod]
+		public async Task TestReservationPerfSpecWithAlias()
+		{
+			LegacyCreateReservationRequest request = await SetupReservationTest("TestDevicePool2", "TestDevicePlatform3Alias:High");
+
+			// create a reservation
+			GetLegacyReservationResponse reservation = ResultToValue(await DeviceController!.CreateDeviceReservationV1Async(request));
+			Assert.AreEqual(1, reservation.DeviceNames.Length);
+			Assert.AreEqual("hello2", reservation.JobName);
+			Assert.AreEqual("abcd", reservation.StepId);
+
+			// get the device in the reservation, and make sure it is the right platform
+			GetLegacyDeviceResponse device = ResultToValue(await DeviceController!.GetDeviceV1Async(reservation.DeviceNames[0]));
+			Assert.AreEqual("TestDevicePlatform3Alias", device.Type);
+			Assert.AreEqual("TestDevicePlatform3_Model4", device.Model);
+
+			// update the reservation
+			GetLegacyReservationResponse renewed = ResultToValue(await DeviceController!.UpdateReservationV1Async(reservation.Guid));
+			Assert.AreEqual(renewed.Guid, reservation.Guid);
+
+			// delete the reservation, would be nice to also test reservation expiration
+			OkResult? deleted = (await DeviceController!.DeleteReservationV1Async(reservation.Guid)) as OkResult;
+			Assert.IsNotNull(deleted);
+
+			// check that telemetry was created
+			List<GetDeviceTelemetryResponse> telemetry = (await DeviceController!.GetDeviceTelemetry()).Value!;
+			Assert.AreEqual(telemetry.Count, 1);
+			Assert.AreEqual(telemetry[0].Telemetry.Count, 1);
+			Assert.AreEqual(telemetry[0].Telemetry[0].StreamId, "ue5-main");
+			Assert.AreEqual(telemetry[0].Telemetry[0].StepId, "abcd");
+			Assert.AreEqual(telemetry[0].Telemetry[0].JobName, "hello2");
+		}
+
 		[TestMethod]
 		public async Task TestProblemDevice()
 		{
@@ -204,7 +254,6 @@ namespace Horde.Build.Tests
 			request.Problem = true;
 			await DeviceController!.UpdateDeviceAsync("testdevice4_platform2_testdevicepool2", request);
 
-
 			IDevice? maintenanceDevice = await DeviceService.GetDeviceAsync(new DeviceId("testdevice1_platform1_testdevicepool1"));
 			Assert.IsNotNull(maintenanceDevice);
 			
@@ -217,8 +266,10 @@ namespace Horde.Build.Tests
 			await DeviceService.TickForTestingAsync();
 
 			List<GetDevicePoolTelemetryResponse> telemetry = (await DeviceController!.GetDevicePoolTelemetry()).Value!;
-			Assert.AreEqual(1, telemetry.Count);
-			foreach (GetDevicePlatformTelemetryResponse t in telemetry[0].Telemetry["testdevicepool1"])
+
+			// 2, as generate an initial telemetry tick in setup
+			Assert.AreEqual(2, telemetry.Count);
+			foreach (GetDevicePlatformTelemetryResponse t in telemetry[1].Telemetry["testdevicepool1"])
 			{
 				if (t.PlatformId == "testdeviceplatform1")
 				{
@@ -243,7 +294,7 @@ namespace Horde.Build.Tests
 				}
 			}
 
-			foreach (GetDevicePlatformTelemetryResponse t in telemetry[0].Telemetry["testdevicepool2"])
+			foreach (GetDevicePlatformTelemetryResponse t in telemetry[1].Telemetry["testdevicepool2"])
 			{
 				if (t.PlatformId == "testdeviceplatform2")
 				{
diff --git a/Engine/Source/Programs/Horde/Horde.Build/Configuration/ConfigUpdateService.cs b/Engine/Source/Programs/Horde/Horde.Build/Configuration/ConfigUpdateService.cs
--- a/Engine/Source/Programs/Horde/Horde.Build/Configuration/ConfigUpdateService.cs
+++ b/Engine/Source/Programs/Horde/Horde.Build/Configuration/ConfigUpdateService.cs
@@ -137,6 +137,7 @@ namespace Horde.Build.Configuration
 				globals.ScheduledDowntime = _cachedGlobalConfig.Downtime;
 				globals.MaxConformCount = _cachedGlobalConfig.MaxConformCount;
 				globals.ComputeClusters = _cachedGlobalConfig.Compute;
+				globals.Devices = _cachedGlobalConfig.Devices;
 				globals.RootAcl = Acl.Merge(null, _cachedGlobalConfig.Acl);
 
 				if (await _mongoService.TryUpdateSingletonAsync(globals))
diff --git a/Engine/Source/Programs/Horde/Horde.Build/Devices/DeviceMessages.cs b/Engine/Source/Programs/Horde/Horde.Build/Devices/DeviceMessages.cs
--- a/Engine/Source/Programs/Horde/Horde.Build/Devices/DeviceMessages.cs
+++ b/Engine/Source/Programs/Horde/Horde.Build/Devices/DeviceMessages.cs
@@ -980,28 +980,4 @@ namespace Horde.Build.Devices
 		[Required]
 		public string DeviceData { get; set; } = null!;
 	}
-
-	/// <summary>
-	/// Platform mapping update for V1 devices
-	/// </summary>
-	public class UpdatePlatformMapRequest
-	{
-		/// <summary>
-		/// Platform name => platform id
-		/// </summary>
-		[Required]
-		public Dictionary<string, string> PlatformMap { get; set; } = null!;
-
-		/// <summary>
-		/// Platform id => name
-		/// </summary>
-		[Required]
-		public Dictionary<string, string> PlatformReverseMap { get; set; } = null!;
-
-		/// <summary>
-		/// Perfspec v1 => model names
-		/// </summary>
-		[Required]
-		public Dictionary<string, string> PerfSpecHighMap { get; set; } = null!;
-	}
 }
diff --git a/Engine/Source/Programs/Horde/Horde.Build/Devices/DeviceService.cs b/Engine/Source/Programs/Horde/Horde.Build/Devices/DeviceService.cs
--- a/Engine/Source/Programs/Horde/Horde.Build/Devices/DeviceService.cs
+++ b/Engine/Source/Programs/Horde/Horde.Build/Devices/DeviceService.cs
@@ -76,16 +76,11 @@ namespace Horde.Build.Devices
 		/// </summary>
 		public Dictionary<string, DevicePlatformId> PlatformMap { get; set; } = new Dictionary<string, DevicePlatformId>();
 
-		/// <summary>
-		/// Platform Id => Platform V1
-		/// </summary>
-		public Dictionary<DevicePlatformId, string> PlatformReverseMap { get; set; } = new Dictionary<DevicePlatformId, string>();
-
 		/// <summary>
 		/// Perfspec V1 => Model
 		/// </summary>
 		public Dictionary<DevicePlatformId, string> PerfSpecHighMap { get; set; } = new Dictionary<DevicePlatformId, string>();
-
+		
 	}
 
 	/// <summary>
@@ -118,6 +113,11 @@ namespace Horde.Build.Devices
 		/// </summary>
 		readonly ProjectService _projectService;
 
+		/// <summary>
+		/// Singleton instance of the mongo service
+		/// </summary>
+		readonly MongoService _mongoService;
+
 		/// <summary>
 		/// The user collection instance
 		/// </summary>
@@ -143,7 +143,7 @@ namespace Horde.Build.Devices
 		/// <summary>
 		/// Device service constructor
 		/// </summary>
-		public DeviceService(IDeviceCollection devices, ISingletonDocument<DevicePlatformMapV1> platformMapSingleton, IUserCollection userCollection, JobService jobService, ProjectService projectService, StreamService streamService, AclService aclService, INotificationService notificationService, IClock clock, ILogger<DeviceService> logger)
+		public DeviceService(IDeviceCollection devices, ISingletonDocument<DevicePlatformMapV1> platformMapSingleton, IUserCollection userCollection, MongoService mongoService, JobService jobService, ProjectService projectService, StreamService streamService, AclService aclService, INotificationService notificationService, IClock clock, ILogger<DeviceService> logger)
 		{
 			UserCollection = userCollection;
 			_devices = devices;
@@ -155,6 +155,7 @@ namespace Horde.Build.Devices
 			_ticker = clock.AddSharedTicker<DeviceService>(TimeSpan.FromMinutes(1.0), TickAsync, logger);
 			_telemetryTicker = clock.AddSharedTicker("DeviceService.Telemetry", TimeSpan.FromMinutes(10.0), TickTelemetryAsync, logger);
 			_logger = logger;
+			_mongoService = mongoService;
 
 			_platformMapSingleton = platformMapSingleton;
 
@@ -187,6 +188,38 @@ namespace Horde.Build.Devices
 		/// </summary>
 		async ValueTask TickTelemetryAsync(CancellationToken stoppingToken)
 		{
+			try
+			{
+				Globals globals = await _mongoService.GetGlobalsAsync();
+
+				if (globals.Devices != null)
+				{
+					await _platformMapSingleton.UpdateAsync(platformMap => {
+
+						platformMap.PlatformMap.Clear();
+						platformMap.PerfSpecHighMap.Clear();
+
+						foreach (DevicePlatformConfig platform in globals.Devices.Platforms)
+						{
+							DevicePlatformId id = new DevicePlatformId(platform.Id);
+							foreach (string name in platform.Names)
+							{
+								platformMap.PlatformMap[name] = id;
+							}
+
+							if (platform.LegacyPerfSpecHighModel != null)
+							{
+								platformMap.PerfSpecHighMap[id] = platform.LegacyPerfSpecHighModel;
+							}
+						}
+					} );
+				}				
+			}
+			catch (Exception ex)
+			{
+				_logger.LogError(ex, "Exception while updating platform map: {Message}", ex.Message);
+			}
+
 			if (!stoppingToken.IsCancellationRequested)
 			{
 				using IScope scope = GlobalTracer.Instance.BuildSpan("DeviceService.TickTelemetryAsync").StartActive();
@@ -682,66 +715,5 @@ namespace Horde.Build.Devices
 		{
 			return await _platformMapSingleton.GetAsync();
 		}
-
-		/// <summary>
-		/// Updates the platform mapping information required for the V1 api
-		/// </summary>				
-		public async Task<bool> UpdatePlatformMapAsync(UpdatePlatformMapRequest request)
-		{
-
-			List<IDevicePlatform> platforms = await GetPlatformsAsync();
-
-			// Update the platform map
-			for (int i = 0; i < 10; i++)
-			{
-				DevicePlatformMapV1 instance = await _platformMapSingleton.GetAsync();
-
-				instance.PlatformMap = new Dictionary<string, DevicePlatformId>();
-				instance.PlatformReverseMap = new Dictionary<DevicePlatformId, string>();
-				instance.PerfSpecHighMap = new Dictionary<DevicePlatformId, string>();
-
-				foreach(KeyValuePair<string, string> entry in request.PlatformMap)
-				{
-					IDevicePlatform? platform = platforms.FirstOrDefault(p => p.Id == new DevicePlatformId(entry.Value));
-					if (platform == null)
-					{
-						throw new Exception($"Unknowm platform in map {entry.Key} : {entry.Value}");
-					}
-
-					instance.PlatformMap.Add(entry.Key, platform.Id);
-				}
-
-				foreach (KeyValuePair<string, string> entry in request.PlatformReverseMap)
-				{
-					IDevicePlatform? platform = platforms.FirstOrDefault(p => p.Id == new DevicePlatformId(entry.Key));
-					if (platform == null)
-					{
-						throw new Exception($"Unknowm platform in reverse map {entry.Key} : {entry.Value}");
-					}
-
-					instance.PlatformReverseMap.Add(platform.Id, entry.Value);
-				}
-
-				foreach (KeyValuePair<string, string> entry in request.PerfSpecHighMap)
-				{
-					IDevicePlatform? platform = platforms.FirstOrDefault(p => p.Id == new DevicePlatformId(entry.Key));
-					if (platform == null)
-					{
-						throw new Exception($"Unknowm platform in spec map {entry.Key} : {entry.Value}");
-					}
-
-					instance.PerfSpecHighMap.Add(platform.Id, entry.Value);
-				}
-
-				if (await _platformMapSingleton.TryUpdateAsync(instance))
-				{
-					return true;
-				}
-
-				Thread.Sleep(1000);
-			}
-
-			return false;
-		}
 	}
 }
diff --git a/Engine/Source/Programs/Horde/Horde.Build/Devices/DevicesController.cs b/Engine/Source/Programs/Horde/Horde.Build/Devices/DevicesController.cs
--- a/Engine/Source/Programs/Horde/Horde.Build/Devices/DevicesController.cs
+++ b/Engine/Source/Programs/Horde/Horde.Build/Devices/DevicesController.cs
@@ -869,20 +869,20 @@ namespace Horde.Build.Devices
 					constraint = tokens[1];
 				}
 
-				DevicePlatformId platformId;
-
 				DevicePlatformMapV1 mapV1 = await _deviceService.GetPlatformMapV1();
 
-				if (!mapV1.PlatformMap.TryGetValue(platformName, out platformId))
-				{
-					string message = $"Unknown platform {platformName}" + details;
-					_logger.LogError("Unknown platform {PlatformName} {Details}", platformName, details);
-					await _deviceService.NotifyDeviceServiceAsync(message, null, request.JobId, request.StepId);
-					return BadRequest(message);
-				}
+				DevicePlatformId platformId = DevicePlatformId.Sanitize(platformName);
 
 				IDevicePlatform? platform = platforms.FirstOrDefault(x => x.Id == platformId);
 
+				if (platform == null)
+				{
+					if (mapV1.PlatformMap.TryGetValue(platformName, out platformId))
+					{
+						platform = platforms.FirstOrDefault(x => x.Id == platformId);
+					}
+				}
+
 				if (platform == null)
 				{
 					string message = $"Unknown platform {platformId}" + details;
@@ -1067,16 +1067,9 @@ namespace Horde.Build.Devices
 			DevicePlatformMapV1 mapV1 = await _deviceService.GetPlatformMapV1();
 
 			if (String.IsNullOrEmpty(platformName))
-			{
-				if (!mapV1.PlatformReverseMap.TryGetValue(device.PlatformId, out platformName))
-				{
-					return BadRequest($"Unable to map platform for {deviceName} : {device.PlatformId}");
-				}
-			}
-
-			if (String.IsNullOrEmpty(platformName))
-			{
-				return BadRequest($"Unable to get platform for {deviceName} from reservation or mapping : {device.PlatformId}");
+			{				
+				_logger.LogError("Unable to map platform for {DeviceName} : {PlatformId} from reservation", deviceName, device.PlatformId);			
+				return BadRequest($"Unable to get platform for {deviceName} from reservation : {device.PlatformId}");
 			}
 
 			GetLegacyDeviceResponse response = new GetLegacyDeviceResponse();
@@ -1155,37 +1148,5 @@ namespace Horde.Build.Devices
 			return Ok();
 
 		}
-
-		/// <summary>
-		/// Updates the platform map for v1 requests
-		/// </summary>
-		[HttpPut]
-		[Route("/api/v1/devices/platformmap")]
-		public async Task<ActionResult> UpdatePlatformMapV1([FromBody] UpdatePlatformMapRequest request)
-		{
-			if (!await _aclService.AuthorizeAsync(AclAction.AdminWrite, User))
-			{
-				return Forbid();
-			}
-
-			bool result;
-			try
-			{
-				result = await _deviceService.UpdatePlatformMapAsync(request);
-			}
-			catch (Exception ex)
-			{
-				_logger.LogError(ex, "Error updating device platform map {Message}", ex.Message);
-				throw;
-			}
-
-			if (!result)
-			{
-				_logger.LogError("Unable to update device platform mapping");
-				return BadRequest();
-			}
-
-			return Ok();
-		}
 	}
 }
diff --git a/Engine/Source/Programs/Horde/Horde.Build/Server/GlobalConfig.cs b/Engine/Source/Programs/Horde/Horde.Build/Server/GlobalConfig.cs
--- a/Engine/Source/Programs/Horde/Horde.Build/Server/GlobalConfig.cs
+++ b/Engine/Source/Programs/Horde/Horde.Build/Server/GlobalConfig.cs
@@ -47,6 +47,11 @@ namespace Horde.Build.Server
 		/// </summary>
 		public List<ComputeClusterConfig> Compute { get; set; } = new List<ComputeClusterConfig>();
 
+		/// <summary>
+		/// Device configuration
+		/// </summary>
+		public DeviceConfig? Devices { get; set; }
+
 		/// <summary>
 		/// List of tools hosted by the server
 		/// </summary>
@@ -149,6 +154,39 @@ namespace Horde.Build.Server
 		public AclConfig? Acl { get; set; }
 	}
 
+	/// <summary>
+	/// Configuration for a device platform 
+	/// </summary>
+	public class DevicePlatformConfig
+	{
+		/// <summary>
+		/// The id for this platform 
+		/// </summary>
+		public string Id { get; set; } = String.Empty;
+
+		/// <summary>
+		/// List of platform names for this device, which may be requested by Gauntlet 
+		/// </summary>
+		public List<string> Names { get; set; } = new List<string>();
+
+		/// <summary>
+		/// Model name for the high perf spec, which may be requested by Gauntlet (Deprecated)
+		/// </summary>
+		public string? LegacyPerfSpecHighModel { get; set; }
+	}
+
+	/// <summary>
+	/// Configuration for devices
+	/// </summary>
+	public class DeviceConfig
+	{
+		/// <summary>
+		/// List of device platforms
+		/// </summary>
+		public List<DevicePlatformConfig> Platforms { get; set; } = new List<DevicePlatformConfig>();
+	}
+
+
 	/// <summary>
 	/// Configuration for a bucket
 	/// </summary>
diff --git a/Engine/Source/Programs/Horde/Horde.Build/Server/Globals.cs b/Engine/Source/Programs/Horde/Horde.Build/Server/Globals.cs
--- a/Engine/Source/Programs/Horde/Horde.Build/Server/Globals.cs
+++ b/Engine/Source/Programs/Horde/Horde.Build/Server/Globals.cs
@@ -329,6 +329,11 @@ namespace Horde.Build.Server
 		/// </summary>
 		public List<ComputeClusterConfig> ComputeClusters { get; set; } = new List<ComputeClusterConfig>();
 
+		/// <summary>
+		/// Device platform configuration
+		/// </summary>
+		public DeviceConfig? Devices { get; set; }
+
 		/// <summary>
 		/// Private constructor for serialization
 		/// </summary>
