Horde: Raise a config error for duplicate project stream id/path
#rnx

diff --git a/Engine/Source/Programs/Horde/Horde.Build/Configuration/ConfigUpdateService.cs b/Engine/Source/Programs/Horde/Horde.Build/Configuration/ConfigUpdateService.cs
--- a/Engine/Source/Programs/Horde/Horde.Build/Configuration/ConfigUpdateService.cs
+++ b/Engine/Source/Programs/Horde/Horde.Build/Configuration/ConfigUpdateService.cs
@@ -216,6 +216,35 @@ namespace Horde.Build.Configuration
 
 				_cachedProjectConfigs[projectRef.Id] = (projectConfig, revision);
 				streamConfigs.AddRange(projectConfig.Streams.Select(x => (projectRef.Id, x, CombinePaths(projectPath, x.Path))));
+
+				// Check that project doesn't contain duplicate streams			
+				try
+				{
+					HashSet<StreamId> streamIds = new HashSet<StreamId>();
+					HashSet<Uri> streamPaths = new HashSet<Uri>();
+
+					foreach ((ProjectId projectId, StreamConfigRef streamRef, Uri streamPath) in streamConfigs)
+					{
+						if (streamIds.Contains(streamRef.Id))
+						{
+							throw new Exception($"Project {projectRef.Id} contains duplicate stream id {streamRef.Id}");
+						}
+
+						if (streamPaths.Contains(streamPath))
+						{
+							throw new Exception($"Project {projectRef.Id} contains duplicate stream path {streamPath}");
+						}
+
+						streamIds.Add(streamRef.Id);
+						streamPaths.Add(streamPath);
+					}
+
+				}
+				catch (Exception ex)
+				{
+					await SendFailureNotificationAsync(ex, projectPath, cancellationToken);
+					continue;
+				}
 			}
 
 			// Get the logo revisions
