Horde: Better duplicate project stream id/path check
#rnx

diff --git a/Engine/Source/Programs/Horde/Horde.Build/Configuration/ConfigUpdateService.cs b/Engine/Source/Programs/Horde/Horde.Build/Configuration/ConfigUpdateService.cs
--- a/Engine/Source/Programs/Horde/Horde.Build/Configuration/ConfigUpdateService.cs
+++ b/Engine/Source/Programs/Horde/Horde.Build/Configuration/ConfigUpdateService.cs
@@ -192,9 +192,30 @@ namespace Horde.Build.Configuration
 				else
 				{
 					_logger.LogInformation("Caching configuration for project {ProjectId} ({Revision})", projectRef.Id, revision);
+
 					try
 					{
 						projectConfig = await ReadDataAsync<ProjectConfig>(revision, projectPath, cancellationToken);
+
+						HashSet<StreamId> streamIds = new HashSet<StreamId>();
+						HashSet<string> streamPaths = new HashSet<string>();
+
+						foreach (StreamConfigRef streamRef in projectConfig.Streams)
+						{
+							if (streamIds.Contains(streamRef.Id))
+							{
+								throw new Exception($"Project {projectRef.Id} contains duplicate stream id {streamRef.Id}");
+							}
+
+							if (streamPaths.Contains(streamRef.Path))
+							{
+								throw new Exception($"Project {projectRef.Id} contains duplicate stream path {streamRef.Path}");
+							}
+
+							streamIds.Add(streamRef.Id);
+							streamPaths.Add(streamRef.Path);
+						}
+
 						if (update)
 						{
 							_logger.LogInformation("Updating configuration for project {ProjectId} ({Revision})", projectRef.Id, revision);
@@ -216,35 +237,6 @@ namespace Horde.Build.Configuration
 
 				_cachedProjectConfigs[projectRef.Id] = (projectConfig, revision);
 				streamConfigs.AddRange(projectConfig.Streams.Select(x => (projectRef.Id, x, CombinePaths(projectPath, x.Path))));
-
-				// Check that project doesn't contain duplicate streams			
-				try
-				{
-					HashSet<StreamId> streamIds = new HashSet<StreamId>();
-					HashSet<Uri> streamPaths = new HashSet<Uri>();
-
-					foreach ((ProjectId projectId, StreamConfigRef streamRef, Uri streamPath) in streamConfigs)
-					{
-						if (streamIds.Contains(streamRef.Id))
-						{
-							throw new Exception($"Project {projectRef.Id} contains duplicate stream id {streamRef.Id}");
-						}
-
-						if (streamPaths.Contains(streamPath))
-						{
-							throw new Exception($"Project {projectRef.Id} contains duplicate stream path {streamPath}");
-						}
-
-						streamIds.Add(streamRef.Id);
-						streamPaths.Add(streamPath);
-					}
-
-				}
-				catch (Exception ex)
-				{
-					await SendFailureNotificationAsync(ex, projectPath, cancellationToken);
-					continue;
-				}
 			}
 
 			// Get the logo revisions
