Horde: Separate stream template state into configuration and state objects, allowing us to avoid storing/fetching redundant data in the database.

diff --git a/Engine/Source/Programs/Horde/Horde.Build.Tests/Fixture.cs b/Engine/Source/Programs/Horde/Horde.Build.Tests/Fixture.cs
--- a/Engine/Source/Programs/Horde/Horde.Build.Tests/Fixture.cs
+++ b/Engine/Source/Programs/Horde/Horde.Build.Tests/Fixture.cs
@@ -24,7 +24,7 @@ namespace Horde.Build.Tests
 	using PoolId = StringId<IPool>;
     using ProjectId = StringId<IProject>;
     using StreamId = StringId<IStream>;
-    using TemplateRefId = StringId<TemplateRef>;
+    using TemplateId = StringId<ITemplateRef>;
 
 	public class Fixture
 	{
@@ -33,8 +33,8 @@ namespace Horde.Build.Tests
 		public ITemplate Template { get; private set; } = null!;
 		public IGraph Graph { get; private set; } = null!;
 		public IStream? Stream { get; private set; }
-		public TemplateRefId TemplateRefId1 { get; private set; }
-		public TemplateRefId TemplateRefId2 { get; private set; }
+		public TemplateId TemplateRefId1 { get; private set; }
+		public TemplateId TemplateRefId2 { get; private set; }
 		public IArtifact Job1Artifact { get; private set; } = null!;
 		public string Job1ArtifactData { get; private set; } = null!;
 		public IAgent Agent1 { get; private set; } = null!;
@@ -64,15 +64,15 @@ namespace Horde.Build.Tests
 			Template = await templateCollection.AddAsync("Test template");
 			Graph = await graphCollection.AddAsync(Template);
 
-			TemplateRefId1 = new TemplateRefId("template1");
-			TemplateRefId2 = new TemplateRefId("template2");
+			TemplateRefId1 = new TemplateId("template1");
+			TemplateRefId2 = new TemplateId("template2");
 
 			List<TemplateRefConfig> templates = new List<TemplateRefConfig>();
 			templates.Add(new TemplateRefConfig { Id = TemplateRefId1, Name = "Test Template" });
 			templates.Add(new TemplateRefConfig { Id = TemplateRefId2, Name = "Test Template" });
 
 			List<TabConfig> tabs = new List<TabConfig>();
-			tabs.Add(new JobsTabConfig { Title = "foo", Templates = new List<TemplateRefId> { TemplateRefId1, TemplateRefId2 } });
+			tabs.Add(new JobsTabConfig { Title = "foo", Templates = new List<TemplateId> { TemplateRefId1, TemplateRefId2 } });
 
 			Dictionary<string, AgentConfig> agentTypes = new()
 			{
@@ -119,7 +119,7 @@ namespace Horde.Build.Tests
 			Job2 = await jobService.CreateJobAsync(
 				jobId: new JobId("5f69ea1b68423e921b035106"),
 				stream: Stream!,
-				templateRefId: new TemplateRefId("template-id-1"),
+				templateRefId: new TemplateId("template-id-1"),
 				templateHash: ContentHash.MD5("made-up-template-hash"),
 				graph: fg,
 				name: "hello2",
diff --git a/Engine/Source/Programs/Horde/Horde.Build.Tests/Fleet/JobQueueStrategyTest.cs b/Engine/Source/Programs/Horde/Horde.Build.Tests/Fleet/JobQueueStrategyTest.cs
--- a/Engine/Source/Programs/Horde/Horde.Build.Tests/Fleet/JobQueueStrategyTest.cs
+++ b/Engine/Source/Programs/Horde/Horde.Build.Tests/Fleet/JobQueueStrategyTest.cs
@@ -149,7 +149,7 @@ namespace Horde.Build.Tests.Fleet
 		private async Task<IJob> AddPlaceholderJob(IGraph graph, StreamId streamId, string nodeNameToExecute)
 		{
 			IJob job = await JobCollection.AddAsync(ObjectId<IJob>.GenerateNewId(), streamId,
-				new StringId<TemplateRef>("bogusTemplateRefId"), ContentHash.Empty, graph, "bogusJobName",
+				new StringId<ITemplateRef>("bogusTemplateRefId"), ContentHash.Empty, graph, "bogusJobName",
 				1000, 1000, null, null, null, null, null, null, null, null, null, false,
 				false, null, null, new List<string> { "-Target=" + nodeNameToExecute });
 
diff --git a/Engine/Source/Programs/Horde/Horde.Build.Tests/IssueServiceTests.cs b/Engine/Source/Programs/Horde/Horde.Build.Tests/IssueServiceTests.cs
--- a/Engine/Source/Programs/Horde/Horde.Build.Tests/IssueServiceTests.cs
+++ b/Engine/Source/Programs/Horde/Horde.Build.Tests/IssueServiceTests.cs
@@ -37,7 +37,7 @@ namespace Horde.Build.Tests
 	using LogId = ObjectId<ILogFile>;
 	using ProjectId = StringId<IProject>;
 	using StreamId = StringId<IStream>;
-	using TemplateRefId = StringId<TemplateRef>;
+	using TemplateId = StringId<ITemplateRef>;
 	using UserId = ObjectId<IUser>;
 
 	[TestClass]
@@ -122,8 +122,8 @@ namespace Horde.Build.Tests
 			string revision = $"config:{streamId}";
 
 			StreamConfig streamConfig = new StreamConfig { Name = streamName };
-			streamConfig.Tabs.Add(new JobsTabConfig { Title = "General", Templates = new List<TemplateRefId> { new TemplateRefId("test-template") } });
-			streamConfig.Templates.Add(new TemplateRefConfig { Id = new TemplateRefId("test-template") });
+			streamConfig.Tabs.Add(new JobsTabConfig { Title = "General", Templates = new List<TemplateId> { new TemplateId("test-template") } });
+			streamConfig.Templates.Add(new TemplateRefConfig { Id = new TemplateId("test-template") });
 			await ConfigCollection.AddConfigAsync(revision, streamConfig);
 
 			return Deref(await StreamCollection.TryCreateOrReplaceAsync(streamId, null, revision, projectId));
@@ -236,7 +236,7 @@ namespace Horde.Build.Tests
 			job.SetupGet(x => x.Id).Returns(jobId);
 			job.SetupGet(x => x.Name).Returns(name);
 			job.SetupGet(x => x.StreamId).Returns(streamId);
-			job.SetupGet(x => x.TemplateId).Returns(new TemplateRefId("test-template"));
+			job.SetupGet(x => x.TemplateId).Returns(new TemplateId("test-template"));
 			job.SetupGet(x => x.Change).Returns(change);
 			job.SetupGet(x => x.Batches).Returns(batches);
 			job.SetupGet(x => x.ShowUgsBadges).Returns(promoteByDefault);
diff --git a/Engine/Source/Programs/Horde/Horde.Build.Tests/JobCollectionTests.cs b/Engine/Source/Programs/Horde/Horde.Build.Tests/JobCollectionTests.cs
--- a/Engine/Source/Programs/Horde/Horde.Build.Tests/JobCollectionTests.cs
+++ b/Engine/Source/Programs/Horde/Horde.Build.Tests/JobCollectionTests.cs
@@ -27,7 +27,7 @@ namespace Horde.Build.Tests
 	using LogId = ObjectId<ILogFile>;
 	using PoolId = StringId<IPool>;
 	using StreamId = StringId<IStream>;
-	using TemplateRefId = StringId<TemplateRef>;
+	using TemplateId = StringId<ITemplateRef>;
 
 	[TestClass]
 	public class JobCollectionTests : TestSetup
@@ -81,7 +81,7 @@ namespace Horde.Build.Tests
 			arguments.Add("-Target=Publish Client");
 			arguments.Add("-Target=Post-Publish Client");
 
-			IJob job = await JobCollection.AddAsync(JobId.GenerateNewId(), new StreamId("ue4-main"), new TemplateRefId("test-build"), ContentHash.SHA1("hello"), baseGraph, "Test job", 123, 123, null, null, null, null, null, null, null, null, null, false, false, null, null, arguments);
+			IJob job = await JobCollection.AddAsync(JobId.GenerateNewId(), new StreamId("ue4-main"), new TemplateId("test-build"), ContentHash.SHA1("hello"), baseGraph, "Test job", 123, 123, null, null, null, null, null, null, null, null, null, false, false, null, null, arguments);
 
 			job = await StartBatch(job, baseGraph, 0);
 			job = await RunStep(job, baseGraph, 0, 0, JobStepOutcome.Success); // Setup Build
@@ -156,7 +156,7 @@ namespace Horde.Build.Tests
 			arguments.Add("-Target=Step 1");
 			arguments.Add("-Target=Step 3");
 
-			IJob job = await JobCollection.AddAsync(JobId.GenerateNewId(), new StreamId("ue4-main"), new TemplateRefId("test-build"), ContentHash.SHA1("hello"), baseGraph, "Test job", 123, 123, null, null, null, null, null, null, null, null, null, false, false, null, null, arguments);
+			IJob job = await JobCollection.AddAsync(JobId.GenerateNewId(), new StreamId("ue4-main"), new TemplateId("test-build"), ContentHash.SHA1("hello"), baseGraph, "Test job", 123, 123, null, null, null, null, null, null, null, null, null, false, false, null, null, arguments);
 
 			job = await StartBatch(job, baseGraph, 0);
 			job = await RunStep(job, baseGraph, 0, 0, JobStepOutcome.Success); // Setup Build
@@ -212,7 +212,7 @@ namespace Horde.Build.Tests
 			arguments.Add("-Target=Step 1");
 			arguments.Add("-Target=Step 3");
 
-			IJob job = await JobCollection.AddAsync(JobId.GenerateNewId(), new StreamId("ue4-main"), new TemplateRefId("test-build"), ContentHash.SHA1("hello"), baseGraph, "Test job", 123, 123, null, null, null, null, null, null, null, null, null, false, false, null, null, arguments);
+			IJob job = await JobCollection.AddAsync(JobId.GenerateNewId(), new StreamId("ue4-main"), new TemplateId("test-build"), ContentHash.SHA1("hello"), baseGraph, "Test job", 123, 123, null, null, null, null, null, null, null, null, null, false, false, null, null, arguments);
 			Assert.AreEqual(1, job.Batches.Count);
 
 			job = await StartBatch(job, baseGraph, 0);
@@ -263,7 +263,7 @@ namespace Horde.Build.Tests
 			arguments.Add("-Target=Step 1");
 			arguments.Add("-Target=Step 3");
 
-			IJob job = await JobCollection.AddAsync(JobId.GenerateNewId(), new StreamId("ue4-main"), new TemplateRefId("test-build"), ContentHash.SHA1("hello"), baseGraph, "Test job", 123, 123, null, null, null, null, null, null, null, null, null, false, false, null, null, arguments);
+			IJob job = await JobCollection.AddAsync(JobId.GenerateNewId(), new StreamId("ue4-main"), new TemplateId("test-build"), ContentHash.SHA1("hello"), baseGraph, "Test job", 123, 123, null, null, null, null, null, null, null, null, null, false, false, null, null, arguments);
 			Assert.AreEqual(1, job.Batches.Count);
 
 			// First retry
diff --git a/Engine/Source/Programs/Horde/Horde.Build.Tests/JobServiceTests.cs b/Engine/Source/Programs/Horde/Horde.Build.Tests/JobServiceTests.cs
--- a/Engine/Source/Programs/Horde/Horde.Build.Tests/JobServiceTests.cs
+++ b/Engine/Source/Programs/Horde/Horde.Build.Tests/JobServiceTests.cs
@@ -25,7 +25,7 @@ namespace Horde.Build.Tests
 	using LogId = ObjectId<ILogFile>;
 	using ProjectId = StringId<IProject>;
 	using StreamId = StringId<IStream>;
-	using TemplateRefId = StringId<TemplateRef>;
+	using TemplateId = StringId<ITemplateRef>;
 
 	[TestClass]
 	public class JobServiceTests : TestSetup
@@ -40,19 +40,19 @@ namespace Horde.Build.Tests
 			ITemplate template = await TemplateCollection.AddAsync("Test template");
 			IGraph graph = await GraphCollection.AddAsync(template);
 
-			TemplateRefId templateRefId1 = new TemplateRefId("template1");
-			TemplateRefId templateRefId2 = new TemplateRefId("template2");
+			TemplateId templateRefId1 = new TemplateId("template1");
+			TemplateId templateRefId2 = new TemplateId("template2");
 
 			StreamConfig streamConfig = new StreamConfig();
-			streamConfig.Templates.Add(new TemplateRefConfig { Id = templateRefId1, Name = "Test Template", ChainedJobs = new List<ChainedJobTemplateConfig> { new ChainedJobTemplateConfig { TemplateId = templateRefId2.ToString(), Trigger = "Setup Build" } } });
+			streamConfig.Templates.Add(new TemplateRefConfig { Id = templateRefId1, Name = "Test Template", ChainedJobs = new List<ChainedJobTemplateConfig> { new ChainedJobTemplateConfig { TemplateId = templateRefId2, Trigger = "Setup Build" } } });
 			streamConfig.Templates.Add(new TemplateRefConfig { Id = templateRefId2, Name = "Test Template" });
-			streamConfig.Tabs.Add(new JobsTabConfig { Title = "foo", Templates = new List<TemplateRefId> { templateRefId1, templateRefId2 } });
+			streamConfig.Tabs.Add(new JobsTabConfig { Title = "foo", Templates = new List<TemplateId> { templateRefId1, templateRefId2 } });
 
 			StreamId streamId = new StreamId("ue5-main");
 			IStream? stream = await StreamService.GetStreamAsync(streamId);
 			stream = await CreateOrReplaceStreamAsync(new StreamId("ue5-main"), stream, projectId, streamConfig);
 
-			IJob job = await JobService.CreateJobAsync(null, stream!, templateRefId1, template.Id, graph, "Hello", 1234, 1233, 999, null, null, null, null, null, null, null, stream!.Templates[templateRefId1].ChainedJobs, true, true, null, null, new List<string>());
+			IJob job = await JobService.CreateJobAsync(null, stream!, templateRefId1, template.Id, graph, "Hello", 1234, 1233, 999, null, null, null, null, null, null, null, stream!.Templates[templateRefId1].Config.ChainedJobs, true, true, null, null, new List<string>());
 			Assert.AreEqual(1, job.ChainedJobs.Count);
 
 			job = Deref(await JobService.UpdateBatchAsync(job, job.Batches[0].Id, LogId.GenerateNewId(), JobStepBatchState.Running));
@@ -105,7 +105,7 @@ namespace Horde.Build.Tests
 			return await JobService.CreateJobAsync(
 				jobId: JobId.GenerateNewId(),
 				stream: fixture!.Stream!,
-				templateRefId: new TemplateRefId(templateRefId),
+				templateRefId: new TemplateId(templateRefId),
 				templateHash: new ContentHash(Encoding.ASCII.GetBytes(templateHash)),
 				graph: fixture!.Graph,
 				name: "hello1",
@@ -178,7 +178,7 @@ namespace Horde.Build.Tests
 
 			graph = await GraphCollection.AppendAsync(graph, new List<NewGroup> { groupA, groupB });
 
-			IJob job = await JobService.CreateJobAsync(null, stream!, new TemplateRefId("temp"), template.Id, graph, "Hello", 1234, 1233, 999, null, null, null, null, null, null, null, null, true, true, null, null, new List<string> { "-Target=Pak" });
+			IJob job = await JobService.CreateJobAsync(null, stream!, new TemplateId("temp"), template.Id, graph, "Hello", 1234, 1233, 999, null, null, null, null, null, null, null, null, true, true, null, null, new List<string> { "-Target=Pak" });
 
 			job = Deref(await JobService.UpdateBatchAsync(job, job.Batches[0].Id, LogId.GenerateNewId(), JobStepBatchState.Running));
 			job = Deref(await JobService.UpdateStepAsync(job, job.Batches[0].Id, job.Batches[0].Steps[0].Id, JobStepState.Running));
diff --git a/Engine/Source/Programs/Horde/Horde.Build.Tests/SchedulerTests.cs b/Engine/Source/Programs/Horde/Horde.Build.Tests/SchedulerTests.cs
--- a/Engine/Source/Programs/Horde/Horde.Build.Tests/SchedulerTests.cs
+++ b/Engine/Source/Programs/Horde/Horde.Build.Tests/SchedulerTests.cs
@@ -25,14 +25,14 @@ namespace Horde.Build.Tests
 	using LogId = ObjectId<ILogFile>;
 	using ProjectId = StringId<IProject>;
 	using StreamId = StringId<IStream>;
-	using TemplateRefId = StringId<TemplateRef>;
+	using TemplateId = StringId<ITemplateRef>;
 
 	[TestClass]
     public class SchedulerTests : TestSetup
 	{
 		ProjectId ProjectId { get; } = new ProjectId("ue5");
 		StreamId StreamId { get; } = new StreamId("ue5-main");
-		TemplateRefId TemplateRefId { get; } = new TemplateRefId("template1");
+		TemplateId TemplateId { get; } = new TemplateId("template1");
 
 		readonly ITemplate _template;
 		readonly HashSet<JobId> _initialJobIds;
@@ -62,8 +62,8 @@ namespace Horde.Build.Tests
 
 			StreamConfig config = new StreamConfig();
 			config.Name = "//UE5/Main";
-			config.Tabs.Add(new JobsTabConfig { Title = "foo", Templates = new List<TemplateRefId> { TemplateRefId } });
-			config.Templates.Add(new TemplateRefConfig { Id = TemplateRefId, Name = "Test", Schedule = schedule });
+			config.Tabs.Add(new JobsTabConfig { Title = "foo", Templates = new List<TemplateId> { TemplateId } });
+			config.Templates.Add(new TemplateRefConfig { Id = TemplateId, Name = "Test", Schedule = schedule });
 
 			return (await CreateOrReplaceStreamAsync(StreamId, stream, ProjectId, config))!; 
 		}
@@ -132,8 +132,8 @@ namespace Horde.Build.Tests
 			DateTime startTime = new DateTime(2021, 1, 1, 12, 0, 0, DateTimeKind.Utc); // Friday Jan 1, 2021 
 			Clock.UtcNow = startTime;
 
-			Schedule schedule = new Schedule(Clock.UtcNow, requireSubmittedChange: false);
-			schedule.Patterns.Add(new SchedulePattern(new List<DayOfWeek> { DayOfWeek.Friday, DayOfWeek.Sunday }, 13 * 60, null, null));
+			ScheduleConfig schedule = new ScheduleConfig { RequireSubmittedChange = false };
+			schedule.Patterns.Add(new SchedulePatternConfig(new List<DayOfWeek> { DayOfWeek.Friday, DayOfWeek.Sunday }, 13 * 60, null, null));
 
 			DateTime? nextTime = schedule.GetNextTriggerTimeUtc(startTime, TimeZoneInfo.Utc);
 			Assert.AreEqual(startTime + TimeSpan.FromHours(1.0), nextTime!.Value);
@@ -154,8 +154,8 @@ namespace Horde.Build.Tests
 			DateTime startTime = new DateTime(2021, 1, 1, 12, 0, 0, DateTimeKind.Utc); // Friday Jan 1, 2021 
 			Clock.UtcNow = startTime;
 
-			Schedule schedule = new Schedule(Clock.UtcNow, requireSubmittedChange: false);
-			schedule.Patterns.Add(new SchedulePattern(null, 13 * 60, 14 * 60, 15));
+			ScheduleConfig schedule = new ScheduleConfig { RequireSubmittedChange = false };
+			schedule.Patterns.Add(new SchedulePatternConfig(null, 13 * 60, 14 * 60, 15));
 
 			DateTime? nextTime = schedule.GetNextTriggerTimeUtc(startTime, TimeZoneInfo.Utc);
 			Assert.AreEqual(startTime + TimeSpan.FromHours(1.0), nextTime!.Value);
@@ -182,9 +182,9 @@ namespace Horde.Build.Tests
 			DateTime startTime = new DateTime(2021, 1, 1, 0, 0, 0, DateTimeKind.Utc); // Friday Jan 1, 2021 
 			Clock.UtcNow = startTime;
 
-			Schedule schedule = new Schedule(Clock.UtcNow, requireSubmittedChange: false);
-			schedule.Patterns.Add(new SchedulePattern(null, 11 * 60, 0, 0));
-			schedule.Patterns.Add(new SchedulePattern(null, 19 * 60, 0, 0));
+			ScheduleConfig schedule = new ScheduleConfig { RequireSubmittedChange = false };
+			schedule.Patterns.Add(new SchedulePatternConfig(null, 11 * 60, 0, 0));
+			schedule.Patterns.Add(new SchedulePatternConfig(null, 19 * 60, 0, 0));
 
 			DateTime? nextTime = schedule.GetNextTriggerTimeUtc(startTime, TimeZoneInfo.Utc);
 			Assert.AreEqual(startTime + TimeSpan.FromHours(11), nextTime!.Value);
@@ -248,9 +248,9 @@ namespace Horde.Build.Tests
 			Assert.AreEqual(100, jobs2[0].CodeChange);
 
 			IStream stream2 = (await StreamCollection.GetAsync(stream.Id))!;
-			Schedule schedule2 = stream2.Templates.First().Value.Schedule!;
+			ITemplateSchedule schedule2 = stream2.Templates.First().Value.Schedule!;
 			Assert.AreEqual(102, schedule2.LastTriggerChange);
-			Assert.AreEqual(Clock.UtcNow, schedule2.LastTriggerTime);
+			Assert.AreEqual(Clock.UtcNow, schedule2.LastTriggerTimeUtc);
 
 			// Trigger another job
 			await Clock.AdvanceAsync(TimeSpan.FromHours(0.5));
@@ -413,21 +413,21 @@ namespace Horde.Build.Tests
 			// Create two templates, the second dependent on the first
 			ITemplate? newTemplate1 = await TemplateCollection.AddAsync("Test template 1");
 			//TemplateRef newTemplateRef1 = new TemplateRef(newTemplate1);
-			TemplateRefId newTemplateRefId1 = new TemplateRefId("new-template-1");
+			TemplateId newTemplateRefId1 = new TemplateId("new-template-1");
 
 			ITemplate? newTemplate2 = await TemplateCollection.AddAsync("Test template 2");
-			TemplateRef newTemplateRef2 = new TemplateRef(newTemplate2);
-			newTemplateRef2.Schedule = new Schedule(Clock.UtcNow);
-			newTemplateRef2.Schedule.Gate = new ScheduleGate(newTemplateRefId1, "TriggerNext");
-			newTemplateRef2.Schedule.Patterns.Add(new SchedulePattern(null, 0, null, 10));
-			newTemplateRef2.Schedule.LastTriggerTime = startTime;
-			TemplateRefId newTemplateRefId2 = new TemplateRefId("new-template-2");
+//			TemplateRef newTemplateRef2 = new TemplateRef(newTemplate2);
+//			newTemplateRef2.Schedule = new Schedule(Clock.UtcNow);
+//			newTemplateRef2.Schedule.Gate = new ScheduleGate(newTemplateRefId1, "TriggerNext");
+//			newTemplateRef2.Schedule.Patterns.Add(new SchedulePattern(null, 0, null, 10));
+//			newTemplateRef2.Schedule.LastTriggerTime = startTime;
+			TemplateId newTemplateRefId2 = new TemplateId("new-template-2");
 
 			IStream? stream = await StreamService.GetStreamAsync(StreamId);
 
 			StreamConfig config = new StreamConfig();
 			config.Name = "//UE5/Main";
-			config.Tabs.Add(new JobsTabConfig { Title = "foo", Templates = new List<TemplateRefId> { newTemplateRefId1, newTemplateRefId2 } });
+			config.Tabs.Add(new JobsTabConfig { Title = "foo", Templates = new List<TemplateId> { newTemplateRefId1, newTemplateRefId2 } });
 
 			stream = (await CreateOrReplaceStreamAsync(StreamId, stream, ProjectId, config))!;
 
@@ -488,18 +488,18 @@ namespace Horde.Build.Tests
 			PerforceService.AddChange(StreamId, 1233, bob, "", new[] { "code.cpp" });
 
 			// Create two templates, the second dependent on the first
-			TemplateRefId newTemplateRefId1 = new TemplateRefId("new-template-1");
+			TemplateId newTemplateRefId1 = new TemplateId("new-template-1");
 			TemplateRefConfig newTemplate1 = new TemplateRefConfig();
 			newTemplate1.Id = newTemplateRefId1;
 
-			TemplateRefId newTemplateRefId2 = new TemplateRefId("new-template-2");
+			TemplateId newTemplateRefId2 = new TemplateId("new-template-2");
 			TemplateRefConfig newTemplate2 = new TemplateRefConfig();
 			newTemplate2.Id = newTemplateRefId2;
 			newTemplate2.Name = "Test template 2";
 			newTemplate2.Schedule = new ScheduleConfig();
 			newTemplate2.Schedule.MaxChanges = 4;
 			newTemplate2.Schedule.Filter = new List<ChangeContentFlags> { ChangeContentFlags.ContainsCode };
-			newTemplate2.Schedule.Gate = new ScheduleGateConfig { TemplateId = newTemplateRefId1.ToString(), Target = "TriggerNext" };
+			newTemplate2.Schedule.Gate = new ScheduleGateConfig { TemplateId = newTemplateRefId1, Target = "TriggerNext" };
 			newTemplate2.Schedule.Patterns.Add(new SchedulePatternConfig { Interval = 10 });// (null, 0, null, 10));
 //			NewTemplate2.Schedule.LastTriggerTime = StartTime;
 
@@ -507,7 +507,7 @@ namespace Horde.Build.Tests
 
 			StreamConfig config = new StreamConfig();
 			config.Name = "//UE5/Main";
-			config.Tabs.Add(new JobsTabConfig { Title = "foo", Templates = new List<TemplateRefId> { newTemplateRefId1, newTemplateRefId2 } });
+			config.Tabs.Add(new JobsTabConfig { Title = "foo", Templates = new List<TemplateId> { newTemplateRefId1, newTemplateRefId2 } });
 			config.Templates.Add(newTemplate1);
 			config.Templates.Add(newTemplate2);
 
@@ -575,7 +575,7 @@ namespace Horde.Build.Tests
 
 			// Make sure the job is registered
 			IStream? stream1 = await StreamService.GetStreamAsync(StreamId);
-			TemplateRef templateRef1 = stream1!.Templates.First().Value;
+			ITemplateRef templateRef1 = stream1!.Templates.First().Value;
 			Assert.AreEqual(1, templateRef1.Schedule!.ActiveJobs.Count);
 			Assert.AreEqual(jobs1[0].Id, templateRef1.Schedule!.ActiveJobs[0]);
 
@@ -584,7 +584,7 @@ namespace Horde.Build.Tests
 
 			// Make sure the job is still registered
 			IStream? stream2 = await StreamService.GetStreamAsync(StreamId);
-			TemplateRef templateRef2 = stream2!.Templates.First().Value;
+			ITemplateRef templateRef2 = stream2!.Templates.First().Value;
 			Assert.AreEqual(1, templateRef2.Schedule!.ActiveJobs.Count);
 			Assert.AreEqual(jobs1[0].Id, templateRef2.Schedule!.ActiveJobs[0]);
 		}
diff --git a/Engine/Source/Programs/Horde/Horde.Build.Tests/Stubs/Collections/JobStepRefCollectionStub.cs b/Engine/Source/Programs/Horde/Horde.Build.Tests/Stubs/Collections/JobStepRefCollectionStub.cs
--- a/Engine/Source/Programs/Horde/Horde.Build.Tests/Stubs/Collections/JobStepRefCollectionStub.cs
+++ b/Engine/Source/Programs/Horde/Horde.Build.Tests/Stubs/Collections/JobStepRefCollectionStub.cs
@@ -17,7 +17,7 @@ namespace Horde.Build.Tests.Stubs.Collections
 	using LogId = ObjectId<ILogFile>;
 	using PoolId = StringId<IPool>;
 	using StreamId = StringId<IStream>;
-	using TemplateRefId = StringId<TemplateRef>;
+	using TemplateRefId = StringId<ITemplateRef>;
 
 	class JobStepRefStub : IJobStepRef
 	{
diff --git a/Engine/Source/Programs/Horde/Horde.Build/Issues/IIssueCollection.cs b/Engine/Source/Programs/Horde/Horde.Build/Issues/IIssueCollection.cs
--- a/Engine/Source/Programs/Horde/Horde.Build/Issues/IIssueCollection.cs
+++ b/Engine/Source/Programs/Horde/Horde.Build/Issues/IIssueCollection.cs
@@ -19,7 +19,7 @@ namespace Horde.Build.Issues
 	using JobId = ObjectId<IJob>;
 	using LogId = ObjectId<ILogFile>;
 	using StreamId = StringId<IStream>;
-	using TemplateRefId = StringId<TemplateRef>;
+	using TemplateId = StringId<ITemplateRef>;
 	using UserId = ObjectId<IUser>;
 
 	/// <summary>
@@ -191,7 +191,7 @@ namespace Horde.Build.Issues
 		public string StreamName { get; set; }
 
 		/// <inheritdoc cref="IIssueSpan.TemplateRefId"/>
-		public TemplateRefId TemplateRefId { get; set; }
+		public TemplateId TemplateRefId { get; set; }
 
 		/// <inheritdoc cref="IIssueSpan.NodeName"/>
 		public string NodeName { get; set; }
@@ -214,7 +214,7 @@ namespace Horde.Build.Issues
 		/// <summary>
 		/// Constructor
 		/// </summary>
-		public NewIssueSpanData(StreamId streamId, string streamName, TemplateRefId templateRefId, string nodeName, NewIssueFingerprint fingerprint, NewIssueStepData firstFailure)
+		public NewIssueSpanData(StreamId streamId, string streamName, TemplateId templateRefId, string nodeName, NewIssueFingerprint fingerprint, NewIssueStepData firstFailure)
 		{
 			StreamId = streamId;
 			StreamName = streamName;
@@ -536,7 +536,7 @@ namespace Horde.Build.Issues
 		/// <param name="name">Name of the node</param>
 		/// <param name="change">Changelist number to query</param>
 		/// <returns>List of open issues</returns>
-		Task<List<IIssueSpan>> FindOpenSpansAsync(StreamId streamId, TemplateRefId templateId, string name, int change);
+		Task<List<IIssueSpan>> FindOpenSpansAsync(StreamId streamId, TemplateId templateId, string name, int change);
 
 		/// <summary>
 		/// Searches for open issues
diff --git a/Engine/Source/Programs/Horde/Horde.Build/Issues/IIssueSpan.cs b/Engine/Source/Programs/Horde/Horde.Build/Issues/IIssueSpan.cs
--- a/Engine/Source/Programs/Horde/Horde.Build/Issues/IIssueSpan.cs
+++ b/Engine/Source/Programs/Horde/Horde.Build/Issues/IIssueSpan.cs
@@ -9,7 +9,7 @@ using MongoDB.Bson;
 namespace Horde.Build.Issues
 {
 	using StreamId = StringId<IStream>;
-	using TemplateRefId = StringId<TemplateRef>;
+	using TemplateId = StringId<ITemplateRef>;
 	using UserId = ObjectId<IUser>;
 
 	/// <summary>
@@ -35,7 +35,7 @@ namespace Horde.Build.Issues
 		/// <summary>
 		/// The template containing this node
 		/// </summary>
-		public TemplateRefId TemplateRefId { get; }
+		public TemplateId TemplateRefId { get; }
 
 		/// <summary>
 		/// Name of the node
diff --git a/Engine/Source/Programs/Horde/Horde.Build/Issues/IssueCollection.cs b/Engine/Source/Programs/Horde/Horde.Build/Issues/IssueCollection.cs
--- a/Engine/Source/Programs/Horde/Horde.Build/Issues/IssueCollection.cs
+++ b/Engine/Source/Programs/Horde/Horde.Build/Issues/IssueCollection.cs
@@ -25,7 +25,7 @@ namespace Horde.Build.Issues
 	using JobId = ObjectId<IJob>;
 	using LogId = ObjectId<ILogFile>;
 	using StreamId = StringId<IStream>;
-	using TemplateRefId = StringId<TemplateRef>;
+	using TemplateId = StringId<ITemplateRef>;
 	using UserId = ObjectId<IUser>;
 
 	class IssueCollection : IIssueCollection
@@ -258,7 +258,7 @@ namespace Horde.Build.Issues
 			public string StreamName { get; set; }
 
 			[BsonRequired]
-			public TemplateRefId TemplateRefId { get; set; }
+			public TemplateId TemplateRefId { get; set; }
 
 			[BsonRequired]
 			public string NodeName { get; set; }
@@ -1274,7 +1274,7 @@ namespace Horde.Build.Issues
 		}
 
 		/// <inheritdoc/>
-		public async Task<List<IIssueSpan>> FindOpenSpansAsync(StreamId streamId, TemplateRefId templateId, string nodeName, int change)
+		public async Task<List<IIssueSpan>> FindOpenSpansAsync(StreamId streamId, TemplateId templateId, string nodeName, int change)
 		{
 			List<IssueSpan> spans = await _issueSpans.Find(x => x.StreamId == streamId && x.TemplateRefId == templateId && x.NodeName == nodeName && change >= x.MinChange && change <= x.MaxChange).ToListAsync();
 			return spans.ConvertAll<IIssueSpan>(x => x);
diff --git a/Engine/Source/Programs/Horde/Horde.Build/Issues/IssueMessages.cs b/Engine/Source/Programs/Horde/Horde.Build/Issues/IssueMessages.cs
--- a/Engine/Source/Programs/Horde/Horde.Build/Issues/IssueMessages.cs
+++ b/Engine/Source/Programs/Horde/Horde.Build/Issues/IssueMessages.cs
@@ -13,7 +13,7 @@ using MongoDB.Bson;
 namespace Horde.Build.Issues
 {
 	using StreamId = StringId<IStream>;
-	using TemplateRefId = StringId<TemplateRef>;
+	using TemplateId = StringId<ITemplateRef>;
 	using WorkflowId = StringId<WorkflowConfig>;
 
 	/// <summary>
@@ -340,12 +340,12 @@ namespace Horde.Build.Issues
 			Resolved = spans.All(x => x.NextSuccess != null);
 
 			AffectedTemplates = new List<GetIssueAffectedTemplateResponse>();
-			foreach (IGrouping<TemplateRefId, IIssueSpan> template in spans.GroupBy(x => x.TemplateRefId))
+			foreach (IGrouping<TemplateId, IIssueSpan> template in spans.GroupBy(x => x.TemplateRefId))
 			{
 				string templateName = template.Key.ToString();
-				if (stream != null && stream.Templates.TryGetValue(template.Key, out TemplateRef? templateRef))
+				if (stream != null && stream.Templates.TryGetValue(template.Key, out ITemplateRef? templateRef))
 				{
-					templateName = templateRef.Name;
+					templateName = templateRef.Config.Name;
 				}
 
 				HashSet<ObjectId> unresolvedTemplateSpans = new HashSet<ObjectId>(template.Where(x => x.NextSuccess == null).Select(x => x.Id));
@@ -355,10 +355,10 @@ namespace Horde.Build.Issues
 				AffectedTemplates.Add(new GetIssueAffectedTemplateResponse(template.Key.ToString(), templateName, template.All(x => x.NextSuccess != null), templateStep?.Severity ?? IssueSeverity.Unspecified));
 			}
 
-			HashSet<TemplateRefId> templateIdsSet = new HashSet<TemplateRefId>(spans.Select(x => x.TemplateRefId));
+			HashSet<TemplateId> templateIdsSet = new HashSet<TemplateId>(spans.Select(x => x.TemplateRefId));
 			TemplateIds = templateIdsSet.Select(x => x.ToString()).ToList();
 
-			HashSet<TemplateRefId> unresolvedTemplateIdsSet = new HashSet<TemplateRefId>(spans.Where(x => x.NextSuccess == null).Select(x => x.TemplateRefId));
+			HashSet<TemplateId> unresolvedTemplateIdsSet = new HashSet<TemplateId>(spans.Where(x => x.NextSuccess == null).Select(x => x.TemplateRefId));
 			UnresolvedTemplateIds = unresolvedTemplateIdsSet.Select(x => x.ToString()).ToList();
 			ResolvedTemplateIds = templateIdsSet.Except(unresolvedTemplateIdsSet).Select(x => x.ToString()).ToList();
 		}
diff --git a/Engine/Source/Programs/Horde/Horde.Build/Issues/IssueReportService.cs b/Engine/Source/Programs/Horde/Horde.Build/Issues/IssueReportService.cs
--- a/Engine/Source/Programs/Horde/Horde.Build/Issues/IssueReportService.cs
+++ b/Engine/Source/Programs/Horde/Horde.Build/Issues/IssueReportService.cs
@@ -21,7 +21,7 @@ using MongoDB.Bson.Serialization.Options;
 
 namespace Horde.Build.Issues
 {
-	using TemplateRefId = StringId<TemplateRef>;
+	using TemplateId = StringId<ITemplateRef>;
 	using WorkflowId = StringId<WorkflowConfig>;
 
 #pragma warning disable CS1591 // Missing XML comment for publicly visible type or member
@@ -227,7 +227,7 @@ namespace Horde.Build.Issues
 			List<IJob> jobs = await _jobCollection.FindAsync(streamId: stream.Id, minCreateTime: minTime);
 
 			Dictionary<WorkflowId, WorkflowStats> workflowIdToStats = new Dictionary<WorkflowId, WorkflowStats>();
-			foreach (IGrouping<TemplateRefId, IJob> templateGroup in jobs.GroupBy(x => x.TemplateId))
+			foreach (IGrouping<TemplateId, IJob> templateGroup in jobs.GroupBy(x => x.TemplateId))
 			{
 				WorkflowId? templateWorkflowId = null;
 				if (stream.Config.TryGetTemplate(templateGroup.Key, out TemplateRefConfig? templateRefConfig))
diff --git a/Engine/Source/Programs/Horde/Horde.Build/Issues/IssueService.cs b/Engine/Source/Programs/Horde/Horde.Build/Issues/IssueService.cs
--- a/Engine/Source/Programs/Horde/Horde.Build/Issues/IssueService.cs
+++ b/Engine/Source/Programs/Horde/Horde.Build/Issues/IssueService.cs
@@ -30,7 +30,7 @@ using OpenTracing.Util;
 namespace Horde.Build.Issues
 {
 	using StreamId = StringId<IStream>;
-	using TemplateRefId = StringId<TemplateRef>;
+	using TemplateId = StringId<ITemplateRef>;
 	using UserId = ObjectId<IUser>;
 	using WorkflowId = StringId<WorkflowConfig>;
 
@@ -189,7 +189,7 @@ namespace Horde.Build.Issues
 		/// <summary>
 		/// Cache of templates to show desktop alerts for
 		/// </summary>
-		Dictionary<StreamId, HashSet<TemplateRefId>> _cachedDesktopAlerts = new Dictionary<StreamId, HashSet<TemplateRefId>>();
+		Dictionary<StreamId, HashSet<TemplateId>> _cachedDesktopAlerts = new Dictionary<StreamId, HashSet<TemplateId>>();
 
 		/// <summary>
 		/// Logger for tracing
@@ -254,12 +254,12 @@ namespace Horde.Build.Issues
 		{
 			DateTime utcNow = _clock.UtcNow;
 
-			Dictionary<StreamId, HashSet<TemplateRefId>> newCachedDesktopAlerts = new Dictionary<StreamId, HashSet<TemplateRefId>>();
+			Dictionary<StreamId, HashSet<TemplateId>> newCachedDesktopAlerts = new Dictionary<StreamId, HashSet<TemplateId>>();
 
 			List<IStream> cachedStreams = await _streams.StreamCollection.FindAllAsync();
 			foreach(IStream cachedStream in cachedStreams)
 			{
-				HashSet<TemplateRefId> templates = new HashSet<TemplateRefId>(cachedStream.Templates.Where(x => x.Value.ShowUgsAlerts).Select(x => x.Key));
+				HashSet<TemplateId> templates = new HashSet<TemplateId>(cachedStream.Templates.Where(x => x.Value.Config.ShowUgsAlerts).Select(x => x.Key));
 				if(templates.Count > 0)
 				{
 					newCachedDesktopAlerts[cachedStream.Id] = templates;
@@ -363,12 +363,12 @@ namespace Horde.Build.Issues
 
 			bool bShowDesktopAlerts = false;
 
-			HashSet<(StreamId, TemplateRefId)> checkedTemplates = new HashSet<(StreamId, TemplateRefId)>();
+			HashSet<(StreamId, TemplateId)> checkedTemplates = new HashSet<(StreamId, TemplateId)>();
 			foreach(IIssueSpan span in spans)
 			{
 				if (span.NextSuccess == null && checkedTemplates.Add((span.StreamId, span.TemplateRefId)))
 				{
-					HashSet<TemplateRefId>? templates;
+					HashSet<TemplateId>? templates;
 					if (_cachedDesktopAlerts.TryGetValue(span.StreamId, out templates) && templates.Contains(span.TemplateRefId))
 					{
 						bShowDesktopAlerts = true;
diff --git a/Engine/Source/Programs/Horde/Horde.Build/Jobs/IJob.cs b/Engine/Source/Programs/Horde/Horde.Build/Jobs/IJob.cs
--- a/Engine/Source/Programs/Horde/Horde.Build/Jobs/IJob.cs
+++ b/Engine/Source/Programs/Horde/Horde.Build/Jobs/IJob.cs
@@ -30,7 +30,7 @@ namespace Horde.Build.Jobs
 	using ReportPlacement = HordeCommon.Rpc.ReportPlacement;
 	using SessionId = ObjectId<ISession>;
 	using StreamId = StringId<IStream>;
-	using TemplateRefId = StringId<TemplateRef>;
+	using TemplateId = StringId<ITemplateRef>;
 	using UserId = ObjectId<IUser>;
 
 	/// <summary>
@@ -528,7 +528,7 @@ namespace Horde.Build.Jobs
 		/// <summary>
 		/// The template to trigger on success
 		/// </summary>
-		public TemplateRefId TemplateRefId { get; }
+		public TemplateId TemplateRefId { get; }
 
 		/// <summary>
 		/// The triggered job id
@@ -564,7 +564,7 @@ namespace Horde.Build.Jobs
 		/// <summary>
 		/// The template ref id
 		/// </summary>
-		public TemplateRefId TemplateId { get; }
+		public TemplateId TemplateId { get; }
 
 		/// <summary>
 		/// The template that this job was created from
diff --git a/Engine/Source/Programs/Horde/Horde.Build/Jobs/IJobCollection.cs b/Engine/Source/Programs/Horde/Horde.Build/Jobs/IJobCollection.cs
--- a/Engine/Source/Programs/Horde/Horde.Build/Jobs/IJobCollection.cs
+++ b/Engine/Source/Programs/Horde/Horde.Build/Jobs/IJobCollection.cs
@@ -24,7 +24,7 @@ namespace Horde.Build.Jobs
 	using PoolId = StringId<IPool>;
 	using SessionId = ObjectId<ISession>;
 	using StreamId = StringId<IStream>;
-	using TemplateRefId = StringId<TemplateRef>;
+	using TemplateId = StringId<ITemplateRef>;
 	using UserId = ObjectId<IUser>;
 
 	/// <summary>
@@ -58,7 +58,7 @@ namespace Horde.Build.Jobs
 		/// <param name="notificationChannelFilter">Notification channel filter for this job</param>
 		/// <param name="arguments">Arguments for the job</param>
 		/// <returns>The new job document</returns>
-		Task<IJob> AddAsync(JobId jobId, StreamId streamId, TemplateRefId templateRefId, ContentHash templateHash, IGraph graph, string name, int change, int codeChange, int? preflightChange, int? clonedPreflightChange, string? preflightDescription, UserId? startedByUserId, Priority? priority, bool? autoSubmit, bool? updateIssues, bool? promoteIssuesByDefault, List<ChainedJobTemplate>? jobTriggers, bool showUgsBadges, bool showUgsAlerts, string? notificationChannel, string? notificationChannelFilter, List<string>? arguments);
+		Task<IJob> AddAsync(JobId jobId, StreamId streamId, TemplateId templateRefId, ContentHash templateHash, IGraph graph, string name, int change, int codeChange, int? preflightChange, int? clonedPreflightChange, string? preflightDescription, UserId? startedByUserId, Priority? priority, bool? autoSubmit, bool? updateIssues, bool? promoteIssuesByDefault, List<ChainedJobTemplateConfig>? jobTriggers, bool showUgsBadges, bool showUgsAlerts, string? notificationChannel, string? notificationChannelFilter, List<string>? arguments);
 
 		/// <summary>
 		/// Gets a job with the given unique id
@@ -110,7 +110,7 @@ namespace Horde.Build.Jobs
 		/// <param name="indexHint">Name of index to be specified as a hint to the database query planner</param>
 		/// <param name="excludeUserJobs">Whether to exclude user jobs from the find</param>
 		/// <returns>List of jobs matching the given criteria</returns>
-		Task<List<IJob>> FindAsync(JobId[]? jobIds = null, StreamId? streamId = null, string? name = null, TemplateRefId[]? templates = null, int? minChange = null, int? maxChange = null, int? preflightChange = null, bool? preflightOnly = null, UserId? preflightStartedByUser = null, UserId? startedByUser = null, DateTimeOffset? minCreateTime = null, DateTimeOffset? maxCreateTime = null, DateTimeOffset? modifiedBefore = null, DateTimeOffset? modifiedAfter = null, int? index = null, int? count = null, bool consistentRead = true, string? indexHint = null, bool? excludeUserJobs = null);
+		Task<List<IJob>> FindAsync(JobId[]? jobIds = null, StreamId? streamId = null, string? name = null, TemplateId[]? templates = null, int? minChange = null, int? maxChange = null, int? preflightChange = null, bool? preflightOnly = null, UserId? preflightStartedByUser = null, UserId? startedByUser = null, DateTimeOffset? minCreateTime = null, DateTimeOffset? maxCreateTime = null, DateTimeOffset? modifiedBefore = null, DateTimeOffset? modifiedAfter = null, int? index = null, int? count = null, bool consistentRead = true, string? indexHint = null, bool? excludeUserJobs = null);
 
 		/// <summary>
 		/// Searches for jobs in a specified stream and templates
@@ -123,7 +123,7 @@ namespace Horde.Build.Jobs
 		/// <param name="index">Index of the first result to return</param>
 		/// <param name="count">Number of results to return</param>
 		/// <param name="consistentRead">If the database read should be made to the replica server</param>
-		Task<List<IJob>> FindLatestByStreamWithTemplatesAsync(StreamId streamId, TemplateRefId[] templates, UserId? preflightStartedByUser = null, DateTimeOffset? maxCreateTime = null, DateTimeOffset? modifiedAfter = null, int? index = null, int? count = null, bool consistentRead = false);
+		Task<List<IJob>> FindLatestByStreamWithTemplatesAsync(StreamId streamId, TemplateId[] templates, UserId? preflightStartedByUser = null, DateTimeOffset? maxCreateTime = null, DateTimeOffset? modifiedAfter = null, int? index = null, int? count = null, bool consistentRead = false);
 
 		/// <summary>
 		/// Updates a new job
@@ -141,7 +141,7 @@ namespace Horde.Build.Jobs
 		/// <param name="arguments">New arguments for the job</param>
 		/// <param name="labelIdxToTriggerId">New trigger ID for a label in the job</param>
 		/// <param name="jobTrigger">New downstream job id</param>
-		Task<IJob?> TryUpdateJobAsync(IJob job, IGraph graph, string? name = null, Priority? priority = null, bool? autoSubmit = null, int? autoSubmitChange = null, string? autoSubmitMessage = null, UserId? abortedByUserId = null, ObjectId? notificationTriggerId = null, List<Report>? reports = null, List<string>? arguments = null, KeyValuePair<int, ObjectId>? labelIdxToTriggerId = null, KeyValuePair<TemplateRefId, JobId>? jobTrigger = null);
+		Task<IJob?> TryUpdateJobAsync(IJob job, IGraph graph, string? name = null, Priority? priority = null, bool? autoSubmit = null, int? autoSubmitChange = null, string? autoSubmitMessage = null, UserId? abortedByUserId = null, ObjectId? notificationTriggerId = null, List<Report>? reports = null, List<string>? arguments = null, KeyValuePair<int, ObjectId>? labelIdxToTriggerId = null, KeyValuePair<TemplateId, JobId>? jobTrigger = null);
 
 		/// <summary>
 		/// Updates the state of a batch
diff --git a/Engine/Source/Programs/Horde/Horde.Build/Jobs/IJobStepRef.cs b/Engine/Source/Programs/Horde/Horde.Build/Jobs/IJobStepRef.cs
--- a/Engine/Source/Programs/Horde/Horde.Build/Jobs/IJobStepRef.cs
+++ b/Engine/Source/Programs/Horde/Horde.Build/Jobs/IJobStepRef.cs
@@ -16,7 +16,7 @@ namespace Horde.Build.Jobs
 	using LogId = ObjectId<ILogFile>;
 	using PoolId = StringId<IPool>;
 	using StreamId = StringId<IStream>;
-	using TemplateRefId = StringId<TemplateRef>;
+	using TemplateId = StringId<ITemplateRef>;
 
 	/// <summary>
 	/// Unique id struct for JobStepRef objects. Includes a job id, batch id, and step id to uniquely identify the step.
@@ -134,7 +134,7 @@ namespace Horde.Build.Jobs
 		/// <summary>
 		/// Template for the job being executed
 		/// </summary>
-		public TemplateRefId TemplateId { get; }
+		public TemplateId TemplateId { get; }
 
 		/// <summary>
 		/// The change number being built
diff --git a/Engine/Source/Programs/Horde/Horde.Build/Jobs/IJobStepRefCollection.cs b/Engine/Source/Programs/Horde/Horde.Build/Jobs/IJobStepRefCollection.cs
--- a/Engine/Source/Programs/Horde/Horde.Build/Jobs/IJobStepRefCollection.cs
+++ b/Engine/Source/Programs/Horde/Horde.Build/Jobs/IJobStepRefCollection.cs
@@ -16,7 +16,7 @@ namespace Horde.Build.Jobs
 	using LogId = ObjectId<ILogFile>;
 	using PoolId = StringId<IPool>;
 	using StreamId = StringId<IStream>;
-	using TemplateRefId = StringId<TemplateRef>;
+	using TemplateId = StringId<ITemplateRef>;
 
 	/// <summary>
 	/// Interface for a collection of JobStepRef documents
@@ -42,7 +42,7 @@ namespace Horde.Build.Jobs
 		/// <param name="initTime">Time taken for the batch containing this step to initializer</param>
 		/// <param name="startTimeUtc">Start time</param>
 		/// <param name="finishTimeUtc">Finish time for the step, if known</param>
-		Task<IJobStepRef> InsertOrReplaceAsync(JobStepRefId id, string jobName, string stepName, StreamId streamId, TemplateRefId templateId, int change, LogId? logId, PoolId? poolId, AgentId? agentId, JobStepOutcome? outcome, int? lastSuccess, int? lastWarning, float waitTime, float initTime, DateTime startTimeUtc, DateTime? finishTimeUtc);
+		Task<IJobStepRef> InsertOrReplaceAsync(JobStepRefId id, string jobName, string stepName, StreamId streamId, TemplateId templateId, int change, LogId? logId, PoolId? poolId, AgentId? agentId, JobStepOutcome? outcome, int? lastSuccess, int? lastWarning, float waitTime, float initTime, DateTime startTimeUtc, DateTime? finishTimeUtc);
 
 		/// <summary>
 		/// Gets the history of a given node
@@ -54,7 +54,7 @@ namespace Horde.Build.Jobs
 		/// <param name="includeFailed">Whether to include failed nodes</param>
 		/// <param name="count">Number of results to return</param>
 		/// <returns>List of step references</returns>
-		Task<List<IJobStepRef>> GetStepsForNodeAsync(StreamId streamId, TemplateRefId templateId, string nodeName, int? change, bool includeFailed, int count);
+		Task<List<IJobStepRef>> GetStepsForNodeAsync(StreamId streamId, TemplateId templateId, string nodeName, int? change, bool includeFailed, int count);
 
 		/// <summary>
 		/// Gets the previous job that ran a given step
@@ -64,7 +64,7 @@ namespace Horde.Build.Jobs
 		/// <param name="nodeName">Name of the step to find</param>
 		/// <param name="change">The current changelist number</param>
 		/// <returns>The previous job, or null.</returns>
-		Task<IJobStepRef?> GetPrevStepForNodeAsync(StreamId streamId, TemplateRefId templateId, string nodeName, int change);
+		Task<IJobStepRef?> GetPrevStepForNodeAsync(StreamId streamId, TemplateId templateId, string nodeName, int change);
 
 		/// <summary>
 		/// Gets the next job that ran a given step
@@ -74,7 +74,7 @@ namespace Horde.Build.Jobs
 		/// <param name="nodeName">Name of the step to find</param>
 		/// <param name="change">The current changelist number</param>
 		/// <returns>The previous job, or null.</returns>
-		Task<IJobStepRef?> GetNextStepForNodeAsync(StreamId streamId, TemplateRefId templateId, string nodeName, int change);
+		Task<IJobStepRef?> GetNextStepForNodeAsync(StreamId streamId, TemplateId templateId, string nodeName, int change);
 	}
 
 	static class JobStepRefCollectionExtensions
diff --git a/Engine/Source/Programs/Horde/Horde.Build/Jobs/JobCollection.cs b/Engine/Source/Programs/Horde/Horde.Build/Jobs/JobCollection.cs
--- a/Engine/Source/Programs/Horde/Horde.Build/Jobs/JobCollection.cs
+++ b/Engine/Source/Programs/Horde/Horde.Build/Jobs/JobCollection.cs
@@ -36,7 +36,7 @@ namespace Horde.Build.Jobs
 	using PoolId = StringId<IPool>;
 	using SessionId = ObjectId<ISession>;
 	using StreamId = StringId<IStream>;
-	using TemplateRefId = StringId<TemplateRef>;
+	using TemplateId = StringId<ITemplateRef>;
 	using UserId = ObjectId<IUser>;
 
 	/// <summary>
@@ -180,7 +180,7 @@ namespace Horde.Build.Jobs
 		class ChainedJobDocument : IChainedJob
 		{
 			public string Target { get; set; }
-			public TemplateRefId TemplateRefId { get; set; }
+			public TemplateId TemplateRefId { get; set; }
 			public JobId? JobId { get; set; }
 
 			[BsonConstructor]
@@ -189,10 +189,10 @@ namespace Horde.Build.Jobs
 				Target = String.Empty;
 			}
 
-			public ChainedJobDocument(ChainedJobTemplate trigger)
+			public ChainedJobDocument(ChainedJobTemplateConfig trigger)
 			{
 				Target = trigger.Trigger;
-				TemplateRefId = trigger.TemplateRefId;
+				TemplateRefId = trigger.TemplateId;
 			}
 		}
 
@@ -208,7 +208,7 @@ namespace Horde.Build.Jobs
 			public JobId Id { get; set; }
 
 			public StreamId StreamId { get; set; }
-			public TemplateRefId TemplateId { get; set; }
+			public TemplateId TemplateId { get; set; }
 			public ContentHash? TemplateHash { get; set; }
 			public ContentHash GraphHash { get; set; }
 
@@ -297,7 +297,7 @@ namespace Horde.Build.Jobs
 				GraphHash = null!;
 			}
 
-			public JobDocument(JobId id, StreamId streamId, TemplateRefId templateId, ContentHash templateHash, ContentHash graphHash, string name, int change, int codeChange, int preflightChange, int clonedPreflightChange, string? preflightDescription, UserId? startedByUserId, Priority? priority, bool? autoSubmit, bool? updateIssues, bool? promoteIssuesByDefault, DateTime createTimeUtc, List<ChainedJobDocument> chainedJobs, bool showUgsBadges, bool showUgsAlerts, string? notificationChannel, string? notificationChannelFilter, List<string>? arguments)
+			public JobDocument(JobId id, StreamId streamId, TemplateId templateId, ContentHash templateHash, ContentHash graphHash, string name, int change, int codeChange, int preflightChange, int clonedPreflightChange, string? preflightDescription, UserId? startedByUserId, Priority? priority, bool? autoSubmit, bool? updateIssues, bool? promoteIssuesByDefault, DateTime createTimeUtc, List<ChainedJobDocument> chainedJobs, bool showUgsBadges, bool showUgsAlerts, string? notificationChannel, string? notificationChannelFilter, List<string>? arguments)
 			{
 				Id = id;
 				StreamId = streamId;
@@ -396,7 +396,7 @@ namespace Horde.Build.Jobs
 		}
 
 		/// <inheritdoc/>
-		public async Task<IJob> AddAsync(JobId jobId, StreamId streamId, TemplateRefId templateRefId, ContentHash templateHash, IGraph graph, string name, int change, int codeChange, int? preflightChange, int? clonedPreflightChange, string? preflightDescription, UserId? startedByUserId, Priority? priority, bool? autoSubmit, bool? updateIssues, bool? promoteIssuesByDefault, List<ChainedJobTemplate>? chainedJobs, bool showUgsBadges, bool showUgsAlerts, string? notificationChannel, string? notificationChannelFilter, List<string>? arguments)
+		public async Task<IJob> AddAsync(JobId jobId, StreamId streamId, TemplateId templateRefId, ContentHash templateHash, IGraph graph, string name, int change, int codeChange, int? preflightChange, int? clonedPreflightChange, string? preflightDescription, UserId? startedByUserId, Priority? priority, bool? autoSubmit, bool? updateIssues, bool? promoteIssuesByDefault, List<ChainedJobTemplateConfig>? chainedJobs, bool showUgsBadges, bool showUgsAlerts, string? notificationChannel, string? notificationChannelFilter, List<string>? arguments)
 		{
 			List<ChainedJobDocument> jobTriggers = new List<ChainedJobDocument>();
 			if (chainedJobs == null)
@@ -447,7 +447,7 @@ namespace Horde.Build.Jobs
 		}
 
 		/// <inheritdoc/>
-		public async Task<List<IJob>> FindAsync(JobId[]? jobIds, StreamId? streamId, string? name, TemplateRefId[]? templates, int? minChange, int? maxChange, int? preflightChange, bool? preflightOnly, UserId ? preflightStartedByUser, UserId? startedByUser, DateTimeOffset? minCreateTime, DateTimeOffset? maxCreateTime, DateTimeOffset? modifiedBefore, DateTimeOffset? modifiedAfter, int? index, int? count, bool consistentRead, string? indexHint, bool? excludeUserJobs)
+		public async Task<List<IJob>> FindAsync(JobId[]? jobIds, StreamId? streamId, string? name, TemplateId[]? templates, int? minChange, int? maxChange, int? preflightChange, bool? preflightOnly, UserId ? preflightStartedByUser, UserId? startedByUser, DateTimeOffset? minCreateTime, DateTimeOffset? maxCreateTime, DateTimeOffset? modifiedBefore, DateTimeOffset? modifiedAfter, int? index, int? count, bool consistentRead, string? indexHint, bool? excludeUserJobs)
 		{
 			FilterDefinitionBuilder<JobDocument> filterBuilder = Builders<JobDocument>.Filter;
 
@@ -538,7 +538,7 @@ namespace Horde.Build.Jobs
 		}
 
 		/// <inheritdoc/>
-		public async Task<List<IJob>> FindLatestByStreamWithTemplatesAsync(StreamId streamId, TemplateRefId[] templates, UserId? preflightStartedByUser, DateTimeOffset? maxCreateTime, DateTimeOffset? modifiedAfter, int? index, int? count, bool consistentRead)
+		public async Task<List<IJob>> FindLatestByStreamWithTemplatesAsync(StreamId streamId, TemplateId[] templates, UserId? preflightStartedByUser, DateTimeOffset? maxCreateTime, DateTimeOffset? modifiedAfter, int? index, int? count, bool consistentRead)
 		{
 			string indexHint = _createTimeIndex.Name;
 			if (modifiedAfter != null)
@@ -554,7 +554,7 @@ namespace Horde.Build.Jobs
 		}
 
 		/// <inheritdoc/>
-		public async Task<IJob?> TryUpdateJobAsync(IJob inJob, IGraph graph, string? name, Priority? priority, bool? autoSubmit, int? autoSubmitChange, string? autoSubmitMessage, UserId? abortedByUserId, ObjectId? notificationTriggerId, List<Report>? reports, List<string>? arguments, KeyValuePair<int, ObjectId>? labelIdxToTriggerId, KeyValuePair<TemplateRefId, JobId>? jobTrigger)
+		public async Task<IJob?> TryUpdateJobAsync(IJob inJob, IGraph graph, string? name, Priority? priority, bool? autoSubmit, int? autoSubmitChange, string? autoSubmitMessage, UserId? abortedByUserId, ObjectId? notificationTriggerId, List<Report>? reports, List<string>? arguments, KeyValuePair<int, ObjectId>? labelIdxToTriggerId, KeyValuePair<TemplateId, JobId>? jobTrigger)
 		{
 			// Create the update 
 			UpdateDefinitionBuilder<JobDocument> updateBuilder = Builders<JobDocument>.Update;
diff --git a/Engine/Source/Programs/Horde/Horde.Build/Jobs/JobMessages.cs b/Engine/Source/Programs/Horde/Horde.Build/Jobs/JobMessages.cs
--- a/Engine/Source/Programs/Horde/Horde.Build/Jobs/JobMessages.cs
+++ b/Engine/Source/Programs/Horde/Horde.Build/Jobs/JobMessages.cs
@@ -15,7 +15,7 @@ using HordeCommon;
 
 namespace Horde.Build.Jobs
 {
-	using TemplateRefId = StringId<TemplateRef>;
+	using TemplateId = StringId<ITemplateRef>;
 
 	/// <summary>
 	/// State of the job
@@ -53,7 +53,7 @@ namespace Horde.Build.Jobs
 		/// The template for this job
 		/// </summary>
 		[Required]
-		public TemplateRefId TemplateId { get; set; }
+		public TemplateId TemplateId { get; set; }
 
 		/// <summary>
 		/// Name of the job
@@ -98,7 +98,7 @@ namespace Horde.Build.Jobs
 		/// <summary>
 		/// Private constructor for serialization
 		/// </summary>
-		public CreateJobRequest(string streamId, TemplateRefId templateId)
+		public CreateJobRequest(string streamId, TemplateId templateId)
 		{
 			StreamId = streamId;
 			TemplateId = templateId;
diff --git a/Engine/Source/Programs/Horde/Horde.Build/Jobs/JobService.cs b/Engine/Source/Programs/Horde/Horde.Build/Jobs/JobService.cs
--- a/Engine/Source/Programs/Horde/Horde.Build/Jobs/JobService.cs
+++ b/Engine/Source/Programs/Horde/Horde.Build/Jobs/JobService.cs
@@ -36,7 +36,7 @@ namespace Horde.Build.Jobs
 	using LogId = ObjectId<ILogFile>;
 	using SessionId = ObjectId<ISession>;
 	using StreamId = StringId<IStream>;
-	using TemplateRefId = StringId<TemplateRef>;
+	using TemplateId = StringId<ITemplateRef>;
 	using UserId = ObjectId<IUser>;
 
 	/// <summary>
@@ -188,7 +188,7 @@ namespace Horde.Build.Jobs
 		/// <param name="notificationChannelFilter">Notification Channel filter for this job</param>
 		/// <param name="arguments">Arguments for the job</param>
 		/// <returns>Unique id representing the job</returns>
-		public async Task<IJob> CreateJobAsync(JobId? jobId, IStream stream, TemplateRefId templateRefId, ContentHash templateHash, IGraph graph, string name, int change, int codeChange, int? preflightChange, int? clonedPreflightChange, string? preflightDescription, UserId? startedByUserId, Priority? priority, bool? autoSubmit, bool? updateIssues, bool? promoteIssuesByDefault, List<ChainedJobTemplate>? jobTriggers, bool showUgsBadges, bool showUgsAlerts, string? notificationChannel, string? notificationChannelFilter, IReadOnlyList<string> arguments)
+		public async Task<IJob> CreateJobAsync(JobId? jobId, IStream stream, TemplateId templateRefId, ContentHash templateHash, IGraph graph, string name, int change, int codeChange, int? preflightChange, int? clonedPreflightChange, string? preflightDescription, UserId? startedByUserId, Priority? priority, bool? autoSubmit, bool? updateIssues, bool? promoteIssuesByDefault, List<ChainedJobTemplateConfig>? jobTriggers, bool showUgsBadges, bool showUgsAlerts, string? notificationChannel, string? notificationChannelFilter, IReadOnlyList<string> arguments)
 		{
 			using IScope traceScope = GlobalTracer.Instance.BuildSpan("JobService.CreateJobAsync").StartActive();
 			traceScope.Span.SetTag("JobId", jobId);
@@ -518,7 +518,7 @@ namespace Horde.Build.Jobs
 		/// <param name="consistentRead">If the database read should be made to the replica server</param>
 		/// <param name="excludeUserJobs">Whether to exclude user jobs from the find</param>
 		/// <returns>List of jobs matching the given criteria</returns>
-		public async Task<List<IJob>> FindJobsAsync(JobId[]? jobIds = null, StreamId? streamId = null, string? name = null, TemplateRefId[]? templates = null, int? minChange = null, int? maxChange = null, int? preflightChange = null, bool? preflightOnly = null, UserId? preflightStartedByUser = null, UserId? startedByUser = null, DateTimeOffset ? minCreateTime = null, DateTimeOffset? maxCreateTime = null, string? target = null, JobStepState[]? state = null, JobStepOutcome[]? outcome = null, DateTimeOffset? modifiedBefore = null, DateTimeOffset? modifiedAfter = null, int? index = null, int? count = null, bool consistentRead = true, bool? excludeUserJobs = null)
+		public async Task<List<IJob>> FindJobsAsync(JobId[]? jobIds = null, StreamId? streamId = null, string? name = null, TemplateId[]? templates = null, int? minChange = null, int? maxChange = null, int? preflightChange = null, bool? preflightOnly = null, UserId? preflightStartedByUser = null, UserId? startedByUser = null, DateTimeOffset ? minCreateTime = null, DateTimeOffset? maxCreateTime = null, string? target = null, JobStepState[]? state = null, JobStepOutcome[]? outcome = null, DateTimeOffset? modifiedBefore = null, DateTimeOffset? modifiedAfter = null, int? index = null, int? count = null, bool consistentRead = true, bool? excludeUserJobs = null)
 		{
 			using IScope scope = GlobalTracer.Instance.BuildSpan("JobService.FindJobsAsync").StartActive();
 			scope.Span.SetTag("JobIds", jobIds);
@@ -603,7 +603,7 @@ namespace Horde.Build.Jobs
 		/// <param name="count">Number of results to return</param>
 		/// <param name="consistentRead">If the database read should be made to the replica server</param>
 		/// <returns>List of jobs matching the given criteria</returns>
-		public async Task<List<IJob>> FindJobsByStreamWithTemplatesAsync(StreamId streamId, TemplateRefId[] templates, UserId? preflightStartedByUser = null, DateTimeOffset? maxCreateTime = null, DateTimeOffset? modifiedAfter = null, int? index = null, int? count = null, bool consistentRead = true)
+		public async Task<List<IJob>> FindJobsByStreamWithTemplatesAsync(StreamId streamId, TemplateId[] templates, UserId? preflightStartedByUser = null, DateTimeOffset? maxCreateTime = null, DateTimeOffset? modifiedAfter = null, int? index = null, int? count = null, bool consistentRead = true)
 		{
 			using IScope scope = GlobalTracer.Instance.BuildSpan("JobService.FindJobsByStreamWithTemplatesAsync").StartActive();
 			scope.Span.SetTag("StreamId", streamId);
@@ -720,7 +720,7 @@ namespace Horde.Build.Jobs
 		/// <param name="nodeName">Name of the node</param>
 		/// <param name="change">Maximum changelist to consider</param>
 		/// <returns>Expected duration for the given step</returns>
-		async Task<JobStepTimingData> GetStepTimingInfo(StreamId streamId, TemplateRefId templateId, string nodeName, int? change)
+		async Task<JobStepTimingData> GetStepTimingInfo(StreamId streamId, TemplateId templateId, string nodeName, int? change)
 		{
 			using IScope traceScope = GlobalTracer.Instance.BuildSpan("JobService.GetStepTimingInfo").StartActive();
 			traceScope.Span.SetTag("StreamId", streamId);
@@ -1216,7 +1216,7 @@ namespace Horde.Build.Jobs
 				// Update the job
 				JobId chainedJobId = JobId.GenerateNewId();
 
-				IJob? newJob = await _jobs.TryUpdateJobAsync(job, graph, jobTrigger: new KeyValuePair<TemplateRefId, JobId>(jobTrigger.TemplateRefId, chainedJobId));
+				IJob? newJob = await _jobs.TryUpdateJobAsync(job, graph, jobTrigger: new KeyValuePair<TemplateId, JobId>(jobTrigger.TemplateRefId, chainedJobId));
 				if(newJob != null)
 				{
 					IStream? stream = await _streamService.GetStreamAsync(newJob.StreamId);
@@ -1226,7 +1226,7 @@ namespace Horde.Build.Jobs
 						break;
 					}
 
-					TemplateRef? templateRef;
+					ITemplateRef? templateRef;
 					if (!stream.Templates.TryGetValue(jobTrigger.TemplateRefId, out templateRef))
 					{
 						_logger.LogWarning("Cannot find template {TemplateId} in stream {StreamId}", jobTrigger.TemplateRefId, newJob.StreamId);
@@ -1243,7 +1243,7 @@ namespace Horde.Build.Jobs
 					IGraph triggerGraph = await _graphs.AddAsync(template);
 					_logger.LogInformation("Creating downstream job {ChainedJobId} from job {JobId}", chainedJobId, newJob.Id);
 
-					await CreateJobAsync(chainedJobId, stream, jobTrigger.TemplateRefId, templateRef.Hash, triggerGraph, templateRef.Name, newJob.Change, newJob.CodeChange, newJob.PreflightChange, newJob.ClonedPreflightChange, newJob.PreflightDescription, newJob.StartedByUserId, template.Priority, null, newJob.UpdateIssues, newJob.PromoteIssuesByDefault, templateRef.ChainedJobs, false, false, templateRef.NotificationChannel, templateRef.NotificationChannelFilter, template.Arguments);
+					await CreateJobAsync(chainedJobId, stream, jobTrigger.TemplateRefId, templateRef.Hash, triggerGraph, templateRef.Config.Name, newJob.Change, newJob.CodeChange, newJob.PreflightChange, newJob.ClonedPreflightChange, newJob.PreflightDescription, newJob.StartedByUserId, template.Priority, null, newJob.UpdateIssues, newJob.PromoteIssuesByDefault, templateRef.Config.ChainedJobs, false, false, templateRef.Config.NotificationChannel, templateRef.Config.NotificationChannelFilter, template.Arguments);
 					return newJob;
 				}
 
diff --git a/Engine/Source/Programs/Horde/Horde.Build/Jobs/JobStepRefCollection.cs b/Engine/Source/Programs/Horde/Horde.Build/Jobs/JobStepRefCollection.cs
--- a/Engine/Source/Programs/Horde/Horde.Build/Jobs/JobStepRefCollection.cs
+++ b/Engine/Source/Programs/Horde/Horde.Build/Jobs/JobStepRefCollection.cs
@@ -19,7 +19,7 @@ namespace Horde.Build.Jobs
 	using LogId = ObjectId<ILogFile>;
 	using PoolId = StringId<IPool>;
 	using StreamId = StringId<IStream>;
-	using TemplateRefId = StringId<TemplateRef>;
+	using TemplateId = StringId<ITemplateRef>;
 
 	/// <summary>
 	/// Collection of JobStepRef documents
@@ -36,7 +36,7 @@ namespace Horde.Build.Jobs
 			public string Name { get; set; }
 
 			public StreamId StreamId { get; set; }
-			public TemplateRefId TemplateId { get; set; }
+			public TemplateId TemplateId { get; set; }
 			public int Change { get; set; }
 			public LogId? LogId { get; set; }
 			public PoolId? PoolId { get; set; }
@@ -68,7 +68,7 @@ namespace Horde.Build.Jobs
 			DateTime? IJobStepRef.FinishTimeUtc => FinishTimeUtc ?? FinishTime?.UtcDateTime;
 			string IJobStepRef.NodeName => Name;
 
-			public JobStepRef(JobStepRefId id, string jobName, string nodeName, StreamId streamId, TemplateRefId templateId, int change, LogId? logId, PoolId? poolId, AgentId? agentId, JobStepOutcome? outcome, int? lastSuccess, int? lastWarning, float batchWaitTime, float batchInitTime, DateTime startTimeUtc, DateTime? finishTimeUtc)
+			public JobStepRef(JobStepRefId id, string jobName, string nodeName, StreamId streamId, TemplateId templateId, int change, LogId? logId, PoolId? poolId, AgentId? agentId, JobStepOutcome? outcome, int? lastSuccess, int? lastWarning, float batchWaitTime, float batchInitTime, DateTime startTimeUtc, DateTime? finishTimeUtc)
 			{
 				Id = id;
 				JobName = jobName;
@@ -103,7 +103,7 @@ namespace Horde.Build.Jobs
 		}
 
 		/// <inheritdoc/>
-		public async Task<IJobStepRef> InsertOrReplaceAsync(JobStepRefId id, string jobName, string stepName, StreamId streamId, TemplateRefId templateId, int change, LogId? logId, PoolId? poolId, AgentId? agentId, JobStepOutcome? outcome, int? lastSuccess, int? lastWarning, float waitTime, float initTime, DateTime startTimeUtc, DateTime? finishTimeUtc)
+		public async Task<IJobStepRef> InsertOrReplaceAsync(JobStepRefId id, string jobName, string stepName, StreamId streamId, TemplateId templateId, int change, LogId? logId, PoolId? poolId, AgentId? agentId, JobStepOutcome? outcome, int? lastSuccess, int? lastWarning, float waitTime, float initTime, DateTime startTimeUtc, DateTime? finishTimeUtc)
 		{
 			JobStepRef newJobStepRef = new JobStepRef(id, jobName, stepName, streamId, templateId, change, logId, poolId, agentId, outcome, lastSuccess, lastWarning, waitTime, initTime, startTimeUtc, finishTimeUtc);
 			await _jobStepRefs.ReplaceOneAsync(Builders<JobStepRef>.Filter.Eq(x => x.Id, newJobStepRef.Id), newJobStepRef, new ReplaceOptions { IsUpsert = true });
@@ -111,7 +111,7 @@ namespace Horde.Build.Jobs
 		}
 
 		/// <inheritdoc/>
-		public async Task<List<IJobStepRef>> GetStepsForNodeAsync(StreamId streamId, TemplateRefId templateId, string nodeName, int? change, bool includeFailed, int count)
+		public async Task<List<IJobStepRef>> GetStepsForNodeAsync(StreamId streamId, TemplateId templateId, string nodeName, int? change, bool includeFailed, int count)
 		{
 			// Find all the steps matching the given criteria
 			FilterDefinitionBuilder<JobStepRef> filterBuilder = Builders<JobStepRef>.Filter;
@@ -134,13 +134,13 @@ namespace Horde.Build.Jobs
 		}
 
 		/// <inheritdoc/>
-		public async Task<IJobStepRef?> GetPrevStepForNodeAsync(StreamId streamId, TemplateRefId templateId, string nodeName, int change)
+		public async Task<IJobStepRef?> GetPrevStepForNodeAsync(StreamId streamId, TemplateId templateId, string nodeName, int change)
 		{
 			return await _jobStepRefs.Find(x => x.StreamId == streamId && x.TemplateId == templateId && x.Name == nodeName && x.Change < change && x.Outcome != null).SortByDescending(x => x.Change).FirstOrDefaultAsync();
 		}
 
 		/// <inheritdoc/>
-		public async Task<IJobStepRef?> GetNextStepForNodeAsync(StreamId streamId, TemplateRefId templateId, string nodeName, int change)
+		public async Task<IJobStepRef?> GetNextStepForNodeAsync(StreamId streamId, TemplateId templateId, string nodeName, int change)
 		{
 			return await _jobStepRefs.Find(x => x.StreamId == streamId && x.TemplateId == templateId && x.Name == nodeName && x.Change > change && x.Outcome != null).SortBy(x => x.Change).FirstOrDefaultAsync();
 		}
diff --git a/Engine/Source/Programs/Horde/Horde.Build/Jobs/JobTaskSource.cs b/Engine/Source/Programs/Horde/Horde.Build/Jobs/JobTaskSource.cs
--- a/Engine/Source/Programs/Horde/Horde.Build/Jobs/JobTaskSource.cs
+++ b/Engine/Source/Programs/Horde/Horde.Build/Jobs/JobTaskSource.cs
@@ -438,31 +438,28 @@ namespace Horde.Build.Jobs
 					}
 					else
 					{
-						TemplateRef? templateRef;
+						ITemplateRef? templateRef;
 						if (stream.Templates.TryGetValue(newJob.TemplateId, out templateRef))
 						{							
 							if (templateRef.StepStates != null)
 							{
 								for (int i = 0; i < templateRef.StepStates.Count; i++)
 								{
-									TemplateStepState state = templateRef.StepStates[i];
-									if (state.PausedByUserId != null)
-									{
-										IJobStep? step = batch.Steps.FirstOrDefault(x => graph.Groups[batch.GroupIdx].Nodes[x.NodeIdx].Name.Equals(state.Name, StringComparison.Ordinal));
+									ITemplateStep state = templateRef.StepStates[i];
 
-										if (step != null)
+									IJobStep? step = batch.Steps.FirstOrDefault(x => graph.Groups[batch.GroupIdx].Nodes[x.NodeIdx].Name.Equals(state.Name, StringComparison.Ordinal));
+									if (step != null)
+									{
+										JobId jobId = newJob.Id;
+										newJob = await _jobs.TryUpdateStepAsync(newJob, graph, batch.Id, step.Id, JobStepState.Skipped, newError: JobStepError.Paused);
+										if (newJob == null)
+										{
+											_logger.LogError("Job {JobId} failed to update step {StepName} pause state", jobId, state.Name);
+											break;
+										}
+										else
 										{
-											JobId jobId = newJob.Id;
-											newJob = await _jobs.TryUpdateStepAsync(newJob, graph, batch.Id, step.Id, JobStepState.Skipped, newError: JobStepError.Paused);
-											if (newJob == null)
-											{
-												_logger.LogError("Job {JobId} failed to update step {StepName} pause state", jobId, state.Name);
-												break;
-											}
-											else
-											{
-												_logger.LogInformation("Job {JobId} step {StepName} has been skipped due to being paused", jobId, state.Name);
-											}
+											_logger.LogInformation("Job {JobId} step {StepName} has been skipped due to being paused", jobId, state.Name);
 										}
 									}
 								}
diff --git a/Engine/Source/Programs/Horde/Horde.Build/Jobs/JobsController.cs b/Engine/Source/Programs/Horde/Horde.Build/Jobs/JobsController.cs
--- a/Engine/Source/Programs/Horde/Horde.Build/Jobs/JobsController.cs
+++ b/Engine/Source/Programs/Horde/Horde.Build/Jobs/JobsController.cs
@@ -29,7 +29,7 @@ namespace Horde.Build.Jobs
 {
 	using JobId = ObjectId<IJob>;
 	using StreamId = StringId<IStream>;
-	using TemplateRefId = StringId<TemplateRef>;
+	using TemplateId = StringId<ITemplateRef>;
 	using UserId = ObjectId<IUser>;
 
 	/// <summary>
@@ -88,10 +88,10 @@ namespace Horde.Build.Jobs
 			}
 
 			// Get the name of the template ref
-			TemplateRefId templateRefId = create.TemplateId;
+			TemplateId templateRefId = create.TemplateId;
 
 			// Augment the request with template properties
-			TemplateRef? templateRef;
+			ITemplateRef? templateRef;
 			if (!stream.Templates.TryGetValue(templateRefId, out templateRef))
 			{
 				return BadRequest($"Template {create.TemplateId} is not available for stream {stream.Id}");
@@ -113,7 +113,7 @@ namespace Horde.Build.Jobs
 
 			// Get the name of the new job
 			string name = create.Name ?? template.Name;
-			if (create.TemplateId == new TemplateRefId("stage-to-marketplace") && create.Arguments != null)
+			if (create.TemplateId == new TemplateId("stage-to-marketplace") && create.Arguments != null)
 			{
 				foreach (string argument in create.Arguments)
 				{
@@ -190,7 +190,7 @@ namespace Horde.Build.Jobs
 			}
 
 			// Create the job
-			IJob job = await _jobService.CreateJobAsync(null, stream, templateRefId, template.Id, graph, name, change, codeChange, create.PreflightChange, null, preflightDescription, User.GetUserId(), priority, create.AutoSubmit, updateIssues, false, templateRef.ChainedJobs, templateRef.ShowUgsBadges, templateRef.ShowUgsAlerts, templateRef.NotificationChannel, templateRef.NotificationChannelFilter, arguments);
+			IJob job = await _jobService.CreateJobAsync(null, stream, templateRefId, template.Id, graph, name, change, codeChange, create.PreflightChange, null, preflightDescription, User.GetUserId(), priority, create.AutoSubmit, updateIssues, false, templateRef.Config.ChainedJobs, templateRef.Config.ShowUgsBadges, templateRef.Config.ShowUgsAlerts, templateRef.Config.NotificationChannel, templateRef.Config.NotificationChannelFilter, arguments);
 			await UpdateNotificationsAsync(job.Id, new UpdateNotificationsRequest { Slack = true });
 			return new CreateJobResponse(job.Id.ToString());
 		}
@@ -203,7 +203,7 @@ namespace Horde.Build.Jobs
 		/// <param name="target"></param>
 		/// <param name="outcomes"></param>
 		/// <returns></returns>
-		async Task<int> ExecuteChangeQueryAsync(IStream stream, TemplateRefId templateId, string? target, List<JobStepOutcome> outcomes)
+		async Task<int> ExecuteChangeQueryAsync(IStream stream, TemplateId templateId, string? target, List<JobStepOutcome> outcomes)
 		{
 			IList<IJob> jobs = await _jobService.FindJobsAsync(streamId: stream.Id, templates: new[] { templateId }, target: target, state: new[] { JobStepState.Completed }, outcome: outcomes.ToArray(), count: 1, excludeUserJobs: true);
 			if (jobs.Count == 0)
@@ -606,10 +606,10 @@ namespace Horde.Build.Jobs
 				return BadRequest("Missing/invalid query parameter streamId");
 			}
 
-			TemplateRefId[] templateRefIds = templates switch
+			TemplateId[] templateRefIds = templates switch
 			{
-				{ Length: > 0 } => templates.Select(x => new TemplateRefId(x)).ToArray(),
-				_ => Array.Empty<TemplateRefId>()
+				{ Length: > 0 } => templates.Select(x => new TemplateId(x)).ToArray(),
+				_ => Array.Empty<TemplateId>()
 			};
 
 			List<IJob> jobs = await _jobService.FindJobsByStreamWithTemplatesAsync(new StreamId(streamId), templateRefIds, count: count, consistentRead: false);
@@ -710,7 +710,7 @@ namespace Horde.Build.Jobs
 			JobId[]? jobIdValues = (ids == null) ? (JobId[]?)null : Array.ConvertAll(ids, x => new JobId(x));
 			StreamId? streamIdValue = (streamId == null)? (StreamId?)null : new StreamId(streamId);
 			
-			TemplateRefId[]? templateRefIds = (templates != null && templates.Length > 0) ? templates.Select(x => new TemplateRefId(x)).ToArray() : null;
+			TemplateId[]? templateRefIds = (templates != null && templates.Length > 0) ? templates.Select(x => new TemplateId(x)).ToArray() : null;
 
 			if (includePreflight == false)
 			{
@@ -816,7 +816,7 @@ namespace Horde.Build.Jobs
 			[FromQuery] bool consistentRead = false)
 		{
 			StreamId streamIdValue = new StreamId(streamId);
-			TemplateRefId[] templateRefIds = templates.Select(x => new TemplateRefId(x)).ToArray();
+			TemplateId[] templateRefIds = templates.Select(x => new TemplateId(x)).ToArray();
 			UserId? preflightStartedByUserIdValue = preflightStartedByUserId != null ? new UserId(preflightStartedByUserId) : null;
 			count = Math.Min(1000, count);
 
diff --git a/Engine/Source/Programs/Horde/Horde.Build/Jobs/Schedules/ISchedule.cs b/Engine/Source/Programs/Horde/Horde.Build/Jobs/Schedules/ISchedule.cs
deleted file mode 100644
diff --git a/Engine/Source/Programs/Horde/Horde.Build/Jobs/Schedules/ScheduleService.cs b/Engine/Source/Programs/Horde/Horde.Build/Jobs/Schedules/ScheduleService.cs
--- a/Engine/Source/Programs/Horde/Horde.Build/Jobs/Schedules/ScheduleService.cs
+++ b/Engine/Source/Programs/Horde/Horde.Build/Jobs/Schedules/ScheduleService.cs
@@ -31,7 +31,7 @@ namespace Horde.Build.Jobs.Schedules
 {
 	using JobId = ObjectId<IJob>;
 	using StreamId = StringId<IStream>;
-	using TemplateRefId = StringId<TemplateRef>;
+	using TemplateId = StringId<ITemplateRef>;
 
 	/// <summary>
 	/// Manipulates schedule instances
@@ -45,13 +45,13 @@ namespace Horde.Build.Jobs.Schedules
 			public StreamId StreamId { get; set; }
 
 			[CbField("tid")]
-			public TemplateRefId TemplateId { get; set; }
+			public TemplateId TemplateId { get; set; }
 
 			public QueueItem()
 			{
 			}
 
-			public QueueItem(StreamId streamId, TemplateRefId templateId)
+			public QueueItem(StreamId streamId, TemplateId templateId)
 			{
 				StreamId = streamId;
 				TemplateId = templateId;
@@ -285,7 +285,7 @@ namespace Horde.Build.Jobs.Schedules
 			List<IStream> streams = await _streamService.GetStreamsAsync();
 			foreach (IStream stream in streams)
 			{
-				foreach ((TemplateRefId templateId, TemplateRef templateRef) in stream.Templates)
+				foreach ((TemplateId templateId, ITemplateRef templateRef) in stream.Templates)
 				{
 					if (templateRef.Schedule != null)
 					{
@@ -315,15 +315,15 @@ namespace Horde.Build.Jobs.Schedules
 		/// <param name="utcNow"></param>
 		/// <param name="cancellationToken"></param>
 		/// <returns>Async task</returns>
-		internal async Task<bool> TriggerAsync(StreamId streamId, TemplateRefId templateId, DateTime utcNow, CancellationToken cancellationToken)
+		internal async Task<bool> TriggerAsync(StreamId streamId, TemplateId templateId, DateTime utcNow, CancellationToken cancellationToken)
 		{
 			IStream? stream = await _streamService.GetStreamAsync(streamId);
-			if (stream == null || !stream.Templates.TryGetValue(templateId, out TemplateRef? templateRef))
+			if (stream == null || !stream.Templates.TryGetValue(templateId, out ITemplateRef? templateRef))
 			{
 				return false;
 			}
 
-			Schedule? schedule = templateRef.Schedule;
+			ITemplateSchedule? schedule = templateRef.Schedule;
 			if (schedule == null)
 			{
 				return false;
@@ -383,12 +383,12 @@ namespace Horde.Build.Jobs.Schedules
 		/// <param name="utcNow">The current time</param>
 		/// <param name="cancellationToken"></param>
 		/// <returns>Async task</returns>
-		private async Task TriggerAsync(IStream stream, TemplateRefId templateId, TemplateRef templateRef, Schedule schedule, int numActiveJobs, DateTime utcNow, CancellationToken cancellationToken)
+		private async Task TriggerAsync(IStream stream, TemplateId templateId, ITemplateRef templateRef, ITemplateSchedule schedule, int numActiveJobs, DateTime utcNow, CancellationToken cancellationToken)
 		{
 			// Check we're not already at the maximum number of allowed jobs
-			if (schedule.MaxActive != 0 && numActiveJobs >= schedule.MaxActive)
+			if (schedule.Config.MaxActive != 0 && numActiveJobs >= schedule.Config.MaxActive)
 			{
-				_logger.LogInformation("Skipping trigger of {StreamId} template {TemplateId} - already have maximum number of jobs running ({NumJobs})", stream.Id, templateId, schedule.MaxActive);
+				_logger.LogInformation("Skipping trigger of {StreamId} template {TemplateId} - already have maximum number of jobs running ({NumJobs})", stream.Id, templateId, schedule.Config.MaxActive);
 				foreach (JobId jobId in schedule.ActiveJobs)
 				{
 					_logger.LogInformation("Active job for {StreamId} template {TemplateId}: {JobId}", stream.Id, templateId, jobId);
@@ -398,7 +398,7 @@ namespace Horde.Build.Jobs.Schedules
 
 			// Minimum changelist number, inclusive
 			int minChangeNumber = schedule.LastTriggerChange;
-			if (minChangeNumber > 0 && !schedule.RequireSubmittedChange)
+			if (minChangeNumber > 0 && !schedule.Config.RequireSubmittedChange)
 			{
 				minChangeNumber--;
 			}
@@ -408,24 +408,24 @@ namespace Horde.Build.Jobs.Schedules
 
 			// Get the maximum number of changes to trigger
 			int maxNewChanges = 1;
-			if (schedule.MaxChanges != 0)
+			if (schedule.Config.MaxChanges != 0)
 			{
-				maxNewChanges = schedule.MaxChanges;
+				maxNewChanges = schedule.Config.MaxChanges;
 			}
-			if (schedule.MaxActive != 0)
+			if (schedule.Config.MaxActive != 0)
 			{
-				maxNewChanges = Math.Min(maxNewChanges, schedule.MaxActive - numActiveJobs);
+				maxNewChanges = Math.Min(maxNewChanges, schedule.Config.MaxActive - numActiveJobs);
 			}
 
 			// Create a timer to limit the amount we look back through P4 history
 			Stopwatch timer = Stopwatch.StartNew();
-			ChangeContentFlags? filterFlags = schedule.GetFilterFlags();
+			ChangeContentFlags? filterFlags = schedule.Config.GetFilterFlags();
 
 			// Create a file filter
 			FileFilter? fileFilter = null;
-			if (schedule.Files != null)
+			if (schedule.Config.Files != null)
 			{
-				fileFilter = new FileFilter(schedule.Files);
+				fileFilter = new FileFilter(schedule.Config.Files);
 			}
 
 			// Cache the Perforce history as we're iterating through changes to improve query performance
@@ -439,9 +439,9 @@ namespace Horde.Build.Jobs.Schedules
 
 				// Get the next valid change
 				ICommit? changeDetails = null;
-				if (schedule.Gate != null)
+				if (schedule.Config.Gate != null)
 				{
-					changeDetails = await GetNextChangeForGateAsync(stream, templateId, schedule.Gate, minChangeNumber, maxChangeNumber, cancellationToken);
+					changeDetails = await GetNextChangeForGateAsync(stream, templateId, schedule.Config.Gate, minChangeNumber, maxChangeNumber, cancellationToken);
 				}
 				else
 				{
@@ -457,7 +457,7 @@ namespace Horde.Build.Jobs.Schedules
 				{
 					break;
 				}
-				if (changeDetails.Number == minChangeNumber && (schedule.RequireSubmittedChange || triggerChanges.Count > 0))
+				if (changeDetails.Number == minChangeNumber && (schedule.Config.RequireSubmittedChange || triggerChanges.Count > 0))
 				{
 					break;
 				}
@@ -522,7 +522,7 @@ namespace Horde.Build.Jobs.Schedules
 			{
 				cancellationToken.ThrowIfCancellationRequested();
 				List<string> defaultArguments = template.GetDefaultArguments();
-				IJob newJob = await _jobService.CreateJobAsync(null, stream, templateId, template.Id, graph, template.Name, change, codeChange, null, null, null, null, template.Priority, null, null, template.PromoteIssuesByDefault, templateRef.ChainedJobs, templateRef.ShowUgsBadges, templateRef.ShowUgsAlerts, templateRef.NotificationChannel, templateRef.NotificationChannelFilter, defaultArguments);
+				IJob newJob = await _jobService.CreateJobAsync(null, stream, templateId, template.Id, graph, template.Name, change, codeChange, null, null, null, null, template.Priority, null, null, template.PromoteIssuesByDefault, templateRef.Config.ChainedJobs, templateRef.Config.ShowUgsBadges, templateRef.Config.ShowUgsAlerts, templateRef.Config.NotificationChannel, templateRef.Config.NotificationChannelFilter, defaultArguments);
 				_logger.LogInformation("Started new job for {StreamId} template {TemplateId} at CL {Change} (Code CL {CodeChange}): {JobId}", stream.Id, templateId, change, codeChange, newJob.Id);
 				await _streamService.UpdateScheduleTriggerAsync(stream, templateId, utcNow, change, new List<JobId> { newJob.Id }, new List<JobId>());
 			}
@@ -566,13 +566,13 @@ namespace Horde.Build.Jobs.Schedules
 		/// Gets the next change to build for a schedule on a gate
 		/// </summary>
 		/// <returns></returns>
-		private async Task<ICommit?> GetNextChangeForGateAsync(IStream stream, TemplateRefId templateRefId, ScheduleGate gate, int? minChange, int? maxChange, CancellationToken cancellationToken)
+		private async Task<ICommit?> GetNextChangeForGateAsync(IStream stream, TemplateId templateRefId, ScheduleGateConfig gate, int? minChange, int? maxChange, CancellationToken cancellationToken)
 		{
 			for (; ; )
 			{
 				cancellationToken.ThrowIfCancellationRequested();
 
-				List<IJob> jobs = await _jobCollection.FindAsync(streamId: stream.Id, templates: new[] { gate.TemplateRefId }, minChange: minChange, maxChange: maxChange, count: 1);
+				List<IJob> jobs = await _jobCollection.FindAsync(streamId: stream.Id, templates: new[] { gate.TemplateId }, minChange: minChange, maxChange: maxChange, count: 1);
 				if (jobs.Count == 0)
 				{
 					return null;
@@ -591,7 +591,7 @@ namespace Horde.Build.Jobs.Schedules
 						{
 							return await _perforce.GetChangeDetailsAsync(stream, job.Change, cancellationToken);
 						}
-						_logger.LogInformation("Skipping trigger of {StreamName} template {TemplateId} - last {OtherTemplateRefId} job ({JobId}) ended with errors", stream.Id, templateRefId, gate.TemplateRefId, job.Id);
+						_logger.LogInformation("Skipping trigger of {StreamName} template {TemplateId} - last {OtherTemplateRefId} job ({JobId}) ended with errors", stream.Id, templateRefId, gate.TemplateId, job.Id);
 					}
 				}
 
diff --git a/Engine/Source/Programs/Horde/Horde.Build/Jobs/Schedules/Schedules.cs b/Engine/Source/Programs/Horde/Horde.Build/Jobs/Schedules/Schedules.cs
--- a/Engine/Source/Programs/Horde/Horde.Build/Jobs/Schedules/Schedules.cs
+++ b/Engine/Source/Programs/Horde/Horde.Build/Jobs/Schedules/Schedules.cs
@@ -2,79 +2,13 @@
 
 using System;
 using System.Collections.Generic;
-using System.ComponentModel.DataAnnotations;
-using Horde.Build.Acls;
 using Horde.Build.Perforce;
 using Horde.Build.Streams;
 using Horde.Build.Utilities;
 
 namespace Horde.Build.Jobs.Schedules
 {
-	using TemplateRefId = StringId<TemplateRef>;
-
-	/// <summary>
-	/// Information about a schedule pattern
-	/// </summary>
-	public class GetSchedulePatternResponse
-	{
-		/// <summary>
-		/// Days of the week to run this schedule on. If null, the schedule will run every day.
-		/// </summary>
-		public List<string>? DaysOfWeek { get; set; }
-
-		/// <summary>
-		/// Time during the day for the first schedule to trigger. Measured in minutes from midnight.
-		/// </summary>
-		public int MinTime { get; set; }
-
-		/// <summary>
-		/// Time during the day for the last schedule to trigger. Measured in minutes from midnight.
-		/// </summary>
-		public int? MaxTime { get; set; }
-
-		/// <summary>
-		/// Interval between each schedule triggering
-		/// </summary>
-		public int? Interval { get; set; }
-
-		/// <summary>
-		/// Constructor
-		/// </summary>
-		/// <param name="pattern">The pattern to construct from</param>
-		public GetSchedulePatternResponse(SchedulePattern pattern)
-		{
-			DaysOfWeek = pattern.DaysOfWeek?.ConvertAll(x => x.ToString());
-			MinTime = pattern.MinTime;
-			MaxTime = pattern.MaxTime;
-			Interval = pattern.Interval;
-		}
-	}
-
-	/// <summary>
-	/// Gate allowing a schedule to trigger.
-	/// </summary>
-	public class GetScheduleGateResponse
-	{
-		/// <summary>
-		/// The template containing the dependency
-		/// </summary>
-		public string TemplateId { get; set; }
-
-		/// <summary>
-		/// Target to wait for
-		/// </summary>
-		public string Target { get; set; }
-
-		/// <summary>
-		/// Constructor
-		/// </summary>
-		/// <param name="gate">Gate to construct from</param>
-		public GetScheduleGateResponse(ScheduleGate gate)
-		{
-			TemplateId = gate.TemplateRefId.ToString();
-			Target = gate.Target;
-		}
-	}
+	using JobId = ObjectId<IJob>;
 
 	/// <summary>
 	/// Response describing a schedule
@@ -104,7 +38,7 @@ namespace Horde.Build.Jobs.Schedules
 		/// <summary>
 		/// Gate for this schedule to trigger
 		/// </summary>
-		public GetScheduleGateResponse? Gate { get; set; }
+		public ScheduleGateConfig? Gate { get; set; }
 
 		/// <summary>
 		/// The types of changes to run for
@@ -119,7 +53,7 @@ namespace Horde.Build.Jobs.Schedules
 		/// <summary>
 		/// New patterns for the schedule
 		/// </summary>
-		public List<GetSchedulePatternResponse> Patterns { get; set; }
+		public List<SchedulePatternConfig> Patterns { get; set; }
 
 		/// <summary>
 		/// Last changelist number that this was triggered for
@@ -134,28 +68,25 @@ namespace Horde.Build.Jobs.Schedules
 		/// <summary>
 		/// List of active jobs
 		/// </summary>
-		public List<string> ActiveJobs { get; set; }
+		public List<JobId> ActiveJobs { get; set; }
 
 		/// <summary>
 		/// Constructor
 		/// </summary>
-		/// <param name="schedule">The schedule to construct from</param>
-		public GetScheduleResponse(Schedule schedule)
+		/// <param name="schedule">Schedule to construct from</param>
+		public GetScheduleResponse(ITemplateSchedule schedule)
 		{
-			Enabled = schedule.Enabled;
-			MaxActive = schedule.MaxActive;
-			MaxChanges = schedule.MaxChanges;
-			RequireSubmittedChange = schedule.RequireSubmittedChange;
-			if (schedule.Gate != null)
-			{
-				Gate = new GetScheduleGateResponse(schedule.Gate);
-			}
-			Filter = schedule.Filter;
-			TemplateParameters = schedule.TemplateParameters;
-			Patterns = schedule.Patterns.ConvertAll(x => new GetSchedulePatternResponse(x));
+			Enabled = schedule.Config.Enabled;
+			MaxActive = schedule.Config.MaxActive;
+			MaxChanges = schedule.Config.MaxChanges;
+			RequireSubmittedChange = schedule.Config.RequireSubmittedChange;
+			Gate = schedule.Config.Gate;
+			Filter = schedule.Config.Filter;
+			TemplateParameters = schedule.Config.TemplateParameters;
+			Patterns = schedule.Config.Patterns;
 			LastTriggerChange = schedule.LastTriggerChange;
-			LastTriggerTime = schedule.GetLastTriggerTimeUtc();
-			ActiveJobs = schedule.ActiveJobs.ConvertAll(x => x.ToString());
+			LastTriggerTime = schedule.LastTriggerTimeUtc;
+			ActiveJobs = new List<JobId>(schedule.ActiveJobs);
 		}
 	}
 
diff --git a/Engine/Source/Programs/Horde/Horde.Build/Jobs/TestData/ITestData.cs b/Engine/Source/Programs/Horde/Horde.Build/Jobs/TestData/ITestData.cs
--- a/Engine/Source/Programs/Horde/Horde.Build/Jobs/TestData/ITestData.cs
+++ b/Engine/Source/Programs/Horde/Horde.Build/Jobs/TestData/ITestData.cs
@@ -8,7 +8,7 @@ namespace Horde.Build.Jobs.TestData
 {
 	using JobId = ObjectId<IJob>;
 	using StreamId = StringId<IStream>;
-	using TemplateRefId = StringId<TemplateRef>;
+	using TemplateId = StringId<ITemplateRef>;
 
 	/// <summary>
 	/// Stores information about the results of a test
@@ -28,7 +28,7 @@ namespace Horde.Build.Jobs.TestData
 		/// <summary>
 		/// The template reference id
 		/// </summary>
-		TemplateRefId TemplateRefId { get; }
+		TemplateId TemplateRefId { get; }
 
 		/// <summary>
 		/// The job which produced the data
diff --git a/Engine/Source/Programs/Horde/Horde.Build/Jobs/TestData/TestDataCollection.cs b/Engine/Source/Programs/Horde/Horde.Build/Jobs/TestData/TestDataCollection.cs
--- a/Engine/Source/Programs/Horde/Horde.Build/Jobs/TestData/TestDataCollection.cs
+++ b/Engine/Source/Programs/Horde/Horde.Build/Jobs/TestData/TestDataCollection.cs
@@ -14,7 +14,7 @@ namespace Horde.Build.Jobs.TestData
 {
 	using JobId = ObjectId<IJob>;
 	using StreamId = StringId<IStream>;
-	using TemplateRefId = StringId<TemplateRef>;
+	using TemplateId = StringId<ITemplateRef>;
 
 	/// <summary>
 	/// Collection of test data documents
@@ -28,7 +28,7 @@ namespace Horde.Build.Jobs.TestData
 		{
 			public ObjectId Id { get; set; }
 			public StreamId StreamId { get; set; }
-			public TemplateRefId TemplateRefId { get; set; }
+			public TemplateId TemplateRefId { get; set; }
 			public JobId JobId { get; set; }
 			public SubResourceId StepId { get; set; }
 			public int Change { get; set; }
diff --git a/Engine/Source/Programs/Horde/Horde.Build/Notifications/IEvent.cs b/Engine/Source/Programs/Horde/Horde.Build/Notifications/IEvent.cs
--- a/Engine/Source/Programs/Horde/Horde.Build/Notifications/IEvent.cs
+++ b/Engine/Source/Programs/Horde/Horde.Build/Notifications/IEvent.cs
@@ -6,7 +6,7 @@ using Horde.Build.Utilities;
 namespace Horde.Build.Notifications
 {
 	using StreamId = StringId<IStream>;
-	using TemplateRefId = StringId<TemplateRef>;
+	using TemplateId = StringId<ITemplateRef>;
 
 	/// <summary>
 	/// Base interface for events that users can subscribe for notifications to
@@ -33,7 +33,7 @@ namespace Horde.Build.Notifications
 		/// <summary>
 		/// The template id
 		/// </summary>
-		TemplateRefId TemplateId { get; }
+		TemplateId TemplateId { get; }
 	}
 
 	/// <summary>
@@ -49,7 +49,7 @@ namespace Horde.Build.Notifications
 		/// <summary>
 		/// The template id
 		/// </summary>
-		TemplateRefId TemplateId { get; }
+		TemplateId TemplateId { get; }
 
 		/// <summary>
 		/// Name of the label to monitor
@@ -70,7 +70,7 @@ namespace Horde.Build.Notifications
 		/// <summary>
 		/// The template id
 		/// </summary>
-		TemplateRefId TemplateId { get; }
+		TemplateId TemplateId { get; }
 
 		/// <summary>
 		/// Name of the step to monitor
diff --git a/Engine/Source/Programs/Horde/Horde.Build/Notifications/Sinks/SlackNotificationSink.cs b/Engine/Source/Programs/Horde/Horde.Build/Notifications/Sinks/SlackNotificationSink.cs
--- a/Engine/Source/Programs/Horde/Horde.Build/Notifications/Sinks/SlackNotificationSink.cs
+++ b/Engine/Source/Programs/Horde/Horde.Build/Notifications/Sinks/SlackNotificationSink.cs
@@ -43,7 +43,7 @@ namespace Horde.Build.Notifications.Sinks
 	using JobId = ObjectId<IJob>;
 	using LogId = ObjectId<ILogFile>;
 	using StreamId = StringId<IStream>;
-	using TemplateRefId = StringId<TemplateRef>;
+	using TemplateId = StringId<ITemplateRef>;
 	using UserId = ObjectId<IUser>;
 	using WorkflowId = StringId<WorkflowConfig>;
 
@@ -704,10 +704,10 @@ namespace Horde.Build.Notifications.Sinks
 				IStream? stream = await _streamService.GetCachedStream(span.StreamId);
 				if (stream != null)
 				{
-					TemplateRef? templateRef;
-					if (stream.Templates.TryGetValue(span.TemplateRefId, out templateRef) && templateRef.TriageChannel != null)
+					ITemplateRef? templateRef;
+					if (stream.Templates.TryGetValue(span.TemplateRefId, out templateRef) && templateRef.Config.TriageChannel != null)
 					{
-						channels.Add(templateRef.TriageChannel);
+						channels.Add(templateRef.Config.TriageChannel);
 					}
 					else if (stream.Config.TriageChannel != null)
 					{
@@ -1360,7 +1360,7 @@ namespace Horde.Build.Notifications.Sinks
 			public StreamId StreamId { get; set; }
 
 			[JsonPropertyName("t")]
-			public TemplateRefId TemplateId { get; set; }
+			public TemplateId TemplateId { get; set; }
 
 			[JsonPropertyName("i")]
 			public List<int> IssueIds { get; set; } = new List<int>();
@@ -1409,7 +1409,7 @@ namespace Horde.Build.Notifications.Sinks
 					issueIdToInfo[issue.Id] = issue;
 				}
 
-				foreach (IGrouping<TemplateRefId, IIssueSpan> group in report.IssueSpans.GroupBy(x => x.TemplateRefId).OrderBy(x => x.Key.ToString()))
+				foreach (IGrouping<TemplateId, IIssueSpan> group in report.IssueSpans.GroupBy(x => x.TemplateRefId).OrderBy(x => x.Key.ToString()))
 				{
 					TemplateRefConfig? templateConfig;
 					if (!report.Stream.Config.TryGetTemplate(group.Key, out templateConfig))
@@ -1492,7 +1492,7 @@ namespace Horde.Build.Notifications.Sinks
 			}
 		}
 
-		static bool IsIssueOpenForWorkflow(StreamId streamId, TemplateRefId templateId, WorkflowId workflowId, IEnumerable<IIssueSpan> spans)
+		static bool IsIssueOpenForWorkflow(StreamId streamId, TemplateId templateId, WorkflowId workflowId, IEnumerable<IIssueSpan> spans)
 		{
 			foreach (IIssueSpan span in spans)
 			{
@@ -1592,7 +1592,7 @@ namespace Horde.Build.Notifications.Sinks
 			}
 		}
 
-		async Task UpdateReportBlockAsync(string channel, string eventId, DateTime reportTime, IStream stream, TemplateRefId templateId, List<(IIssue, IIssueSpan?, bool)> issues, bool templateHeader)
+		async Task UpdateReportBlockAsync(string channel, string eventId, DateTime reportTime, IStream stream, TemplateId templateId, List<(IIssue, IIssueSpan?, bool)> issues, bool templateHeader)
 		{
 			StringBuilder body = new StringBuilder();
 
diff --git a/Engine/Source/Programs/Horde/Horde.Build/Notifications/SubscriptionCollection.cs b/Engine/Source/Programs/Horde/Horde.Build/Notifications/SubscriptionCollection.cs
--- a/Engine/Source/Programs/Horde/Horde.Build/Notifications/SubscriptionCollection.cs
+++ b/Engine/Source/Programs/Horde/Horde.Build/Notifications/SubscriptionCollection.cs
@@ -19,7 +19,7 @@ using MongoDB.Driver;
 namespace Horde.Build.Notifications
 {
 	using StreamId = StringId<IStream>;
-	using TemplateRefId = StringId<TemplateRef>;
+	using TemplateId = StringId<ITemplateRef>;
 	using UserId = ObjectId<IUser>;
 
 	/// <summary>
@@ -59,7 +59,7 @@ namespace Horde.Build.Notifications
 		class JobCompleteEvent : Event, IJobCompleteEvent
 		{
 			public StreamId StreamId { get; set; }
-			public TemplateRefId TemplateId { get; set; }
+			public TemplateId TemplateId { get; set; }
 			public LabelOutcome Outcome { get; set; }
 
 			public JobCompleteEvent()
@@ -69,7 +69,7 @@ namespace Horde.Build.Notifications
 			public JobCompleteEvent(JobCompleteEventRecord record)
 			{
 				StreamId = new StreamId(record.StreamId);
-				TemplateId = new TemplateRefId(record.TemplateId);
+				TemplateId = new TemplateId(record.TemplateId);
 				Outcome = record.Outcome;
 			}
 
@@ -87,7 +87,7 @@ namespace Horde.Build.Notifications
 		class LabelCompleteEvent : Event, ILabelCompleteEvent
 		{
 			public StreamId StreamId { get; set; }
-			public TemplateRefId TemplateId { get; set; }
+			public TemplateId TemplateId { get; set; }
 			public string? CategoryName { get; set; }
 			public string LabelName { get; set; }
 			public LabelOutcome Outcome { get; set; }
@@ -100,7 +100,7 @@ namespace Horde.Build.Notifications
 			public LabelCompleteEvent(LabelCompleteEventRecord record)
 			{
 				StreamId = new StreamId(record.StreamId);
-				TemplateId = new TemplateRefId(record.TemplateId);
+				TemplateId = new TemplateId(record.TemplateId);
 				CategoryName = record.CategoryName;
 				LabelName = record.LabelName;
 				Outcome = record.Outcome;
@@ -127,7 +127,7 @@ namespace Horde.Build.Notifications
 		class StepCompleteEvent : Event, IStepCompleteEvent
 		{
 			public StreamId StreamId { get; set; }
-			public TemplateRefId TemplateId { get; set; }
+			public TemplateId TemplateId { get; set; }
 			public string StepName { get; set; }
 			public JobStepOutcome Outcome { get; set; }
 
@@ -139,7 +139,7 @@ namespace Horde.Build.Notifications
 			public StepCompleteEvent(StepCompleteEventRecord record)
 			{
 				StreamId = new StreamId(record.StreamId);
-				TemplateId = new TemplateRefId(record.TemplateId);
+				TemplateId = new TemplateId(record.TemplateId);
 				StepName = record.StepName;
 				Outcome = record.Outcome;
 			}
diff --git a/Engine/Source/Programs/Horde/Horde.Build/Notifications/SubscriptionMessages.cs b/Engine/Source/Programs/Horde/Horde.Build/Notifications/SubscriptionMessages.cs
--- a/Engine/Source/Programs/Horde/Horde.Build/Notifications/SubscriptionMessages.cs
+++ b/Engine/Source/Programs/Horde/Horde.Build/Notifications/SubscriptionMessages.cs
@@ -11,7 +11,7 @@ using HordeCommon;
 namespace Horde.Build.Notifications
 {
 	using StreamId = StringId<IStream>;
-	using TemplateRefId = StringId<TemplateRef>;
+	using TemplateId = StringId<ITemplateRef>;
 
 	/// <summary>
 	/// Request object for creating a new subscription
@@ -140,7 +140,7 @@ namespace Horde.Build.Notifications
 		/// <param name="streamId">The stream id</param>
 		/// <param name="templateId">The template id</param>
 		/// <param name="outcome">Outcome of the job</param>
-		public JobCompleteEventRecord(StreamId streamId, TemplateRefId templateId, LabelOutcome outcome)
+		public JobCompleteEventRecord(StreamId streamId, TemplateId templateId, LabelOutcome outcome)
 		{
 			StreamId = streamId.ToString();
 			TemplateId = templateId.ToString();
@@ -200,7 +200,7 @@ namespace Horde.Build.Notifications
 		/// <param name="categoryName">Name of the category</param>
 		/// <param name="labelName">The label name</param>
 		/// <param name="outcome">Outcome of the label</param>
-		public LabelCompleteEventRecord(StreamId streamId, TemplateRefId templateId, string? categoryName, string labelName, LabelOutcome outcome)
+		public LabelCompleteEventRecord(StreamId streamId, TemplateId templateId, string? categoryName, string labelName, LabelOutcome outcome)
 		{
 			StreamId = streamId.ToString();
 			TemplateId = templateId.ToString();
@@ -256,7 +256,7 @@ namespace Horde.Build.Notifications
 		/// <param name="templateId">The template id</param>
 		/// <param name="stepName">The label name</param>
 		/// <param name="outcome">Outcome of the step</param>
-		public StepCompleteEventRecord(StreamId streamId, TemplateRefId templateId, string stepName, JobStepOutcome outcome)
+		public StepCompleteEventRecord(StreamId streamId, TemplateId templateId, string stepName, JobStepOutcome outcome)
 		{
 			StreamId = streamId.ToString();
 			TemplateId = templateId.ToString();
diff --git a/Engine/Source/Programs/Horde/Horde.Build/Streams/IStream.cs b/Engine/Source/Programs/Horde/Horde.Build/Streams/IStream.cs
--- a/Engine/Source/Programs/Horde/Horde.Build/Streams/IStream.cs
+++ b/Engine/Source/Programs/Horde/Horde.Build/Streams/IStream.cs
@@ -8,20 +8,17 @@ using System.Text.RegularExpressions;
 using EpicGames.Core;
 using Horde.Build.Acls;
 using Horde.Build.Agents;
-using Horde.Build.Agents.Pools;
-using Horde.Build.Jobs.Schedules;
-using Horde.Build.Jobs.Templates;
+using Horde.Build.Jobs;
 using Horde.Build.Projects;
 using Horde.Build.Users;
 using Horde.Build.Utilities;
-using HordeCommon;
-using MongoDB.Bson.Serialization.Attributes;
 
 namespace Horde.Build.Streams
 {
+	using JobId = ObjectId<IJob>;
 	using ProjectId = StringId<IProject>;
 	using StreamId = StringId<IStream>;
-	using TemplateRefId = StringId<TemplateRef>;
+	using TemplateId = StringId<ITemplateRef>;
 	using UserId = ObjectId<IUser>;
 
 	/// <summary>
@@ -46,257 +43,142 @@ namespace Horde.Build.Streams
 	}
 
 	/// <summary>
-	/// Allows triggering another downstream job on succesful completion of a step or aggregate
+	/// Information about a stream
 	/// </summary>
-	public class ChainedJobTemplate
+	public interface IStream
 	{
 		/// <summary>
-		/// Name of the target that needs to complete successfully
+		/// Name of the stream.
 		/// </summary>
-		public string Trigger { get; set; } = String.Empty;
+		StreamId Id { get; }
 
 		/// <summary>
-		/// The new template to trigger
+		/// The project that this stream belongs to
 		/// </summary>
-		public TemplateRefId TemplateRefId { get; set; }
+		ProjectId ProjectId { get; }
 
 		/// <summary>
-		/// Default constructor for serialization
+		/// The stream name
 		/// </summary>
-		private ChainedJobTemplate()
-		{
-		}
+		string Name { get; }
 
 		/// <summary>
-		/// Constructor
+		/// The revision of config file used for this stream
 		/// </summary>
-		/// <param name="trigger">Name of the target that needs to complete</param>
-		/// <param name="templateRefId">The new template to trigger</param>
-		public ChainedJobTemplate(string trigger, TemplateRefId templateRefId)
-		{
-			Trigger = trigger;
-			TemplateRefId = templateRefId;
-		}
+		string ConfigRevision { get; }
 
 		/// <summary>
-		/// Constructor
+		/// Configuration settings for the stream
 		/// </summary>
-		/// <param name="request">Request to construct from</param>
-		public ChainedJobTemplate(ChainedJobTemplateConfig request)
-			: this(request.Trigger, new TemplateRefId(request.TemplateId))
-		{
-		}
-	}
+		StreamConfig Config { get; }
 
-	/// <summary>
-	/// A
-	/// </summary>
-	public class TemplateStepState
-	{
 		/// <summary>
-		/// Name of the step
+		/// Whether this stream has been deleted
 		/// </summary>
-		public string Name { get; set; } = String.Empty ;
+		bool Deleted { get; }
 
 		/// <summary>
-		/// User who paused the step
+		/// List of templates available for this stream
 		/// </summary>
-		public UserId? PausedByUserId { get; set; }
+		IReadOnlyDictionary<TemplateId, ITemplateRef> Templates { get; }
 
 		/// <summary>
-		/// The UTC time when the step was paused
+		/// Stream is paused for builds until specified time
 		/// </summary>
-		public DateTime? PauseTimeUtc { get; set; }
-
+		DateTime? PausedUntil { get; }
+		
 		/// <summary>
-		/// Default constructor for serialization
+		/// Comment/reason for why the stream was paused
 		/// </summary>
-		private TemplateStepState()
-		{
-		}
+		string? PauseComment { get; }
 
 		/// <summary>
-		/// Constructor
+		/// The ACL for this object
 		/// </summary>
-		public TemplateStepState( string name, UserId? pausedByUserId = null, DateTime? pauseTimeUtc = null)
-		{
-			Name = name;
-			PausedByUserId = pausedByUserId;
-			PauseTimeUtc = pauseTimeUtc;
-		}
+		Acl? Acl { get; }
 	}
 
 	/// <summary>
-	/// Reference to a template
+	/// Job template in a stream
 	/// </summary>
-	public class TemplateRef
+	public interface ITemplateRef
 	{
 		/// <summary>
-		/// The template name (duplicated from the template object)
+		/// The template id
 		/// </summary>
-		public string Name { get; set; }
+		TemplateId Id { get; }
 
 		/// <summary>
-		/// Hash of the template definition
+		/// Configuration of this template ref
 		/// </summary>
-		public ContentHash Hash { get; set; }
+		TemplateRefConfig Config { get; }
 
 		/// <summary>
-		/// Whether to show badges in UGS for this schedule
-		/// </summary>
-		public bool ShowUgsBadges { get; set; }
-
-		/// <summary>
-		/// Whether to show desktop alerts for build health issues created from jobs this type
-		/// </summary>
-		public bool ShowUgsAlerts { get; set; }
-
-		/// <summary>
-		/// Notification channel for this template. Overrides the stream channel if set.
-		/// </summary>
-		public string? NotificationChannel { get; set; }
-
-		/// <summary>
-		/// Notification channel filter for this template. Errors|Warnings|Success
-		/// </summary>
-		public string? NotificationChannelFilter { get; set; }
-
-		/// <summary>
-		/// Channel for triage notification messages
+		/// Hash of the template definition
 		/// </summary>
-		public string? TriageChannel { get; set; }
+		ContentHash Hash { get; }
 
 		/// <summary>
 		/// List of schedules for this template
 		/// </summary>
-		[BsonIgnoreIfNull]
-		public Schedule? Schedule { get; set; }
-
-		/// <summary>
-		/// List of downstream templates to trigger at the same change
-		/// </summary>
-		[BsonIgnoreIfNull]
-		public List<ChainedJobTemplate>? ChainedJobs { get; set; }
+		ITemplateSchedule? Schedule { get; }
 
 		/// <summary>
 		/// List of template step states
 		/// </summary>
-		[BsonIgnoreIfNull]
-		public List<TemplateStepState>? StepStates { get;set; }
+		IReadOnlyList<ITemplateStep> StepStates { get; }
 
 		/// <summary>
 		/// Custom permissions for this template
 		/// </summary>
-		public Acl? Acl { get; set; }
-
-		/// <summary>
-		/// Private constructor for serialization
-		/// </summary>
-		private TemplateRef()
-		{
-			Name = null!;
-			Hash = null!;
-		}
-
-		/// <summary>
-		/// Constructor
-		/// </summary>
-		/// <param name="template">The template being referenced</param>
-		/// <param name="showUgsBadges">Whether to show badges in UGS for this job</param>
-		/// <param name="showUgsAlerts">Whether to show alerts in UGS for this job</param>
-		/// <param name="notificationChannel">Notification channel for this template</param>
-		/// <param name="notificationChannelFilter">Notification channel filter for this template</param>
-		/// <param name="triageChannel"></param>
-		/// <param name="schedule">Schedule for this template</param>
-		/// <param name="triggers">List of downstream templates to trigger</param>
-		/// <param name="stepStates">List of template step states</param>
-		/// <param name="acl">ACL for this template</param>
-		public TemplateRef(ITemplate template, bool showUgsBadges = false, bool showUgsAlerts = false, string? notificationChannel = null, string? notificationChannelFilter = null, string? triageChannel = null, Schedule? schedule = null, List<ChainedJobTemplate>? triggers = null, List<TemplateStepState>? stepStates = null, Acl? acl = null)
-		{
-			Name = template.Name;
-			Hash = template.Id;
-			ShowUgsBadges = showUgsBadges;
-			ShowUgsAlerts = showUgsAlerts;
-			NotificationChannel = notificationChannel;
-			NotificationChannelFilter = notificationChannelFilter;
-			TriageChannel = triageChannel;
-			Schedule = schedule;
-			ChainedJobs = triggers;
-			StepStates = stepStates;
-			Acl = acl;
-		}
-	}
-
-	/// <summary>
-	/// Extension methods for template refs
-	/// </summary>
-	static class TemplateRefExtensions
-	{
-		/// <summary>
-		/// Adds a new template ref to a list
-		/// </summary>
-		/// <param name="templateRefs">List of template refs</param>
-		/// <param name="templateRef">The template ref to add</param>
-		public static void AddRef(this Dictionary<TemplateRefId, TemplateRef> templateRefs, TemplateRef templateRef)
-		{
-			templateRefs.Add(new TemplateRefId(templateRef.Name), templateRef);
-		}
+		Acl? Acl { get; set; }
 	}
 
 	/// <summary>
-	/// Information about a stream
+	/// Schedule for a template
 	/// </summary>
-	public interface IStream
+	public interface ITemplateSchedule
 	{
 		/// <summary>
-		/// Name of the stream.
+		/// Config for this schedule
 		/// </summary>
-		public StreamId Id { get; }
+		ScheduleConfig Config { get; }
 
 		/// <summary>
-		/// The project that this stream belongs to
+		/// Last changelist number that this was triggered for
 		/// </summary>
-		public ProjectId ProjectId { get; }
+		int LastTriggerChange { get; }
 
 		/// <summary>
-		/// The stream name
+		/// Gets the last trigger time, in UTC
 		/// </summary>
-		public string Name { get; }
+		DateTime LastTriggerTimeUtc { get; }
 
 		/// <summary>
-		/// The revision of config file used for this stream
+		/// List of jobs that are currently active
 		/// </summary>
-		public string ConfigRevision { get; }
-
-		/// <summary>
-		/// Configuration settings for the stream
-		/// </summary>
-		public StreamConfig Config { get; }
-
-		/// <summary>
-		/// Whether this stream has been deleted
-		/// </summary>
-		public bool Deleted { get; }
+		IReadOnlyList<JobId> ActiveJobs { get; }
+	}
 
+	/// <summary>
+	/// Information about a paused template step
+	/// </summary>
+	public interface ITemplateStep
+	{
 		/// <summary>
-		/// List of templates available for this stream
+		/// Name of the step
 		/// </summary>
-		public IReadOnlyDictionary<TemplateRefId, TemplateRef> Templates { get; }
+		string Name { get; }
 
 		/// <summary>
-		/// Stream is paused for builds until specified time
-		/// </summary>
-		public DateTime? PausedUntil { get; }
-		
-		/// <summary>
-		/// Comment/reason for why the stream was paused
+		/// User who paused the step
 		/// </summary>
-		public string? PauseComment { get; }
+		UserId PausedByUserId { get; }
 
 		/// <summary>
-		/// The ACL for this object
+		/// The UTC time when the step was paused
 		/// </summary>
-		public Acl? Acl { get; }
+		DateTime PauseTimeUtc { get; }
 	}
 
 	/// <summary>
@@ -304,6 +186,17 @@ namespace Horde.Build.Streams
 	/// </summary>
 	static class StreamExtensions
 	{
+		/// <summary>
+		/// Gets the next trigger time for a schedule
+		/// </summary>
+		/// <param name="schedule"></param>
+		/// <param name="timeZone"></param>
+		/// <returns></returns>
+		public static DateTime? GetNextTriggerTimeUtc(this ITemplateSchedule schedule, TimeZoneInfo timeZone)
+		{
+			return schedule.Config.GetNextTriggerTimeUtc(schedule.LastTriggerTimeUtc, timeZone);
+		}
+
 		/// <summary>
 		/// Tries to get an agent workspace definition from the given type name
 		/// </summary>
diff --git a/Engine/Source/Programs/Horde/Horde.Build/Streams/IStreamCollection.cs b/Engine/Source/Programs/Horde/Horde.Build/Streams/IStreamCollection.cs
--- a/Engine/Source/Programs/Horde/Horde.Build/Streams/IStreamCollection.cs
+++ b/Engine/Source/Programs/Horde/Horde.Build/Streams/IStreamCollection.cs
@@ -12,7 +12,7 @@ namespace Horde.Build.Streams
 	using JobId = ObjectId<IJob>;
 	using ProjectId = StringId<IProject>;
 	using StreamId = StringId<IStream>;
-	using TemplateRefId = StringId<TemplateRef>;
+	using TemplateId = StringId<ITemplateRef>;
 
 	/// <summary>
 	/// Collection of stream documents
@@ -74,7 +74,7 @@ namespace Horde.Build.Streams
 		/// <param name="lastTriggerChange">New last trigger changelist for the schedule</param>
 		/// <param name="newActiveJobs">New list of active jobs</param>
 		/// <returns>The updated stream if successful, null otherwise</returns>
-		Task<IStream?> TryUpdateScheduleTriggerAsync(IStream stream, TemplateRefId templateRefId, DateTime? lastTriggerTimeUtc, int? lastTriggerChange, List<JobId> newActiveJobs);
+		Task<IStream?> TryUpdateScheduleTriggerAsync(IStream stream, TemplateId templateRefId, DateTime? lastTriggerTimeUtc, int? lastTriggerChange, List<JobId> newActiveJobs);
 
 		/// <summary>
 		/// Attempts to update a stream template ref
@@ -83,7 +83,7 @@ namespace Horde.Build.Streams
 		/// <param name="templateRefId">The template ref to update</param>
 		/// <param name="stepStates">The stream states to update, pass an empty list to clear all step states, otherwise will be a partial update based on included step updates</param>
 		/// <returns></returns>
-		Task<IStream?> TryUpdateTemplateRefAsync(IStream streamInterface, TemplateRefId templateRefId, List<UpdateStepStateRequest>? stepStates = null);
+		Task<IStream?> TryUpdateTemplateRefAsync(IStream streamInterface, TemplateId templateRefId, List<UpdateStepStateRequest>? stepStates = null);
 
 		/// <summary>
 		/// Delete a stream
diff --git a/Engine/Source/Programs/Horde/Horde.Build/Streams/StreamCollection.cs b/Engine/Source/Programs/Horde/Horde.Build/Streams/StreamCollection.cs
--- a/Engine/Source/Programs/Horde/Horde.Build/Streams/StreamCollection.cs
+++ b/Engine/Source/Programs/Horde/Horde.Build/Streams/StreamCollection.cs
@@ -6,6 +6,7 @@ using System.Diagnostics.CodeAnalysis;
 using System.IO;
 using System.Linq;
 using System.Threading.Tasks;
+using EpicGames.Core;
 using Horde.Build.Acls;
 using Horde.Build.Configuration;
 using Horde.Build.Jobs;
@@ -17,6 +18,9 @@ using Horde.Build.Users;
 using Horde.Build.Utilities;
 using HordeCommon;
 using Microsoft.Extensions.Logging;
+using MongoDB.Bson;
+using MongoDB.Bson.IO;
+using MongoDB.Bson.Serialization;
 using MongoDB.Bson.Serialization.Attributes;
 using MongoDB.Driver;
 
@@ -25,7 +29,7 @@ namespace Horde.Build.Streams
 	using JobId = ObjectId<IJob>;
 	using ProjectId = StringId<IProject>;
 	using StreamId = StringId<IStream>;
-	using TemplateRefId = StringId<TemplateRef>;
+	using TemplateId = StringId<ITemplateRef>;
 	using UserId = ObjectId<IUser>;
 
 	/// <summary>
@@ -36,7 +40,7 @@ namespace Horde.Build.Streams
 		/// <summary>
 		/// Information about a stream
 		/// </summary>
-		class StreamDocument : IStream
+		class StreamDoc : IStream
 		{
 			[BsonRequired, BsonId]
 			public StreamId Id { get; set; }
@@ -49,7 +53,7 @@ namespace Horde.Build.Streams
 			[BsonIgnore]
 			public StreamConfig? Config { get; set; }
 
-			public Dictionary<TemplateRefId, TemplateRef> Templates { get; set; } = new Dictionary<TemplateRefId, TemplateRef>();
+			public Dictionary<TemplateId, TemplateRefDoc> Templates { get; set; } = new Dictionary<TemplateId, TemplateRefDoc>();
 			public DateTime? PausedUntil { get; set; }
 			public string? PauseComment { get; set; }
 
@@ -58,22 +62,133 @@ namespace Horde.Build.Streams
 			public bool Deleted { get; set; }
 
 			StreamConfig IStream.Config => Config!;
-			IReadOnlyDictionary<TemplateRefId, TemplateRef> IStream.Templates => Templates;
+
+			IReadOnlyDictionary<TemplateId, ITemplateRef>? _cachedTemplates;
+			IReadOnlyDictionary<TemplateId, ITemplateRef> IStream.Templates
+			{
+				get
+				{
+					_cachedTemplates ??= Templates.ToDictionary(x => x.Key, x => (ITemplateRef)x.Value);
+					return _cachedTemplates;
+				}
+			}
 
 			string IStream.Name => Config!.Name;
 
 			[BsonConstructor]
-			private StreamDocument()
+			private StreamDoc()
 			{
 			}
 
-			public StreamDocument(StreamId id, ProjectId projectId)
+			public StreamDoc(StreamId id, ProjectId projectId)
 			{
 				Id = id;
 				ProjectId = projectId;
 			}
 		}
 
+		class TemplateRefDoc : ITemplateRef
+		{
+			#region Serialized
+
+			[BsonRequired]
+			public ContentHash Hash { get; set; } = null!;
+
+			[BsonIgnoreIfNull]
+			public TemplateScheduleDoc? Schedule { get; set; }
+
+			[BsonIgnoreIfNull]
+			public List<TemplateStepDoc>? StepStates { get; set; }
+
+			[BsonIgnoreIfNull]
+			public Acl? Acl { get; set; }
+
+			#endregion
+
+			[BsonIgnore]
+			StreamDoc? _owner;
+
+			[BsonIgnore]
+			public TemplateId Id { get; private set; }
+
+			[BsonIgnore]
+			TemplateRefConfig? _config;
+
+			[BsonIgnore]
+			public TemplateRefConfig Config
+			{
+				get
+				{
+					_config ??= _owner?.Config?.Templates.First(x => x.Id == Id);
+					return _config!;
+				}
+			}
+
+			ITemplateSchedule? ITemplateRef.Schedule => Schedule;
+			IReadOnlyList<ITemplateStep> ITemplateRef.StepStates => (IReadOnlyList<ITemplateStep>?)StepStates ?? Array.Empty<ITemplateStep>();
+
+			public void PostLoad(StreamDoc owner, TemplateId id)
+			{
+				_owner = owner;
+				Id = id;
+				Schedule?.PostLoad(this);
+			}
+		}
+
+		class TemplateScheduleDoc : ITemplateSchedule
+		{
+			#region Serialized
+
+			public int LastTriggerChange { get; set; }
+
+			[BsonIgnoreIfNull, Obsolete("Use LastTriggerTimeUtc instead")]
+			public DateTimeOffset? LastTriggerTime { get; set; }
+
+			public DateTime LastTriggerTimeUtc { get; set; }
+			public List<JobId> ActiveJobs { get; set; } = new List<JobId>();
+
+			#endregion
+
+			[BsonIgnore]
+			TemplateRefDoc? _owner;
+
+			[BsonIgnore]
+			public ScheduleConfig Config => _owner?.Config.Schedule!;
+
+			IReadOnlyList<JobId> ITemplateSchedule.ActiveJobs => ActiveJobs;
+
+			public void PostLoad(TemplateRefDoc owner)
+			{
+				_owner = owner;
+
+#pragma warning disable CS0618 // Type or member is obsolete
+				if (LastTriggerTime.HasValue)
+				{
+					LastTriggerTimeUtc = LastTriggerTime.Value.UtcDateTime;
+					LastTriggerTime = null;
+				}
+#pragma warning restore CS0618 // Type or member is obsolete
+			}
+		}
+
+		class TemplateStepDoc : ITemplateStep
+		{
+			public string Name { get; set; } = String.Empty;
+			public UserId PausedByUserId { get; set; }
+			public DateTime PauseTimeUtc { get; set; }
+
+			public TemplateStepDoc()
+			{
+			}
+
+			public TemplateStepDoc(string name, UserId pausedByUserId, DateTime pauseTimeUtc)
+			{
+				Name = name;
+				PausedByUserId = pausedByUserId;
+				PauseTimeUtc = pauseTimeUtc;
+			}
+		}
+
 		/// <summary>
 		/// Projection of a stream definition to just include permissions info
 		/// </summary>
@@ -83,10 +198,10 @@ namespace Horde.Build.Streams
 			public Acl? Acl { get; set; }
 			public ProjectId ProjectId { get; set; }
 
-			public static readonly ProjectionDefinition<StreamDocument> Projection = Builders<StreamDocument>.Projection.Include(x => x.Acl).Include(x => x.ProjectId);
+			public static readonly ProjectionDefinition<StreamDoc> Projection = Builders<StreamDoc>.Projection.Include(x => x.Acl).Include(x => x.ProjectId);
 		}
 
-		readonly IMongoCollection<StreamDocument> _streams;
+		readonly IMongoCollection<StreamDoc> _streams;
 		readonly ConfigCollection _configCollection;
 		readonly IClock _clock;
 		readonly ITemplateCollection _templateCollection;
@@ -102,14 +217,14 @@ namespace Horde.Build.Streams
 		/// <param name="logger"></param>
 		public StreamCollection(MongoService mongoService, ConfigCollection configCollection, IClock clock, ITemplateCollection templateCollection, ILogger<StreamCollection> logger)
 		{
-			_streams = mongoService.GetCollection<StreamDocument>("Streams");
+			_streams = mongoService.GetCollection<StreamDoc>("Streams");
 			_configCollection = configCollection;
 			_clock = clock;
 			_templateCollection = templateCollection;
 			_logger = logger;
 		}
 
-		async Task PostLoadAsync(StreamDocument stream)
+		async Task PostLoadAsync(StreamDoc stream)
 		{
 			try
 			{
@@ -127,25 +242,30 @@ namespace Horde.Build.Streams
 				_logger.LogError("Unable to get stream config for {StreamId} at {Revision}", stream.Id, stream.ConfigRevision);
 				throw;
 			}
+
+			foreach (KeyValuePair<TemplateId, TemplateRefDoc> pair in stream.Templates)
+			{
+				pair.Value.PostLoad(stream, pair.Key);
+			}
 		}
 
 		/// <inheritdoc/>
 		public async Task<IStream?> TryCreateOrReplaceAsync(StreamId id, IStream? stream, string revision, ProjectId projectId)
 		{
 			StreamConfig config = await _configCollection.GetConfigAsync<StreamConfig>(revision);
+			Validate(id, config);
 
-			Dictionary<TemplateRefId, TemplateRef> templateRefs = await CreateTemplateRefsAsync(config.Templates, stream, _templateCollection);
-
-			Validate(id, templateRefs, config);
+			StreamDoc? streamDoc = (StreamDoc?)stream;
+			Dictionary<TemplateId, TemplateRefDoc> templateRefs = await CreateTemplateRefsAsync(streamDoc, config.Templates);
 
 			Acl? acl = Acl.Merge(new Acl(), config.Acl);
-			if (stream == null)
+			if (streamDoc == null)
 			{
-				return await TryCreateAsync(id, projectId, revision, config, templateRefs, acl);
+				return await TryCreateAsync(id, projectId, revision, templateRefs, acl);
 			}
 			else
 			{
-				return await TryReplaceAsync(stream, projectId, revision, config, templateRefs, acl);
+				return await TryReplaceAsync(streamDoc, projectId, revision, config, templateRefs, acl);
 			}
 		}
 
@@ -154,89 +274,62 @@ namespace Horde.Build.Streams
 		/// </summary>
 		/// <param name="requests">Request objects</param>
 		/// <param name="stream">The current stream state</param>
-		/// <param name="templateCollection">The template service</param>
 		/// <returns>List of new template references</returns>
-		async Task<Dictionary<TemplateRefId, TemplateRef>> CreateTemplateRefsAsync(List<TemplateRefConfig> requests, IStream? stream, ITemplateCollection templateCollection)
+		async Task<Dictionary<TemplateId, TemplateRefDoc>> CreateTemplateRefsAsync(StreamDoc? stream, List<TemplateRefConfig> requests)
 		{
-			Dictionary<TemplateRefId, TemplateRef> newTemplateRefs = new Dictionary<TemplateRefId, TemplateRef>();
+			Dictionary<TemplateId, TemplateRefDoc> newTemplateRefs = new Dictionary<TemplateId, TemplateRefDoc>();
 			foreach (TemplateRefConfig request in requests)
 			{
 				// Create the template
-				ITemplate newTemplate = await templateCollection.AddAsync(request.Name, request.Priority, request.AllowPreflights, request.UpdateIssues, request.PromoteIssuesByDefault, request.InitialAgentType, request.SubmitNewChange, request.SubmitDescription, request.Arguments, request.Parameters.ConvertAll(x => x.ToModel()));
+				ITemplate template = await _templateCollection.AddAsync(request.Name, request.Priority, request.AllowPreflights, request.UpdateIssues, request.PromoteIssuesByDefault, request.InitialAgentType, request.SubmitNewChange, request.SubmitDescription, request.Arguments, request.Parameters.ConvertAll(x => x.ToModel()));
 
 				// Get an identifier for the new template ref
-				TemplateRefId newTemplateRefId = request.Id;
+				TemplateId templateId = request.Id;
 
-				// Create the schedule object
-				Schedule? schedule = null;
-				if (request.Schedule != null)
+				// Get the existing template ref or create a new one
+				TemplateRefDoc? templateRef;
+				if (stream == null || !stream.Templates.TryGetValue(templateId, out templateRef))
 				{
-					schedule = request.Schedule.ToModel(_clock.UtcNow);
+					templateRef = new TemplateRefDoc();
 				}
 
-				// Add it to the list
-				TemplateRef newTemplateRef = new TemplateRef(newTemplate, request.ShowUgsBadges, request.ShowUgsAlerts, request.NotificationChannel, request.NotificationChannelFilter, request.TriageChannel, schedule, request.ChainedJobs?.ConvertAll(x => new ChainedJobTemplate(x)), null, Acl.Merge(null, request.Acl));
-				if (stream != null && stream.Templates.TryGetValue(newTemplateRefId, out TemplateRef? oldTemplateRef))
-				{
-					if (oldTemplateRef.Schedule != null && newTemplateRef.Schedule != null)
-					{
-						newTemplateRef.Schedule.CopyState(oldTemplateRef.Schedule);
-					}
-
-					if (oldTemplateRef.StepStates != null)
-					{
-						newTemplateRef.StepStates = new List<TemplateStepState>(oldTemplateRef.StepStates);
-					}
-				}
-				newTemplateRefs.Add(newTemplateRefId, newTemplateRef);
-			}
-			foreach (TemplateRef templateRef in newTemplateRefs.Values)
-			{
-				if (templateRef.ChainedJobs != null)
+				// Update the hash
+				templateRef.Hash = template.Id;
+				if (request.Schedule != null)
 				{
-					foreach (ChainedJobTemplate chainedJob in templateRef.ChainedJobs)
-					{
-						if (!newTemplateRefs.ContainsKey(chainedJob.TemplateRefId))
-						{
-							throw new InvalidDataException($"Invalid template ref id '{chainedJob.TemplateRefId}");
-						}
-					}
+					templateRef.Schedule ??= new TemplateScheduleDoc();
+					templateRef.Schedule.LastTriggerTimeUtc = _clock.UtcNow;
 				}
+
+				// Add it to the new lookup
+				newTemplateRefs.Add(templateId, templateRef);
 			}
 			return newTemplateRefs;
 		}
 
 		/// <inheritdoc/>
-		public async Task<IStream?> TryUpdateTemplateRefAsync(IStream streamInterface, TemplateRefId templateRefId, List<UpdateStepStateRequest>? stepStates = null)
+		public async Task<IStream?> TryUpdateTemplateRefAsync(IStream streamInterface, TemplateId templateId, List<UpdateStepStateRequest>? stepStates = null)
 		{
-			StreamDocument stream = (StreamDocument)streamInterface;
+			StreamDoc stream = (StreamDoc)streamInterface;
 
-			TemplateRef? oldTemplateRef;
-			TemplateRef? newTemplateRef = null;
-
-			if (!stream.Templates.TryGetValue(templateRefId, out oldTemplateRef))
+			TemplateRefDoc? templateRef;
+			if (!stream.Templates.TryGetValue(templateId, out templateRef))
 			{
 				return null;
 			}
 
-			ITemplate? template = await _templateCollection.GetAsync(oldTemplateRef.Hash);
-			if (template == null)
-			{
-				return null;
-			}
+			UpdateDefinitionBuilder<StreamDoc> updateBuilder = Builders<StreamDoc>.Update;
+			List<UpdateDefinition<StreamDoc>> updates = new List<UpdateDefinition<StreamDoc>>();
 
-			UpdateDefinitionBuilder<StreamDocument> updateBuilder = Builders<StreamDocument>.Update;
-			List<UpdateDefinition<StreamDocument>> updates = new List<UpdateDefinition<StreamDocument>>();
-
-			Dictionary<TemplateRefId, TemplateRef> newTemplates = new Dictionary<TemplateRefId, TemplateRef>(stream.Templates);
+			Dictionary<TemplateId, TemplateRefDoc> newTemplates = new Dictionary<TemplateId, TemplateRefDoc>(stream.Templates);
 
 			// clear
 			if (stepStates != null && stepStates.Count == 0)
 			{
 				bool hasUpdates = false;
-				foreach (KeyValuePair<TemplateRefId, TemplateRef> entry in newTemplates)
+				foreach (KeyValuePair<TemplateId, TemplateRefDoc> entry in newTemplates)
 				{
-					if (entry.Value?.StepStates != null)
+					if (entry.Value.StepStates != null)
 					{
 						hasUpdates = true;
 						entry.Value.StepStates = null;
@@ -250,52 +343,47 @@ namespace Horde.Build.Streams
 			}
 			else if (stepStates != null)
 			{
-				newTemplateRef = new TemplateRef(template, oldTemplateRef.ShowUgsAlerts, oldTemplateRef.ShowUgsAlerts, oldTemplateRef.NotificationChannel, oldTemplateRef.NotificationChannelFilter, oldTemplateRef.TriageChannel, oldTemplateRef.Schedule, oldTemplateRef.ChainedJobs, null, oldTemplateRef.Acl);
-
 				// get currently valid step states
-				List<TemplateStepState> newStepStates = oldTemplateRef.StepStates?.Where(x => x.PausedByUserId != null).ToList() ?? new List<TemplateStepState>();
+				List<TemplateStepDoc> newStepStates = templateRef.StepStates?.ToList() ?? new List<TemplateStepDoc>();
 
 				// generate update list
 				foreach (UpdateStepStateRequest updateState in stepStates)
 				{
-					TemplateStepState? newState = newStepStates.Where(x => x.Name == updateState.Name).FirstOrDefault();
+					int stateIndex = newStepStates.FindIndex(x => x.Name == updateState.Name);
 
 					UserId? pausedByUserId = updateState.PausedByUserId != null ? new UserId(updateState.PausedByUserId) : null;
 
-					if (newState == null)
+					if (stateIndex == -1)
 					{
 						// if this is a new state without anything set, ignore it
-						if (pausedByUserId == null)
+						if (pausedByUserId != null)
 						{
-							continue;
+							newStepStates.Add(new TemplateStepDoc(updateState.Name, pausedByUserId.Value, _clock.UtcNow));
 						}
-
-						newStepStates.Add(new TemplateStepState(updateState.Name, pausedByUserId, DateTime.UtcNow));
 					}
 					else
 					{
 						if (pausedByUserId == null)
 						{
-							newState.PauseTimeUtc = null;
+							newStepStates.RemoveAt(stateIndex);
 						}
-						else if (newState.PauseTimeUtc == null)
+						else
 						{
-							newState.PauseTimeUtc = DateTime.UtcNow;
+							newStepStates[stateIndex].PausedByUserId = pausedByUserId.Value;
 						}
-
-						newState.PausedByUserId = pausedByUserId;						
 					}
 				}
 
-				newTemplateRef.StepStates = newStepStates.Where(x => x.PausedByUserId != null).ToList();
-				if (newTemplateRef.StepStates.Count == 0)
+				if (newStepStates.Count == 0)
+				{
+					templateRef.StepStates = null;
+				}
+				else
 				{
-					newTemplateRef.StepStates = null;
+					templateRef.StepStates = newStepStates;
 				}
 
-				newTemplates[templateRefId] = newTemplateRef;
 				updates.Add(updateBuilder.Set(x => x.Templates, newTemplates));
-
 			}
 
 			if (updates.Count == 0)
@@ -308,17 +396,17 @@ namespace Horde.Build.Streams
 		}
 
 		/// <inheritdoc/>
-		async Task<IStream?> TryCreateAsync(StreamId id, ProjectId projectId, string configRevision, StreamConfig config, Dictionary<TemplateRefId, TemplateRef> templateRefs, Acl? acl)
+		async Task<IStream?> TryCreateAsync(StreamId id, ProjectId projectId, string configRevision, Dictionary<TemplateId, TemplateRefDoc> templateRefs, Acl? acl)
 		{
-			StreamDocument newStream = new StreamDocument(id, projectId);
+			StreamDoc newStream = new StreamDoc(id, projectId);
 			newStream.ConfigRevision = configRevision;
-			newStream.Config = config;
 			newStream.Templates = templateRefs;
 			newStream.Acl = acl;
 
 			try
 			{
 				await _streams.InsertOneAsync(newStream);
+				await PostLoadAsync(newStream);
 				return newStream;
 			}
 			catch (MongoWriteException ex)
@@ -335,18 +423,14 @@ namespace Horde.Build.Streams
 		}
 
 		/// <inheritdoc/>
-		async Task<IStream?> TryReplaceAsync(IStream streamInterface, ProjectId projectId, string configRevision, StreamConfig config, Dictionary<TemplateRefId, TemplateRef>? templateRefs, Acl? acl)
+		async Task<IStream?> TryReplaceAsync(StreamDoc stream, ProjectId projectId, string configRevision, StreamConfig config, Dictionary<TemplateId, TemplateRefDoc> templateRefs, Acl? acl)
 		{
-			int order = config.Order;
+			UpdateDefinitionBuilder<StreamDoc> updateBuilder = Builders<StreamDoc>.Update;
 
-			StreamDocument stream = (StreamDocument)streamInterface;
-
-			UpdateDefinitionBuilder<StreamDocument> updateBuilder = Builders<StreamDocument>.Update;
-
-			List<UpdateDefinition<StreamDocument>> updates = new List<UpdateDefinition<StreamDocument>>();
+			List<UpdateDefinition<StreamDoc>> updates = new List<UpdateDefinition<StreamDoc>>();
 			updates.Add(updateBuilder.Set(x => x.ProjectId, projectId));
 			updates.Add(updateBuilder.Set(x => x.ConfigRevision, configRevision));
-			updates.Add(updateBuilder.Set(x => x.Templates, templateRefs ?? new Dictionary<TemplateRefId, TemplateRef>()));
+			updates.Add(updateBuilder.Set(x => x.Templates, templateRefs));
 			updates.Add(updateBuilder.SetOrUnsetNullRef(x => x.Acl, acl));
 			updates.Add(updateBuilder.Unset(x => x.Deleted));
 
@@ -356,7 +440,7 @@ namespace Horde.Build.Streams
 		/// <inheritdoc/>
 		public async Task<IStream?> GetAsync(StreamId streamId)
 		{
-			StreamDocument? stream = await _streams.Find<StreamDocument>(x => x.Id == streamId).FirstOrDefaultAsync();
+			StreamDoc? stream = await _streams.Find<StreamDoc>(x => x.Id == streamId).FirstOrDefaultAsync();
 			if (stream != null)
 			{
 				await PostLoadAsync(stream);
@@ -367,14 +451,14 @@ namespace Horde.Build.Streams
 		/// <inheritdoc/>
 		public async Task<IStreamPermissions?> GetPermissionsAsync(StreamId streamId)
 		{
-			return await _streams.Find<StreamDocument>(x => x.Id == streamId).Project<StreamPermissions>(StreamPermissions.Projection).FirstOrDefaultAsync();
+			return await _streams.Find<StreamDoc>(x => x.Id == streamId).Project<StreamPermissions>(StreamPermissions.Projection).FirstOrDefaultAsync();
 		}
 
 		/// <inheritdoc/>
 		public async Task<List<IStream>> FindAllAsync()
 		{
-			List<StreamDocument> results = await _streams.Find(Builders<StreamDocument>.Filter.Ne(x => x.Deleted, true)).ToListAsync();
-			foreach (StreamDocument result in results)
+			List<StreamDoc> results = await _streams.Find(Builders<StreamDoc>.Filter.Ne(x => x.Deleted, true)).ToListAsync();
+			foreach (StreamDoc result in results)
 			{
 				await PostLoadAsync(result);
 			}
@@ -384,10 +468,10 @@ namespace Horde.Build.Streams
 		/// <inheritdoc/>
 		public async Task<List<IStream>> FindForProjectsAsync(ProjectId[] projectIds)
 		{
-			FilterDefinition<StreamDocument> filter = Builders<StreamDocument>.Filter.In(x => x.ProjectId, projectIds) & Builders<StreamDocument>.Filter.Ne(x => x.Deleted, true);
+			FilterDefinition<StreamDoc> filter = Builders<StreamDoc>.Filter.In(x => x.ProjectId, projectIds) & Builders<StreamDoc>.Filter.Ne(x => x.Deleted, true);
 
-			List<StreamDocument> results = await _streams.Find(filter).ToListAsync();
-			foreach (StreamDocument result in results)
+			List<StreamDoc> results = await _streams.Find(filter).ToListAsync();
+			foreach (StreamDoc result in results)
 			{
 				await PostLoadAsync(result);
 			}
@@ -397,11 +481,11 @@ namespace Horde.Build.Streams
 		/// <inheritdoc/>
 		public async Task<IStream?> TryUpdatePauseStateAsync(IStream streamInterface, DateTime? newPausedUntil, string? newPauseComment)
 		{
-			StreamDocument stream = (StreamDocument)streamInterface;
+			StreamDoc stream = (StreamDoc)streamInterface;
 
-			UpdateDefinitionBuilder<StreamDocument> updateBuilder = Builders<StreamDocument>.Update;
+			UpdateDefinitionBuilder<StreamDoc> updateBuilder = Builders<StreamDoc>.Update;
 
-			List<UpdateDefinition<StreamDocument>> updates = new List<UpdateDefinition<StreamDocument>>();
+			List<UpdateDefinition<StreamDoc>> updates = new List<UpdateDefinition<StreamDoc>>();
 			stream.PausedUntil = newPausedUntil;
 			stream.PauseComment = newPauseComment;
 			updates.Add(updateBuilder.Set(x => x.PausedUntil, newPausedUntil));
@@ -411,34 +495,34 @@ namespace Horde.Build.Streams
 		}
 
 		/// <inheritdoc/>
-		public async Task<IStream?> TryUpdateScheduleTriggerAsync(IStream streamInterface, TemplateRefId templateId, DateTime? lastTriggerTimeUtc, int? lastTriggerChange, List<JobId> newActiveJobs)
+		public async Task<IStream?> TryUpdateScheduleTriggerAsync(IStream streamInterface, TemplateId templateId, DateTime? lastTriggerTimeUtc, int? lastTriggerChange, List<JobId> newActiveJobs)
 		{
-			StreamDocument stream = (StreamDocument)streamInterface;
-			TemplateRef template = stream.Templates[templateId];
-			Schedule schedule = template.Schedule!;
+			StreamDoc stream = (StreamDoc)streamInterface;
+			TemplateRefDoc template = stream.Templates[templateId];
+			TemplateScheduleDoc schedule = template.Schedule!;
 
 			// Build the updates. MongoDB driver cannot parse TemplateRefId in expression tree; need to specify field name explicitly
-			List<UpdateDefinition<StreamDocument>> updates = new List<UpdateDefinition<StreamDocument>>();
-			if (lastTriggerTimeUtc.HasValue && lastTriggerTimeUtc.Value != schedule.LastTriggerTime)
+			List<UpdateDefinition<StreamDoc>> updates = new List<UpdateDefinition<StreamDoc>>();
+			if (lastTriggerTimeUtc.HasValue && lastTriggerTimeUtc.Value != schedule.LastTriggerTimeUtc)
 			{
-				FieldDefinition<StreamDocument, DateTimeOffset> lastTriggerTimeField = $"{nameof(stream.Templates)}.{templateId}.{nameof(template.Schedule)}.{nameof(schedule.LastTriggerTime)}";
-				updates.Add(Builders<StreamDocument>.Update.Set(lastTriggerTimeField, lastTriggerTimeUtc.Value));
+				FieldDefinition<StreamDoc, DateTime> lastTriggerTimeField = $"{nameof(stream.Templates)}.{templateId}.{nameof(template.Schedule)}.{nameof(schedule.LastTriggerTimeUtc)}";
+				updates.Add(Builders<StreamDoc>.Update.Set(lastTriggerTimeField, lastTriggerTimeUtc.Value));
 				schedule.LastTriggerTimeUtc = lastTriggerTimeUtc.Value;
 			}
 			if (lastTriggerChange.HasValue && lastTriggerChange.Value > schedule.LastTriggerChange)
 			{
-				FieldDefinition<StreamDocument, int> lastTriggerChangeField = $"{nameof(stream.Templates)}.{templateId}.{nameof(template.Schedule)}.{nameof(schedule.LastTriggerChange)}";
-				updates.Add(Builders<StreamDocument>.Update.Set(lastTriggerChangeField, lastTriggerChange.Value));
+				FieldDefinition<StreamDoc, int> lastTriggerChangeField = $"{nameof(stream.Templates)}.{templateId}.{nameof(template.Schedule)}.{nameof(schedule.LastTriggerChange)}";
+				updates.Add(Builders<StreamDoc>.Update.Set(lastTriggerChangeField, lastTriggerChange.Value));
 				schedule.LastTriggerChange = lastTriggerChange.Value;
 			}
 			if (newActiveJobs != null)
 			{
-				FieldDefinition<StreamDocument, List<JobId>> field = $"{nameof(stream.Templates)}.{templateId}.{nameof(template.Schedule)}.{nameof(schedule.ActiveJobs)}";
-				updates.Add(Builders<StreamDocument>.Update.Set(field, newActiveJobs));
+				FieldDefinition<StreamDoc, List<JobId>> field = $"{nameof(stream.Templates)}.{templateId}.{nameof(template.Schedule)}.{nameof(schedule.ActiveJobs)}";
+				updates.Add(Builders<StreamDoc>.Update.Set(field, newActiveJobs));
 				schedule.ActiveJobs = newActiveJobs;
 			}
 
-			return (updates.Count == 0)? streamInterface : await TryUpdateStreamAsync(stream, Builders<StreamDocument>.Update.Combine(updates));
+			return (updates.Count == 0)? streamInterface : await TryUpdateStreamAsync(stream, Builders<StreamDoc>.Update.Combine(updates));
 		}
 
 		/// <summary>
@@ -447,14 +531,14 @@ namespace Horde.Build.Streams
 		/// <param name="stream">The stream to update</param>
 		/// <param name="update">The update definition</param>
 		/// <returns>The updated document, or null the update failed</returns>
-		private async Task<StreamDocument?> TryUpdateStreamAsync(StreamDocument stream, UpdateDefinition<StreamDocument> update)
+		private async Task<StreamDoc?> TryUpdateStreamAsync(StreamDoc stream, UpdateDefinition<StreamDoc> update)
 		{
-			FilterDefinition<StreamDocument> filter = Builders<StreamDocument>.Filter.Expr(x => x.Id == stream.Id && x.UpdateIndex == stream.UpdateIndex);
+			FilterDefinition<StreamDoc> filter = Builders<StreamDoc>.Filter.Expr(x => x.Id == stream.Id && x.UpdateIndex == stream.UpdateIndex);
 			update = update.Set(x => x.UpdateIndex, stream.UpdateIndex + 1);
 
-			FindOneAndUpdateOptions<StreamDocument> options = new FindOneAndUpdateOptions<StreamDocument> { ReturnDocument = ReturnDocument.After };
+			FindOneAndUpdateOptions<StreamDoc> options = new FindOneAndUpdateOptions<StreamDoc> { ReturnDocument = ReturnDocument.After };
 
-			StreamDocument? result = await _streams.FindOneAndUpdateAsync(filter, update, options);
+			StreamDoc? result = await _streams.FindOneAndUpdateAsync(filter, update, options);
 			if(result != null)
 			{
 				await PostLoadAsync(result);
@@ -465,25 +549,41 @@ namespace Horde.Build.Streams
 		/// <inheritdoc/>
 		public async Task DeleteAsync(StreamId streamId)
 		{
-			await _streams.UpdateOneAsync<StreamDocument>(x => x.Id == streamId, Builders<StreamDocument>.Update.Set(x => x.Deleted, true).Inc(x => x.UpdateIndex, 1));
+			await _streams.UpdateOneAsync<StreamDoc>(x => x.Id == streamId, Builders<StreamDoc>.Update.Set(x => x.Deleted, true).Inc(x => x.UpdateIndex, 1));
 		}
 
 		/// <summary>
 		/// Checks the stream definition for consistency
 		/// </summary>
-		public static void Validate(StreamId streamId, IReadOnlyDictionary<TemplateRefId, TemplateRef> templates, StreamConfig config)
+		public static void Validate(StreamId streamId, StreamConfig config)
 		{
+			HashSet<TemplateId> remainingTemplates = new HashSet<TemplateId>(config.Templates.Select(x => x.Id));
+
 			// Check the default preflight template is valid
 			if (config.DefaultPreflight != null)
 			{
-				if (config.DefaultPreflight.TemplateId != null && !templates.ContainsKey(config.DefaultPreflight.TemplateId.Value))
+				if (config.DefaultPreflight.TemplateId != null && !remainingTemplates.Contains(config.DefaultPreflight.TemplateId.Value))
 				{
 					throw new InvalidStreamException($"Default preflight template was listed as '{config.DefaultPreflight.TemplateId.Value}', but no template was found by that name");
 				}
 			}
 
+			// Check the chained jobs are valid
+			foreach (TemplateRefConfig templateRef in config.Templates)
+			{
+				if (templateRef.ChainedJobs != null)
+				{
+					foreach (ChainedJobTemplateConfig chainedJob in templateRef.ChainedJobs)
+					{
+						if (!remainingTemplates.Contains(chainedJob.TemplateId))
+						{
+							throw new InvalidDataException($"Invalid template ref id '{chainedJob.TemplateId}");
+						}
+					}
+				}
+			}
+
 			// Check that all the templates are referenced by a tab
-			HashSet<TemplateRefId> remainingTemplates = new HashSet<TemplateRefId>(templates.Keys);
 			foreach (JobsTabConfig jobsTab in config.Tabs.OfType<JobsTabConfig>())
 			{
 				if (jobsTab.Templates != null)
diff --git a/Engine/Source/Programs/Horde/Horde.Build/Streams/StreamConfig.cs b/Engine/Source/Programs/Horde/Horde.Build/Streams/StreamConfig.cs
--- a/Engine/Source/Programs/Horde/Horde.Build/Streams/StreamConfig.cs
+++ b/Engine/Source/Programs/Horde/Horde.Build/Streams/StreamConfig.cs
@@ -10,7 +10,6 @@ using Horde.Build.Acls;
 using Horde.Build.Agents.Pools;
 using Horde.Build.Issues;
 using Horde.Build.Jobs.Graphs;
-using Horde.Build.Jobs.Schedules;
 using Horde.Build.Jobs.Templates;
 using Horde.Build.Perforce;
 using Horde.Build.Server;
@@ -20,7 +19,7 @@ using HordeCommon;
 namespace Horde.Build.Streams
 {
 	using PoolId = StringId<IPool>;
-	using TemplateRefId = StringId<TemplateRef>;
+	using TemplateId = StringId<ITemplateRef>;
 	using WorkflowId = StringId<WorkflowConfig>;
 
 	/// <summary>
@@ -89,7 +88,7 @@ namespace Horde.Build.Streams
 		public string? DefaultPreflightTemplate
 		{
 			get => DefaultPreflight?.TemplateId?.ToString();
-			set => DefaultPreflight = (value == null)? null : new DefaultPreflightConfig { TemplateId = new TemplateRefId(value) };
+			set => DefaultPreflight = (value == null)? null : new DefaultPreflightConfig { TemplateId = new TemplateId(value) };
 		}
 
 		/// <summary>
@@ -163,7 +162,7 @@ namespace Horde.Build.Streams
 		/// <param name="templateRefId"></param>
 		/// <param name="templateConfig"></param>
 		/// <returns></returns>
-		public bool TryGetTemplate(TemplateRefId templateRefId, [NotNullWhen(true)] out TemplateRefConfig? templateConfig)
+		public bool TryGetTemplate(TemplateId templateRefId, [NotNullWhen(true)] out TemplateRefConfig? templateConfig)
 		{
 			templateConfig = Templates.FirstOrDefault(x => x.Id == templateRefId);
 			return templateConfig != null;
@@ -215,7 +214,7 @@ namespace Horde.Build.Streams
 		/// <summary>
 		/// List of job template names to show on this page.
 		/// </summary>
-		public List<TemplateRefId>? Templates { get; set; }
+		public List<TemplateId>? Templates { get; set; }
 
 		/// <summary>
 		/// Columns to display for different types of aggregates
@@ -375,7 +374,7 @@ namespace Horde.Build.Streams
 		/// <summary>
 		/// The template id to query
 		/// </summary>
-		public TemplateRefId? TemplateId { get; set; }
+		public TemplateId? TemplateId { get; set; }
 
 		/// <summary>
 		/// The target to query
@@ -396,7 +395,7 @@ namespace Horde.Build.Streams
 		/// <summary>
 		/// The template id to query
 		/// </summary>
-		public TemplateRefId? TemplateId { get; set; }
+		public TemplateId? TemplateId { get; set; }
 
 		/// <summary>
 		/// Query for the change to use
@@ -419,7 +418,7 @@ namespace Horde.Build.Streams
 		/// Id of the template to trigger
 		/// </summary>
 		[Required]
-		public string TemplateId { get; set; } = String.Empty;
+		public TemplateId TemplateId { get; set; }
 	}
 
 	/// <summary>
@@ -427,14 +426,14 @@ namespace Horde.Build.Streams
 	/// </summary>
 	public class TemplateRefConfig : TemplateConfig
 	{
-		TemplateRefId _id;
+		TemplateId _id;
 
 		/// <summary>
 		/// Optional identifier for this ref. If not specified, an id will be generated from the name.
 		/// </summary>
-		public TemplateRefId Id
+		public TemplateId Id
 		{
-			get => _id.IsEmpty ? TemplateRefId.Sanitize(Name) : _id;
+			get => _id.IsEmpty ? TemplateId.Sanitize(Name) : _id;
 			set => _id = value;
 		}
 
@@ -501,7 +500,7 @@ namespace Horde.Build.Streams
 		/// <summary>
 		/// Days of the week to run this schedule on. If null, the schedule will run every day.
 		/// </summary>
-		public List<string>? DaysOfWeek { get; set; }
+		public List<DayOfWeek>? DaysOfWeek { get; set; }
 
 		/// <summary>
 		/// Time during the day for the first schedule to trigger. Measured in minutes from midnight.
@@ -519,12 +518,72 @@ namespace Horde.Build.Streams
 		public int? Interval { get; set; }
 
 		/// <summary>
-		/// Constructs a model object from this request
+		/// Constructor
 		/// </summary>
-		/// <returns>Model object</returns>
-		public SchedulePattern ToModel()
+		public SchedulePatternConfig()
 		{
-			return new SchedulePattern(DaysOfWeek?.ConvertAll(x => Enum.Parse<DayOfWeek>(x)), MinTime, MaxTime, Interval);
+		}
+
+		/// <summary>
+		/// Constructor
+		/// </summary>
+		/// <param name="daysOfWeek">Which days of the week the schedule should run</param>
+		/// <param name="minTime">Time during the day for the first schedule to trigger. Measured in minutes from midnight.</param>
+		/// <param name="maxTime">Time during the day for the last schedule to trigger. Measured in minutes from midnight.</param>
+		/// <param name="interval">Interval between each schedule triggering</param>
+		public SchedulePatternConfig(List<DayOfWeek>? daysOfWeek, int minTime, int? maxTime, int? interval)
+		{
+			DaysOfWeek = daysOfWeek;
+			MinTime = minTime;
+			MaxTime = maxTime;
+			Interval = interval;
+		}
+
+		/// <summary>
+		/// Calculates the trigger index based on the given time in minutes
+		/// </summary>
+		/// <param name="lastTimeUtc">Time for the last trigger</param>
+		/// <param name="timeZone">The timezone for running the schedule</param>
+		/// <returns>Index of the trigger</returns>
+		public DateTime GetNextTriggerTimeUtc(DateTime lastTimeUtc, TimeZoneInfo timeZone)
+		{
+			// Convert last time into the correct timezone for running the scheule
+			DateTimeOffset lastTime = TimeZoneInfo.ConvertTime((DateTimeOffset)lastTimeUtc, timeZone);
+
+			// Get the base time (ie. the start of this day) for anchoring the schedule
+			DateTimeOffset baseTime = new DateTimeOffset(lastTime.Year, lastTime.Month, lastTime.Day, 0, 0, 0, lastTime.Offset);
+			for (; ; )
+			{
+				if (DaysOfWeek == null || DaysOfWeek.Contains(baseTime.DayOfWeek))
+				{
+					// Get the last time in minutes from the start of this day
+					int lastTimeMinutes = (int)(lastTime - baseTime).TotalMinutes;
+
+					// Get the time of the first trigger of this day. If the last time is less than this, this is the next trigger.
+					if (lastTimeMinutes < MinTime)
+					{
+						return baseTime.AddMinutes(MinTime).UtcDateTime;
+					}
+
+					// Otherwise, get the time for the last trigger in the day.
+					if (Interval.HasValue && Interval.Value > 0)
+					{
+						int actualMaxTime = MaxTime ?? ((24 * 60) - 1);
+						if (lastTimeMinutes < actualMaxTime)
+						{
+							int lastIndex = (lastTimeMinutes - MinTime) / Interval.Value;
+							int nextIndex = lastIndex + 1;
+
+							int nextTimeMinutes = MinTime + (nextIndex * Interval.Value);
+							if (nextTimeMinutes <= actualMaxTime)
+							{
+								return baseTime.AddMinutes(nextTimeMinutes).UtcDateTime;
+							}
+						}
+					}
+				}
+				baseTime = baseTime.AddDays(1.0);
+			}
 		}
 	}
 
@@ -537,22 +596,13 @@ namespace Horde.Build.Streams
 		/// The template containing the dependency
 		/// </summary>
 		[Required]
-		public string TemplateId { get; set; } = String.Empty;
+		public TemplateId TemplateId { get; set; }
 
 		/// <summary>
 		/// Target to wait for
 		/// </summary>
 		[Required]
 		public string Target { get; set; } = String.Empty;
-
-		/// <summary>
-		/// Constructs a model object
-		/// </summary>
-		/// <returns>New model object.</returns>
-		public ScheduleGate ToModel()
-		{
-			return new ScheduleGate(new TemplateRefId(TemplateId), Target);
-		}
 	}
 
 	/// <summary>
@@ -568,7 +618,7 @@ namespace Horde.Build.Streams
 		/// <summary>
 		/// Whether the schedule should be enabled
 		/// </summary>
-		public bool Enabled { get; set; }
+		public bool Enabled { get; set; } = true;
 
 		/// <summary>
 		/// Maximum number of builds that can be active at once
@@ -611,13 +661,45 @@ namespace Horde.Build.Streams
 		public List<SchedulePatternConfig> Patterns { get; set; } = new List<SchedulePatternConfig>();
 
 		/// <summary>
-		/// Constructs a model object
+		/// Gets the flags to filter changes by
 		/// </summary>
-		/// <param name="currentTimeUtc">The current time</param>
-		/// <returns>New model object</returns>
-		public Schedule ToModel(DateTime currentTimeUtc)
+		/// <returns>Set of filter flags</returns>
+		public ChangeContentFlags? GetFilterFlags()
 		{
-			return new Schedule(currentTimeUtc, Enabled, MaxActive, MaxChanges, RequireSubmittedChange, Gate?.ToModel(), Filter, Files, TemplateParameters, Patterns.ConvertAll(x => x.ToModel()));
+			if (Filter == null)
+			{
+				return null;
+			}
+
+			ChangeContentFlags flags = 0;
+			foreach (ChangeContentFlags flag in Filter)
+			{
+				flags |= flag;
+			}
+			return flags;
+		}
+
+		/// <summary>
+		/// Get the next time that the schedule will trigger
+		/// </summary>
+		/// <param name="lastTimeUtc">Last time at which the schedule triggered</param>
+		/// <param name="timeZone">Timezone to evaluate the trigger</param>
+		/// <returns>Next time at which the schedule will trigger</returns>
+		public DateTime? GetNextTriggerTimeUtc(DateTime lastTimeUtc, TimeZoneInfo timeZone)
+		{
+			DateTime? nextTriggerTimeUtc = null;
+			if (Enabled)
+			{
+				foreach (SchedulePatternConfig pattern in Patterns)
+				{
+					DateTime patternTriggerTime = pattern.GetNextTriggerTimeUtc(lastTimeUtc, timeZone);
+					if (nextTriggerTimeUtc == null || patternTriggerTime < nextTriggerTimeUtc)
+					{
+						nextTriggerTimeUtc = patternTriggerTime;
+					}
+				}
+			}
+			return nextTriggerTimeUtc;
 		}
 	}
 
diff --git a/Engine/Source/Programs/Horde/Horde.Build/Streams/StreamMessages.cs b/Engine/Source/Programs/Horde/Horde.Build/Streams/StreamMessages.cs
--- a/Engine/Source/Programs/Horde/Horde.Build/Streams/StreamMessages.cs
+++ b/Engine/Source/Programs/Horde/Horde.Build/Streams/StreamMessages.cs
@@ -5,18 +5,15 @@ using System.Collections.Generic;
 using System.ComponentModel.DataAnnotations;
 using EpicGames.Core;
 using Horde.Build.Acls;
-using Horde.Build.Jobs.Graphs;
 using Horde.Build.Jobs.Templates;
 using Horde.Build.Jobs.Schedules;
 using Horde.Build.Users;
 using Horde.Build.Utilities;
-using HordeCommon;
 using Horde.Build.Issues;
 
 namespace Horde.Build.Streams
 {
-	using TemplateRefId = StringId<TemplateRef>;
-	using WorkflowId = StringId<WorkflowConfig>;
+	using TemplateId = StringId<ITemplateRef>;
 
 	/// <summary>
 	/// Step state update request
@@ -163,32 +160,6 @@ namespace Horde.Build.Streams
 		}
 	}
 
-	/// <summary>
-	/// Trigger for another template
-	/// </summary>
-	public class GetChainedJobTemplateResponse
-	{
-		/// <summary>
-		/// Name of the target that needs to complete before starting the other template
-		/// </summary>
-		public string Trigger { get; set; }
-
-		/// <summary>
-		/// Id of the template to trigger
-		/// </summary>
-		public string TemplateId { get; set; }
-
-		/// <summary>
-		/// Constructor
-		/// </summary>
-		/// <param name="trigger">The trigger definition</param>
-		public GetChainedJobTemplateResponse(ChainedJobTemplate trigger)
-		{
-			Trigger = trigger.Trigger;
-			TemplateId = trigger.TemplateRefId.ToString();
-		}
-	}
-
 	/// <summary>
 	/// State information for a step in the stream
 	/// </summary>
@@ -219,7 +190,7 @@ namespace Horde.Build.Streams
 		/// <summary>
 		/// Constructor
 		/// </summary>
-		public GetTemplateStepStateResponse(TemplateStepState state, GetThinUserInfoResponse? pausedByUserInfo)
+		public GetTemplateStepStateResponse(ITemplateStep state, GetThinUserInfoResponse? pausedByUserInfo)
 		{
 			Name = state.Name;			
 			PauseTimeUtc = state.PauseTimeUtc;
@@ -275,7 +246,7 @@ namespace Horde.Build.Streams
 		/// <summary>
 		/// List of templates to trigger
 		/// </summary>
-		public List<GetChainedJobTemplateResponse>? ChainedJobs { get; set; }
+		public List<ChainedJobTemplateConfig>? ChainedJobs { get; set; }
 
 		/// <summary>
 		/// List of step states
@@ -295,17 +266,17 @@ namespace Horde.Build.Streams
 		/// <param name="template">The template instance</param>
 		/// <param name="stepStates">The template step states</param>
 		/// <param name="bIncludeAcl">Whether to include the ACL in the response</param>
-		public GetTemplateRefResponse(TemplateRefId id, TemplateRef templateRef, ITemplate template, List<GetTemplateStepStateResponse>? stepStates, bool bIncludeAcl)
+		public GetTemplateRefResponse(TemplateId id, ITemplateRef templateRef, ITemplate template, List<GetTemplateStepStateResponse>? stepStates, bool bIncludeAcl)
 			: base(template)
 		{
 			Id = id.ToString();
 			Hash = templateRef.Hash.ToString();
-			ShowUgsBadges = templateRef.ShowUgsBadges;
-			ShowUgsAlerts = templateRef.ShowUgsAlerts;
-			NotificationChannel = templateRef.NotificationChannel;
-			NotificationChannelFilter = templateRef.NotificationChannelFilter;
+			ShowUgsBadges = templateRef.Config.ShowUgsBadges;
+			ShowUgsAlerts = templateRef.Config.ShowUgsAlerts;
+			NotificationChannel = templateRef.Config.NotificationChannel;
+			NotificationChannelFilter = templateRef.Config.NotificationChannelFilter;
 			Schedule = (templateRef.Schedule != null) ? new GetScheduleResponse(templateRef.Schedule) : null;
-			ChainedJobs = (templateRef.ChainedJobs != null && templateRef.ChainedJobs.Count > 0) ? templateRef.ChainedJobs.ConvertAll(x => new GetChainedJobTemplateResponse(x)) : null;
+			ChainedJobs = (templateRef.Config.ChainedJobs != null && templateRef.Config.ChainedJobs.Count > 0) ? templateRef.Config.ChainedJobs : null;
 			StepStates = stepStates;
 			Acl = (bIncludeAcl && templateRef.Acl != null)? new GetAclResponse(templateRef.Acl) : null;
 		}
diff --git a/Engine/Source/Programs/Horde/Horde.Build/Streams/StreamService.cs b/Engine/Source/Programs/Horde/Horde.Build/Streams/StreamService.cs
--- a/Engine/Source/Programs/Horde/Horde.Build/Streams/StreamService.cs
+++ b/Engine/Source/Programs/Horde/Horde.Build/Streams/StreamService.cs
@@ -17,7 +17,7 @@ namespace Horde.Build.Streams
 	using JobId = ObjectId<IJob>;
 	using ProjectId = StringId<IProject>;
 	using StreamId = StringId<IStream>;
-	using TemplateRefId = StringId<TemplateRef>;
+	using TemplateId = StringId<ITemplateRef>;
 
 	/// <summary>
 	/// Cache of information about stream ACLs
@@ -114,12 +114,12 @@ namespace Horde.Build.Streams
 		/// <param name="addJobs">Jobs to add</param>
 		/// <param name="removeJobs">Jobs to remove</param>
 		/// <returns>True if the stream was updated</returns>
-		public async Task<IStream?> UpdateScheduleTriggerAsync(IStream stream, TemplateRefId templateRefId, DateTime? lastTriggerTimeUtc = null, int? lastTriggerChange = null, List<JobId>? addJobs = null, List<JobId>? removeJobs = null)
+		public async Task<IStream?> UpdateScheduleTriggerAsync(IStream stream, TemplateId templateRefId, DateTime? lastTriggerTimeUtc = null, int? lastTriggerChange = null, List<JobId>? addJobs = null, List<JobId>? removeJobs = null)
 		{
 			IStream? newStream = stream;
 			while (newStream != null)
 			{
-				TemplateRef? templateRef;
+				ITemplateRef? templateRef;
 				if (!newStream.Templates.TryGetValue(templateRefId, out templateRef))
 				{
 					break;
@@ -158,7 +158,7 @@ namespace Horde.Build.Streams
 		/// <param name="templateRefId"></param>
 		/// <param name="stepStates"></param>
 		/// <returns></returns>
-		public Task<IStream?> TryUpdateTemplateRefAsync(IStream streamInterface, TemplateRefId templateRefId, List<UpdateStepStateRequest>? stepStates = null)
+		public Task<IStream?> TryUpdateTemplateRefAsync(IStream streamInterface, TemplateId templateRefId, List<UpdateStepStateRequest>? stepStates = null)
 		{
 			return _streams.TryUpdateTemplateRefAsync(streamInterface, templateRefId, stepStates );
 		}
@@ -293,7 +293,7 @@ namespace Horde.Build.Streams
 		/// <param name="user">The principal to authorize</param>
 		/// <param name="cache">Cache of project permissions</param>
 		/// <returns>True if the action is authorized</returns>
-		public Task<bool> AuthorizeAsync(IStream stream, TemplateRef template, AclAction action, ClaimsPrincipal user, ProjectPermissionsCache? cache)
+		public Task<bool> AuthorizeAsync(IStream stream, ITemplateRef template, AclAction action, ClaimsPrincipal user, ProjectPermissionsCache? cache)
 		{
 			bool? result = template.Acl?.Authorize(action, user);
 			if (result == null)
diff --git a/Engine/Source/Programs/Horde/Horde.Build/Streams/StreamsController.cs b/Engine/Source/Programs/Horde/Horde.Build/Streams/StreamsController.cs
--- a/Engine/Source/Programs/Horde/Horde.Build/Streams/StreamsController.cs
+++ b/Engine/Source/Programs/Horde/Horde.Build/Streams/StreamsController.cs
@@ -21,7 +21,7 @@ namespace Horde.Build.Streams
 {
 	using ProjectId = StringId<IProject>;
 	using StreamId = StringId<IStream>;
-	using TemplateRefId = StringId<TemplateRef>;
+	using TemplateId = StringId<ITemplateRef>;
 
 	/// <summary>
 	/// Controller for the /api/v1/streams endpoint
@@ -118,10 +118,10 @@ namespace Horde.Build.Streams
 			bool bIncludeAcl = stream.Acl != null && await _streamService.AuthorizeAsync(stream, AclAction.ViewPermissions, User, cache);
 
 			List<GetTemplateRefResponse> apiTemplateRefs = new List<GetTemplateRefResponse>();
-			foreach (KeyValuePair<TemplateRefId, TemplateRef> pair in stream.Templates)
+			foreach (KeyValuePair<TemplateId, ITemplateRef> pair in stream.Templates)
 			{
 				using IScope templateScope = GlobalTracer.Instance.BuildSpan("CreateGetStreamResponse.Template").StartActive();
-				templateScope.Span.SetTag("templateName", pair.Value.Name);
+				templateScope.Span.SetTag("templateName", pair.Value.Config.Name);
 				
 				if (await _streamService.AuthorizeAsync(stream, pair.Value, AclAction.ViewTemplate, User, cache))
 				{
@@ -129,19 +129,14 @@ namespace Horde.Build.Streams
 					
 					if (template != null)						
 					{
-						TemplateRef tref = pair.Value;
+						ITemplateRef tref = pair.Value;
 
 						List<GetTemplateStepStateResponse>? stepStates = null;
 						if (tref.StepStates != null)
 						{
 							for (int i = 0; i < tref.StepStates.Count; i++)
 							{
-								TemplateStepState state = tref.StepStates[i];
-
-								if (state.PausedByUserId == null)
-								{
-									continue;
-								}
+								ITemplateStep state = tref.StepStates[i];
 
 								stepStates ??= new List<GetTemplateStepStateResponse>();
 
@@ -253,7 +248,7 @@ namespace Horde.Build.Streams
 				return Forbid(AclAction.ViewJob, streamId);
 			}
 
-			TemplateRefId templateIdValue = new TemplateRefId(templateId);
+			TemplateId templateIdValue = new TemplateId(templateId);
 
 			List<IJobStepRef> steps = await _jobStepRefCollection.GetStepsForNodeAsync(streamId, templateIdValue, step, change, true, count);
 			return steps.ConvertAll(x => PropertyFilter.Apply(new GetJobStepRefResponse(x), filter));
@@ -292,7 +287,7 @@ namespace Horde.Build.Streams
 		[HttpGet]
 		[Route("/api/v1/streams/{streamId}/templates/{templateId}")]
 		[ProducesResponseType(typeof(List<GetTemplateResponse>), 200)]
-		public async Task<ActionResult<object>> GetTemplateAsync(StreamId streamId, TemplateRefId templateId, [FromQuery] PropertyFilter? filter = null)
+		public async Task<ActionResult<object>> GetTemplateAsync(StreamId streamId, TemplateId templateId, [FromQuery] PropertyFilter? filter = null)
 		{
 			IStream? stream = await _streamService.GetStreamAsync(streamId);
 			if (stream == null)
@@ -300,7 +295,7 @@ namespace Horde.Build.Streams
 				return NotFound(streamId);
 			}
 
-			TemplateRef? templateRef;
+			ITemplateRef? templateRef;
 			if (!stream.Templates.TryGetValue(templateId, out templateRef))
 			{
 				return NotFound(streamId, templateId);
@@ -325,7 +320,7 @@ namespace Horde.Build.Streams
 		[HttpPut]
 		[Authorize]
 		[Route("/api/v1/streams/{streamId}/templates/{templateRefId}")]
-		public async Task<ActionResult> UpdateStreamTemplateRefAsync(StreamId streamId, TemplateRefId templateRefId, [FromBody] UpdateTemplateRefRequest update)
+		public async Task<ActionResult> UpdateStreamTemplateRefAsync(StreamId streamId, TemplateId templateRefId, [FromBody] UpdateTemplateRefRequest update)
 		{
 			if (!await _aclService.AuthorizeAsync(AclAction.AdminWrite, User))
 			{
diff --git a/Engine/Source/Programs/Horde/Horde.Build/Utilities/HordeControllerBase.cs b/Engine/Source/Programs/Horde/Horde.Build/Utilities/HordeControllerBase.cs
--- a/Engine/Source/Programs/Horde/Horde.Build/Utilities/HordeControllerBase.cs
+++ b/Engine/Source/Programs/Horde/Horde.Build/Utilities/HordeControllerBase.cs
@@ -20,7 +20,7 @@ namespace Horde.Build.Utilities
 	using StreamId = StringId<IStream>;
 	using ToolId = StringId<Tool>;
 	using ToolDeploymentId = ObjectId<ToolDeployment>;
-	using TemplateRefId = StringId<TemplateRef>;
+	using TemplateId = StringId<ITemplateRef>;
 
 	/// <summary>
 	/// Base class for Horde controllers
@@ -108,7 +108,7 @@ namespace Horde.Build.Utilities
 		/// Returns a 403 (forbidden) response with the given action and object
 		/// </summary>
 		[NonAction]
-		protected ActionResult Forbid(AclAction action, TemplateRefId templateId)
+		protected ActionResult Forbid(AclAction action, TemplateId templateId)
 		{
 			return Forbid(action, "template {TemplateId}", templateId);
 		}
@@ -207,7 +207,7 @@ namespace Horde.Build.Utilities
 		/// Returns a 404 response for the given object
 		/// </summary>
 		[NonAction]
-		protected ActionResult NotFound(StreamId streamId, TemplateRefId templateId)
+		protected ActionResult NotFound(StreamId streamId, TemplateId templateId)
 		{
 			return NotFound("Template {TemplateId} not found on stream {StreamId}", templateId, streamId);
 		}
