Horde: Abstract Perforce operations behind an ICommitSource interface, exposed as the "Commits" property on stream objects.

diff --git a/Engine/Source/Programs/Horde/Horde.Build.Tests/Stubs/Services/PerforceServiceStub.cs b/Engine/Source/Programs/Horde/Horde.Build.Tests/Stubs/Services/PerforceServiceStub.cs
--- a/Engine/Source/Programs/Horde/Horde.Build.Tests/Stubs/Services/PerforceServiceStub.cs
+++ b/Engine/Source/Programs/Horde/Horde.Build.Tests/Stubs/Services/PerforceServiceStub.cs
@@ -103,7 +103,7 @@ namespace Horde.Build.Tests.Stubs.Services
 			streamChanges.Add(number, new Commit(streamId, number, originalNumber, author.Id, owner.Id, description, files.ToList()));
 		}
 
-		public Task<List<ICommit>> GetChangesAsync(IStream stream, int? minChange, int? maxChange, int numResults, CancellationToken cancellationToken)
+		public Task<List<ICommit>> GetChangesAsync(IStream stream, int? minChange, int? maxChange, int? numResults, CancellationToken cancellationToken)
 		{
 			List<ICommit> results = new List<ICommit>();
 
@@ -120,7 +120,7 @@ namespace Horde.Build.Tests.Stubs.Services
 					{
 						results.Add(details);
 					}
-					if (numResults > 0 && results.Count >= numResults)
+					if (numResults != null && numResults > 0 && results.Count >= numResults)
 					{
 						break;
 					}
diff --git a/Engine/Source/Programs/Horde/Horde.Build/Issues/IssueService.cs b/Engine/Source/Programs/Horde/Horde.Build/Issues/IssueService.cs
--- a/Engine/Source/Programs/Horde/Horde.Build/Issues/IssueService.cs
+++ b/Engine/Source/Programs/Horde/Horde.Build/Issues/IssueService.cs
@@ -151,7 +151,6 @@ namespace Horde.Build.Issues
 		readonly IIssueCollection _issueCollection;
 		readonly StreamService _streams;
 		readonly IUserCollection _userCollection;
-		readonly IPerforceService _perforce;
 		readonly ILogFileService _logFileService;
 		readonly IClock _clock;
 		readonly ITicker _ticker;
@@ -204,7 +203,7 @@ namespace Horde.Build.Issues
 		/// <summary>
 		/// Constructor
 		/// </summary>
-		public IssueService(IIssueCollection issueCollection, IJobStepRefCollection jobStepRefs, StreamService streams, IUserCollection userCollection, IPerforceService perforce, ILogFileService logFileService, IClock clock, ILogger<IssueService> logger)
+		public IssueService(IIssueCollection issueCollection, IJobStepRefCollection jobStepRefs, StreamService streams, IUserCollection userCollection, ILogFileService logFileService, IClock clock, ILogger<IssueService> logger)
 		{
 			Type[] issueTypes = Assembly.GetExecutingAssembly().GetTypes().Where(x => !x.IsAbstract && typeof(IIssue).IsAssignableFrom(x)).ToArray();
 			foreach (Type issueType in issueTypes)
@@ -217,7 +216,6 @@ namespace Horde.Build.Issues
 			_jobStepRefs = jobStepRefs;
 			_streams = streams;
 			_userCollection = userCollection;
-			_perforce = perforce;
 			_logFileService = logFileService;
 			_clock = clock;
 			_ticker = clock.AddTicker<IssueService>(TimeSpan.FromMinutes(1.0), TickAsync, logger);
@@ -1085,8 +1083,8 @@ namespace Horde.Build.Issues
 				if (stream != null)
 				{
 					_logger.LogInformation("Querying fix changelist {FixChange} in {StreamId}", fixChange, streamId);
-					List<ICommit> changes = await _perforce.GetChangesAsync(stream, fixChange, fixChange, 1);
-					bContainsFixChange = changes.Count > 0;
+					ICommit? change = await stream.Commits.FindCommitsAsync(fixChange, fixChange, 1).FirstOrDefaultAsync();
+					bContainsFixChange = change != null;
 				}
 				cachedContainsFixChange[(streamId, fixChange)] = bContainsFixChange;
 			}
@@ -1109,7 +1107,7 @@ namespace Horde.Build.Issues
 				_logger.LogDebug("Querying for changes in {StreamName} between {MinChange} and {MaxChange}", stream.Name, minChange, maxChange);
 
 				// Get the submitted changes before this job
-				List<ICommit> changes = await PerforceServiceExtensions.GetChangeDetailsAsync(_perforce, stream, minChange, maxChange, MaxChanges);
+				List<ICommit> changes = await stream.Commits.FindCommitsAsync(minChange, maxChange, MaxChanges).ToListAsync();
 				_logger.LogDebug("Found {NumResults} changes", changes.Count);
 
 				// Get the handler to rank them
diff --git a/Engine/Source/Programs/Horde/Horde.Build/Jobs/Schedules/ScheduleService.cs b/Engine/Source/Programs/Horde/Horde.Build/Jobs/Schedules/ScheduleService.cs
--- a/Engine/Source/Programs/Horde/Horde.Build/Jobs/Schedules/ScheduleService.cs
+++ b/Engine/Source/Programs/Horde/Horde.Build/Jobs/Schedules/ScheduleService.cs
@@ -63,15 +63,13 @@ namespace Horde.Build.Jobs.Schedules
 
 		class PerforceHistory
 		{
-			readonly IPerforceService _perforce;
 			readonly IStream _stream;
 			int _maxResults;
 			readonly List<ICommit> _changes = new List<ICommit>();
 			int _nextIndex;
 
-			public PerforceHistory(IPerforceService perforce, IStream stream)
+			public PerforceHistory(IStream stream)
 			{
-				_perforce = perforce;
 				_stream = stream;
 				_maxResults = 10;
 			}
@@ -80,7 +78,7 @@ namespace Horde.Build.Jobs.Schedules
 			{
 				while (index >= _changes.Count)
 				{
-					List<ICommit> newChanges = await _perforce.GetChangesAsync(_stream, null, null, _maxResults);
+					List<ICommit> newChanges = await _stream.Commits.FindCommitsAsync(null, null, _maxResults).ToListAsync();
 
 					int numResults = newChanges.Count;
 					if (_changes.Count > 0)
@@ -93,7 +91,7 @@ namespace Horde.Build.Jobs.Schedules
 					}
 					if(newChanges.Count > 0)
 					{
-						_changes.AddRange((await _perforce.GetChangeDetailsAsync(_stream, newChanges.ConvertAll(x => x.Number))).OrderByDescending(x => x.Number));
+						_changes.AddRange(newChanges);
 					}
 					_maxResults += 10;
 				}
@@ -429,7 +427,7 @@ namespace Horde.Build.Jobs.Schedules
 			}
 
 			// Cache the Perforce history as we're iterating through changes to improve query performance
-			PerforceHistory history = new PerforceHistory(_perforce, stream);
+			PerforceHistory history = new PerforceHistory(stream);
 			
 			// Start as many jobs as possible
 			List<(int Change, int CodeChange)> triggerChanges = new List<(int, int)>();
diff --git a/Engine/Source/Programs/Horde/Horde.Build/Perforce/ICommitSource.cs b/Engine/Source/Programs/Horde/Horde.Build/Perforce/ICommitSource.cs
new file mode 100644
--- /dev/null
+++ b/Engine/Source/Programs/Horde/Horde.Build/Perforce/ICommitSource.cs
@@ -0,0 +1,58 @@
+// Copyright Epic Games, Inc. All Rights Reserved.
+
+using System;
+using System.Collections.Generic;
+using System.Linq;
+using System.Threading;
+using System.Threading.Tasks;
+
+namespace Horde.Build.Perforce
+{
+	/// <summary>
+	/// VCS abstraction. Provides information about commits to a particular stream.
+	/// </summary>
+	public interface ICommitSource
+	{
+		/// <summary>
+		/// Gets the latest change for a particular stream
+		/// </summary>
+		/// <param name="changeNumber">Change numbers to query</param>
+		/// <param name="cancellationToken">Cancellation token for the operation</param>
+		/// <returns>Commit details</returns>
+		Task<ICommit> GetCommitAsync(int changeNumber, CancellationToken cancellationToken = default);
+
+		/// <summary>
+		/// Finds changes submitted to a stream, in reverse order.
+		/// </summary>
+		/// <param name="minChange">The minimum changelist number</param>
+		/// <param name="maxChange">The maximum changelist number</param>
+		/// <param name="maxResults">Maximum number of results to return</param>
+		/// <param name="cancellationToken">Cancellation token for the operation</param>
+		/// <returns>Changelist information</returns>
+		IAsyncEnumerable<ICommit> FindCommitsAsync(int? minChange, int? maxChange, int? maxResults, CancellationToken cancellationToken = default);
+
+		/// <summary>
+		/// Gets the latest change number
+		/// </summary>
+		/// <param name="cancellationToken">Cancellation token for the operation</param>
+		/// <returns></returns>
+		Task<int> LatestNumberAsync(CancellationToken cancellationToken = default);
+	}
+
+	/// <summary>
+	/// Extension methods for <see cref="ICommitSource"/>
+	/// </summary>
+	public static class CommitSourceExtensions
+	{
+		/// <summary>
+		/// Finds the latest commit from a source
+		/// </summary>
+		/// <param name="source">The commit source to query</param>
+		/// <param name="cancellationToken">Cancellation token for the operation</param>
+		/// <returns>The latest commit</returns>
+		public static async Task<ICommit> LatestAsync(this ICommitSource source, CancellationToken cancellationToken)
+		{
+			return await source.FindCommitsAsync(null, null, 1, cancellationToken).FirstAsync(cancellationToken);
+		}
+	}
+}
diff --git a/Engine/Source/Programs/Horde/Horde.Build/Perforce/IPerforceService.cs b/Engine/Source/Programs/Horde/Horde.Build/Perforce/IPerforceService.cs
--- a/Engine/Source/Programs/Horde/Horde.Build/Perforce/IPerforceService.cs
+++ b/Engine/Source/Programs/Horde/Horde.Build/Perforce/IPerforceService.cs
@@ -95,7 +95,7 @@ namespace Horde.Build.Perforce
 		/// <param name="results">Number of results to return</param>
 		/// <param name="cancellationToken">Cancellation token for the operation</param>
 		/// <returns>Latest changelist number</returns>
-		public Task<List<ICommit>> GetChangesAsync(IStream stream, int? minChange, int? maxChange, int results, CancellationToken cancellationToken = default);
+		public Task<List<ICommit>> GetChangesAsync(IStream stream, int? minChange, int? maxChange, int? results, CancellationToken cancellationToken = default);
 
 		/// <summary>
 		/// Checks a shelf is valid for the given stream
diff --git a/Engine/Source/Programs/Horde/Horde.Build/Perforce/PerforceService.cs b/Engine/Source/Programs/Horde/Horde.Build/Perforce/PerforceService.cs
--- a/Engine/Source/Programs/Horde/Horde.Build/Perforce/PerforceService.cs
+++ b/Engine/Source/Programs/Horde/Horde.Build/Perforce/PerforceService.cs
@@ -7,6 +7,7 @@ using System.Globalization;
 using System.IO;
 using System.Linq;
 using System.Net;
+using System.Runtime.CompilerServices;
 using System.Runtime.InteropServices;
 using System.Text;
 using System.Text.RegularExpressions;
@@ -595,13 +596,13 @@ namespace Horde.Build.Perforce
 		}
 
 		/// <inheritdoc/>
-		public async Task<List<ICommit>> GetChangesAsync(IStream stream, int? minChange, int? maxChange, int maxResults, CancellationToken cancellationToken)
+		public async Task<List<ICommit>> GetChangesAsync(IStream stream, int? minChange, int? maxChange, int? maxResults, CancellationToken cancellationToken)
 		{
 			using IScope scope = GlobalTracer.Instance.BuildSpan("PerforceService.GetChangesAsync").StartActive();
 			scope.Span.SetTag("ClusterName", stream.Config.ClusterName);
 			scope.Span.SetTag("MinChange", minChange ?? -1);
 			scope.Span.SetTag("MaxChange", maxChange ?? -1);
-			scope.Span.SetTag("MaxResults", maxResults);
+			scope.Span.SetTag("MaxResults", maxResults ?? -1);
 
 			List<ICommit> results = new List<ICommit>();
 			using (PooledConnectionHandle perforce = await ConnectWithStreamClientAsync(stream, null, cancellationToken))
@@ -610,7 +611,7 @@ namespace Horde.Build.Perforce
 
 				InfoRecord info = await perforce.GetInfoAsync(cancellationToken);
 
-				List<ChangesRecord> changes = await perforce.GetChangesAsync(ChangesOptions.IncludeTimes | ChangesOptions.LongOutput, maxResults, ChangeStatus.Submitted, filter, cancellationToken);
+				List<ChangesRecord> changes = await perforce.GetChangesAsync(ChangesOptions.IncludeTimes | ChangesOptions.LongOutput, maxResults ?? -1, ChangeStatus.Submitted, filter, cancellationToken);
 
 				foreach (ChangesRecord change in changes)
 				{
@@ -1140,4 +1141,35 @@ namespace Horde.Build.Perforce
 			return false;
 		}
 	}
+
+	class PerforceCommitSource : ICommitSource
+	{
+		readonly IPerforceService _perforceService;
+		readonly IStream _stream;
+
+		public PerforceCommitSource(IPerforceService perforceService, IStream stream)
+		{
+			_perforceService = perforceService;
+			_stream = stream;
+		}
+
+		public async IAsyncEnumerable<ICommit> FindCommitsAsync(int? minChange, int? maxChange, int? maxResults, [EnumeratorCancellation] CancellationToken cancellationToken = default)
+		{
+			List<ICommit> commits = await _perforceService.GetChangesAsync(_stream, minChange, maxChange, maxResults, cancellationToken);
+			foreach(ICommit commit in commits)
+			{
+				yield return commit;
+			}
+		}
+
+		public Task<ICommit> GetCommitAsync(int changeNumber, CancellationToken cancellationToken = default)
+		{
+			return _perforceService.GetChangeDetailsAsync(_stream, changeNumber, cancellationToken);
+		}
+
+		public Task<int> LatestNumberAsync(CancellationToken cancellationToken = default)
+		{
+			return _perforceService.GetLatestChangeAsync(_stream);
+		}
+	}
 }
diff --git a/Engine/Source/Programs/Horde/Horde.Build/Streams/IStream.cs b/Engine/Source/Programs/Horde/Horde.Build/Streams/IStream.cs
--- a/Engine/Source/Programs/Horde/Horde.Build/Streams/IStream.cs
+++ b/Engine/Source/Programs/Horde/Horde.Build/Streams/IStream.cs
@@ -9,6 +9,7 @@ using EpicGames.Core;
 using Horde.Build.Acls;
 using Horde.Build.Agents;
 using Horde.Build.Jobs;
+using Horde.Build.Perforce;
 using Horde.Build.Projects;
 using Horde.Build.Users;
 using Horde.Build.Utilities;
@@ -96,6 +97,11 @@ namespace Horde.Build.Streams
 		/// The ACL for this object
 		/// </summary>
 		Acl? Acl { get; }
+
+		/// <summary>
+		/// Accessor for commits of this stream
+		/// </summary>
+		ICommitSource Commits { get; }
 	}
 
 	/// <summary>
diff --git a/Engine/Source/Programs/Horde/Horde.Build/Streams/StreamCollection.cs b/Engine/Source/Programs/Horde/Horde.Build/Streams/StreamCollection.cs
--- a/Engine/Source/Programs/Horde/Horde.Build/Streams/StreamCollection.cs
+++ b/Engine/Source/Programs/Horde/Horde.Build/Streams/StreamCollection.cs
@@ -12,6 +12,7 @@ using Horde.Build.Configuration;
 using Horde.Build.Jobs;
 using Horde.Build.Jobs.Schedules;
 using Horde.Build.Jobs.Templates;
+using Horde.Build.Perforce;
 using Horde.Build.Projects;
 using Horde.Build.Server;
 using Horde.Build.Users;
@@ -50,9 +51,6 @@ namespace Horde.Build.Streams
 
 			public string ConfigRevision { get; set; } = String.Empty;
 
-			[BsonIgnore]
-			public StreamConfig? Config { get; set; }
-
 			public Dictionary<TemplateId, TemplateRefDoc> Templates { get; set; } = new Dictionary<TemplateId, TemplateRefDoc>();
 			public DateTime? PausedUntil { get; set; }
 			public string? PauseComment { get; set; }
@@ -61,9 +59,43 @@ namespace Horde.Build.Streams
 			public int UpdateIndex { get; set; }
 			public bool Deleted { get; set; }
 
-			StreamConfig IStream.Config => Config!;
+			[BsonConstructor]
+			private StreamDoc()
+			{
+			}
+
+			public StreamDoc(StreamId id, ProjectId projectId)
+			{
+				Id = id;
+				ProjectId = projectId;
+			}
+
+			#region IStream Implementation
+
+			[BsonIgnore]
+			public string Name => Config?.Name!;
+
+			[BsonIgnore]
+			public StreamConfig Config { get; private set; } = null!;
+
+			[BsonIgnore]
+			IPerforceService? _perforceService;
+
+			[BsonIgnore]
+			PerforceCommitSource? _cachedCommitSource;
 
+			[BsonIgnore]
 			IReadOnlyDictionary<TemplateId, ITemplateRef>? _cachedTemplates;
+
+			ICommitSource IStream.Commits
+			{
+				get
+				{
+					_cachedCommitSource ??= (_perforceService == null)? null : new PerforceCommitSource(_perforceService!, this);
+					return _cachedCommitSource!;
+				}
+			}
+
 			IReadOnlyDictionary<TemplateId, ITemplateRef> IStream.Templates
 			{
 				get
@@ -73,24 +105,38 @@ namespace Horde.Build.Streams
 				}
 			}
 
-			string IStream.Name => Config!.Name;
-
-			[BsonConstructor]
-			private StreamDoc()
+			public async Task PostLoadAsync(ConfigCollection configCollection, IPerforceService perforceService, ILogger logger)
 			{
-			}
+				_perforceService = perforceService;
 
-			public StreamDoc(StreamId id, ProjectId projectId)
-			{
-				Id = id;
-				ProjectId = projectId;
+				try
+				{
+					if (Deleted)
+					{
+						Config = new StreamConfig();
+					}
+					else
+					{
+						Config = await configCollection.GetConfigAsync<StreamConfig>(ConfigRevision);
+					}
+				}
+				catch (Exception)
+				{
+					logger.LogError("Unable to get stream config for {StreamId} at {Revision}", Id, ConfigRevision);
+					throw;
+				}
+
+				foreach (KeyValuePair<TemplateId, TemplateRefDoc> pair in Templates)
+				{
+					pair.Value.PostLoad(this, pair.Key);
+				}
 			}
+
+			#endregion
 		}
 
 		class TemplateRefDoc : ITemplateRef
 		{
-			#region Serialized
-
 			[BsonRequired]
 			public ContentHash Hash { get; set; } = null!;
 
@@ -103,7 +149,7 @@ namespace Horde.Build.Streams
 			[BsonIgnoreIfNull]
 			public Acl? Acl { get; set; }
 
-			#endregion
+			#region ITemplateRef implementation
 
 			[BsonIgnore]
 			StreamDoc? _owner;
@@ -133,12 +179,12 @@ namespace Horde.Build.Streams
 				Id = id;
 				Schedule?.PostLoad(this);
 			}
+
+			#endregion
 		}
 
 		class TemplateScheduleDoc : ITemplateSchedule
 		{
-			#region Serialized
-
 			public int LastTriggerChange { get; set; }
 
 			[BsonIgnoreIfNull, Obsolete("Use LastTriggerTimeUtc instead")]
@@ -147,7 +193,7 @@ namespace Horde.Build.Streams
 			public DateTime LastTriggerTimeUtc { get; set; }
 			public List<JobId> ActiveJobs { get; set; } = new List<JobId>();
 
-			#endregion
+			#region ITemplateSchedule implementation
 
 			[BsonIgnore]
 			TemplateRefDoc? _owner;
@@ -169,6 +215,8 @@ namespace Horde.Build.Streams
 				}
 #pragma warning restore CS0618 // Type or member is obsolete
 			}
+
+			#endregion
 		}
 
 		class TemplateStepDoc : ITemplateStep
@@ -203,6 +251,7 @@ namespace Horde.Build.Streams
 
 		readonly IMongoCollection<StreamDoc> _streams;
 		readonly ConfigCollection _configCollection;
+		readonly IPerforceService _perforceService;
 		readonly IClock _clock;
 		readonly ITemplateCollection _templateCollection;
 		readonly ILogger<StreamCollection> _logger;
@@ -212,42 +261,21 @@ namespace Horde.Build.Streams
 		/// </summary>
 		/// <param name="mongoService">The database service instance</param>
 		/// <param name="configCollection"></param>
+		/// <param name="perforceService"></param>
 		/// <param name="clock"></param>
 		/// <param name="templateCollection"></param>
 		/// <param name="logger"></param>
-		public StreamCollection(MongoService mongoService, ConfigCollection configCollection, IClock clock, ITemplateCollection templateCollection, ILogger<StreamCollection> logger)
+		public StreamCollection(MongoService mongoService, ConfigCollection configCollection, IPerforceService perforceService, IClock clock, ITemplateCollection templateCollection, ILogger<StreamCollection> logger)
 		{
 			_streams = mongoService.GetCollection<StreamDoc>("Streams");
 			_configCollection = configCollection;
+			_perforceService = perforceService;
 			_clock = clock;
 			_templateCollection = templateCollection;
 			_logger = logger;
 		}
 
-		async Task PostLoadAsync(StreamDoc stream)
-		{
-			try
-			{
-				if (stream.Deleted)
-				{
-					stream.Config = new StreamConfig();
-				}
-				else
-				{
-					stream.Config = await _configCollection.GetConfigAsync<StreamConfig>(stream.ConfigRevision);
-				}
-			}
-			catch (Exception)
-			{
-				_logger.LogError("Unable to get stream config for {StreamId} at {Revision}", stream.Id, stream.ConfigRevision);
-				throw;
-			}
-
-			foreach (KeyValuePair<TemplateId, TemplateRefDoc> pair in stream.Templates)
-			{
-				pair.Value.PostLoad(stream, pair.Key);
-			}
-		}
+		Task PostLoadAsync(StreamDoc stream) => stream.PostLoadAsync(_configCollection, _perforceService, _logger);
 
 		/// <inheritdoc/>
 		public async Task<IStream?> TryCreateOrReplaceAsync(StreamId id, IStream? stream, string revision, ProjectId projectId)
diff --git a/Engine/Source/Programs/Horde/Horde.Build/Streams/StreamsController.cs b/Engine/Source/Programs/Horde/Horde.Build/Streams/StreamsController.cs
--- a/Engine/Source/Programs/Horde/Horde.Build/Streams/StreamsController.cs
+++ b/Engine/Source/Programs/Horde/Horde.Build/Streams/StreamsController.cs
@@ -34,19 +34,18 @@ namespace Horde.Build.Streams
 		private readonly StreamService _streamService;
 		private readonly ITemplateCollection _templateCollection;
 		private readonly IJobStepRefCollection _jobStepRefCollection;
-		private readonly IPerforceService _perforceService;
 		private readonly IUserCollection _userCollection;
 		private readonly AclService _aclService;
+
 		/// <summary>
 		/// Constructor
 		/// </summary>
-		public StreamsController(AclService aclService, StreamService streamService, ITemplateCollection templateCollection, IJobStepRefCollection jobStepRefCollection, IUserCollection userCollection, IPerforceService perforceService)
+		public StreamsController(AclService aclService, StreamService streamService, ITemplateCollection templateCollection, IJobStepRefCollection jobStepRefCollection, IUserCollection userCollection)
 		{
 			_streamService = streamService;
 			_templateCollection = templateCollection;
 			_jobStepRefCollection = jobStepRefCollection;
 			_userCollection = userCollection;
-			_perforceService = perforceService;
 			_aclService = aclService;
 		}
 
@@ -178,7 +177,7 @@ namespace Horde.Build.Streams
 				return Forbid(AclAction.ViewChanges, streamId);
 			}
 
-			List<ICommit> commits = await _perforceService.GetChangesAsync(stream, min, max, results);
+			List<ICommit> commits = await stream.Commits.FindCommitsAsync(min, max, results).ToListAsync();
 
 			List<GetChangeSummaryResponse> responses = new List<GetChangeSummaryResponse>();
 			foreach (ICommit commit in commits)
@@ -211,7 +210,7 @@ namespace Horde.Build.Streams
 				return Forbid(AclAction.ViewChanges, streamId);
 			}
 
-			ICommit? changeDetails = await _perforceService.GetChangeDetailsAsync(stream, changeNumber);
+			ICommit? changeDetails = await stream.Commits.GetCommitAsync(changeNumber);
 			if(changeDetails == null)
 			{
 				return NotFound("CL {Change} not found in stream {StreamId}", changeNumber, streamId);
