Horde: Change build health rows to link on summary instead of entire row
#rnx
#skipci

diff --git a/Engine/Source/Programs/Horde/HordeDashboard/src/components/jobDetailsV2/JobDetailHealthV2.tsx b/Engine/Source/Programs/Horde/HordeDashboard/src/components/jobDetailsV2/JobDetailHealthV2.tsx
--- a/Engine/Source/Programs/Horde/HordeDashboard/src/components/jobDetailsV2/JobDetailHealthV2.tsx
+++ b/Engine/Source/Programs/Horde/HordeDashboard/src/components/jobDetailsV2/JobDetailHealthV2.tsx
@@ -57,7 +57,7 @@ class HealthDataView extends JobDataView {
       if (details.jobData) {
          this.updateReady();
       }
-      
+
       const key = details.jobId + this.stepId + this.labelIdx;
 
       let issueRequests: Promise<IssueData[]>[] = [];
@@ -74,7 +74,7 @@ class HealthDataView extends JobDataView {
          backend.getIssues({ jobId: details.jobId, stepId: this.stepId, label: this.labelIdx, count: 50, resolved: true })];
       } else {
          this.issues = cached.issues;
-      } 
+      }
 
       if (issueRequests.length) {
 
@@ -82,7 +82,7 @@ class HealthDataView extends JobDataView {
          let issues: GetIssueResponse[] = [];
          responses.forEach(r => issues.push(...r));
          this.issues = issues;
-         issueCache.set(key, {issues: issues, time: moment(Date.now())});         
+         issueCache.set(key, { issues: issues, time: moment(Date.now()) });
       }
 
       let jiraKeys: string[] = [];
@@ -117,8 +117,8 @@ class HealthDataView extends JobDataView {
 
       this.initialize(this.issues?.length ? [sideRail] : undefined)
       this.updateReady();
-      
-                 
+
+
    }
 
    clear() {
@@ -135,7 +135,7 @@ class HealthDataView extends JobDataView {
       }
 
       this.updateReady();
-   }   
+   }
 
    issues: GetIssueResponse[] = [];
 
@@ -168,7 +168,7 @@ export const HealthPanel: React.FC<{ jobDetails: JobDetailsV2 }> = observer(({ j
          dataView?.clear();
       };
    }, [dataView]);
-   
+
    dataView.subscribe();
 
    // query changes and immediately initializes
@@ -244,7 +244,11 @@ export const HealthPanel: React.FC<{ jobDetails: JobDetailsV2 }> = observer(({ j
             summary = summary.slice(0, 100) + "...";
          }
 
-         return <Stack horizontal disableShrink={true}>{<IssueStatusIcon issue={issue} streamId={jobDetails.jobData!.streamId} />}<Text variant={textSize} style={{ textDecoration: !!issue.resolvedAt ? "line-through" : undefined }}>{`Issue ${issue.id} - ${summary}`}</Text></Stack>;
+
+         query.set("issue", item.issue.id.toString());
+         const href = `${location.pathname}?${query.toString()}` + window.location.hash;
+
+         return <Link to={href} className="job-item" onClick={() => setIssueHistory(true)}><Stack horizontal disableShrink={true}>{<IssueStatusIcon issue={issue} streamId={jobDetails.jobData!.streamId} />}<Text variant={textSize} style={{ textDecoration: !!issue.resolvedAt ? "line-through" : undefined }}>{`Issue ${issue.id} - ${summary}`}</Text></Stack></Link>;
       }
 
       if (column.name === "Quarantine") {
@@ -348,13 +352,7 @@ export const HealthPanel: React.FC<{ jobDetails: JobDetailsV2 }> = observer(({ j
 
       if (props) {
 
-         const item = props!.item as HealthItem;
-         query.set("issue", item.issue.id.toString());
-         const href = `${location.pathname}?${query.toString()}`+ window.location.hash;
-
-         return <Link to={href} className="job-item" onClick={() => setIssueHistory(true)}>
-            <DetailsRow className={classes.detailsRow} {...props} />
-         </Link>;
+         return <DetailsRow className={classes.detailsRow} {...props} />
 
       }
       return null;
