Horde: Fix content issues not identifying keys correctly when engine outputs relative paths.

diff --git a/Engine/Source/Programs/Horde/Horde.Build.Tests/IssueServiceTests.cs b/Engine/Source/Programs/Horde/Horde.Build.Tests/IssueServiceTests.cs
--- a/Engine/Source/Programs/Horde/Horde.Build.Tests/IssueServiceTests.cs
+++ b/Engine/Source/Programs/Horde/Horde.Build.Tests/IssueServiceTests.cs
@@ -1173,6 +1173,37 @@ namespace Horde.Build.Tests
 			}
 		}
 
+		[TestMethod]
+		public async Task ContentIssueTest()
+		{
+			IJob job1 = CreateJob(_mainStreamId, 120, "Cook Test", _graph);
+			await ParseEventsAsync(job1, 0, 0, new[] 
+			{
+				// Note: using relative paths here, which can't be mapped to depot paths
+				@"LogBlueprint: Warning: [AssetLog] ..\..\..\QAGame\Plugins\NiagaraFluids\Content\Blueprints\Phsyarum_BP.uasset: [Compiler] Fill Texture 2D : Usage of 'Fill Texture 2D' has been deprecated. This function has been replaced by object user variables on the emitter to specify render targets to fill with data." 
+			});
+			await UpdateCompleteStep(job1, 0, 0, JobStepOutcome.Failure);
+
+			List<IIssue> issues1 = await IssueCollection.FindIssuesAsync();
+			Assert.AreEqual(1, issues1.Count);
+			Assert.AreEqual("Warnings in Phsyarum_BP.uasset", issues1[0].Summary);
+
+			IJob job2 = CreateJob(_mainStreamId, 125, "Cook Test", _graph);
+			await ParseEventsAsync(job2, 0, 0, new[] 
+			{
+				// Add a new warning; should create a new issue
+				@"LogBlueprint: Warning: [AssetLog] ..\..\..\QAGame\Plugins\NiagaraFluids\Content\Blueprints\Phsyarum_BP.uasset: [Compiler] Fill Texture 2D : Usage of 'Fill Texture 2D' has been deprecated. This function has been replaced by object user variables on the emitter to specify render targets to fill with data.",
+				@"LogBlueprint: Warning: [AssetLog] ..\..\..\QAGame\Plugins\NiagaraFluids\Content\Blueprints\Phsyarum_BP2.uasset: [Compiler] Fill Texture 2D : Usage of 'Fill Texture 2D' has been deprecated. This function has been replaced by object user variables on the emitter to specify render targets to fill with data.", 
+			});
+			await UpdateCompleteStep(job2, 0, 0, JobStepOutcome.Failure);
+
+			List<IIssue> issues2 = await IssueCollection.FindIssuesAsync();
+			issues2.SortBy(x => x.Id);
+			Assert.AreEqual(2, issues2.Count);
+			Assert.AreEqual("Warnings in Phsyarum_BP.uasset", issues2[0].Summary);
+			Assert.AreEqual("Warnings in Phsyarum_BP2.uasset", issues2[1].Summary);
+		}
+
 		[TestMethod]
 		public async Task HashedIssueTest()
 		{
diff --git a/Engine/Source/Programs/Horde/Horde.Build/Issues/Handlers/ContentIssueHandler.cs b/Engine/Source/Programs/Horde/Horde.Build/Issues/Handlers/ContentIssueHandler.cs
--- a/Engine/Source/Programs/Horde/Horde.Build/Issues/Handlers/ContentIssueHandler.cs
+++ b/Engine/Source/Programs/Horde/Horde.Build/Issues/Handlers/ContentIssueHandler.cs
@@ -54,7 +54,7 @@ namespace Horde.Build.Issues.Handlers
 			foreach (ILogEventLine line in eventData.Lines)
 			{
 				string? relativePath;
-				if (line.Data.TryGetNestedProperty("properties.asset.relativePath", out relativePath))
+				if (line.Data.TryGetNestedProperty("properties.asset.relativePath", out relativePath) || line.Data.TryGetNestedProperty("properties.asset.$text", out relativePath))
 				{
 					int endIdx = relativePath.LastIndexOfAny(new char[] { '/', '\\' }) + 1;
 					string fileName = relativePath.Substring(endIdx);
