Horde: Use string identifiers for singletons.

diff --git a/Engine/Source/Programs/Horde/Horde.Build/Agents/ConformList.cs b/Engine/Source/Programs/Horde/Horde.Build/Agents/ConformList.cs
--- a/Engine/Source/Programs/Horde/Horde.Build/Agents/ConformList.cs
+++ b/Engine/Source/Programs/Horde/Horde.Build/Agents/ConformList.cs
@@ -60,7 +60,7 @@ namespace Horde.Build.Agents
 	/// <summary>
 	/// List of machines that are currently conforming
 	/// </summary>
-	[SingletonDocument("60afc737f0d2a70754229300")]
+	[SingletonDocument("conform-list", "60afc737f0d2a70754229300")]
 	public class ConformList : SingletonBase
 	{
 		/// <summary>
diff --git a/Engine/Source/Programs/Horde/Horde.Build/Agents/Software/AgentSoftwareService.cs b/Engine/Source/Programs/Horde/Horde.Build/Agents/Software/AgentSoftwareService.cs
--- a/Engine/Source/Programs/Horde/Horde.Build/Agents/Software/AgentSoftwareService.cs
+++ b/Engine/Source/Programs/Horde/Horde.Build/Agents/Software/AgentSoftwareService.cs
@@ -65,7 +65,7 @@ namespace Horde.Build.Agents.Software
 	/// <summary>
 	/// Singleton document used to track different versions of the agent software
 	/// </summary>
-	[SingletonDocument("5f455039d97900f2b6c735a9")]
+	[SingletonDocument("agent-software-channels", "5f455039d97900f2b6c735a9")]
 	public class AgentSoftwareChannels : SingletonBase
 	{
 		/// <summary>
diff --git a/Engine/Source/Programs/Horde/Horde.Build/Devices/DeviceService.cs b/Engine/Source/Programs/Horde/Horde.Build/Devices/DeviceService.cs
--- a/Engine/Source/Programs/Horde/Horde.Build/Devices/DeviceService.cs
+++ b/Engine/Source/Programs/Horde/Horde.Build/Devices/DeviceService.cs
@@ -68,7 +68,7 @@ namespace Horde.Build.Devices
 	/// <summary>
 	/// Platform map required by V1 API
 	/// </summary>
-	[SingletonDocument("6165a2e26fd5f104e31e6862")]
+	[SingletonDocument("device-platform-map", "6165a2e26fd5f104e31e6862")]
 	public class DevicePlatformMapV1 : SingletonBase
 	{
 		/// <summary>
diff --git a/Engine/Source/Programs/Horde/Horde.Build/Issues/IssueCollection.cs b/Engine/Source/Programs/Horde/Horde.Build/Issues/IssueCollection.cs
--- a/Engine/Source/Programs/Horde/Horde.Build/Issues/IssueCollection.cs
+++ b/Engine/Source/Programs/Horde/Horde.Build/Issues/IssueCollection.cs
@@ -30,7 +30,7 @@ namespace Horde.Build.Issues
 
 	class IssueCollection : IIssueCollection
 	{
-		[SingletonDocument("5e4c226440ce25fa3207a9af")]
+		[SingletonDocument("issue-ledger", "5e4c226440ce25fa3207a9af")]
 		class IssueLedger : SingletonBase
 		{
 			public int NextId { get; set; }
diff --git a/Engine/Source/Programs/Horde/Horde.Build/Issues/IssueReportService.cs b/Engine/Source/Programs/Horde/Horde.Build/Issues/IssueReportService.cs
--- a/Engine/Source/Programs/Horde/Horde.Build/Issues/IssueReportService.cs
+++ b/Engine/Source/Programs/Horde/Horde.Build/Issues/IssueReportService.cs
@@ -66,7 +66,7 @@ namespace Horde.Build.Issues
 
 #pragma warning restore CS1591 // Missing XML comment for publicly visible type or member
 
-	[SingletonDocument("6268871c211d05611b3e4fd8")]
+	[SingletonDocument("issue-report-state", "6268871c211d05611b3e4fd8")]
 	class IssueReportState : SingletonBase
 	{
 		[BsonDictionaryOptions(DictionaryRepresentation.ArrayOfDocuments)]
diff --git a/Engine/Source/Programs/Horde/Horde.Build/Perforce/PerforceLoadBalancer.cs b/Engine/Source/Programs/Horde/Horde.Build/Perforce/PerforceLoadBalancer.cs
--- a/Engine/Source/Programs/Horde/Horde.Build/Perforce/PerforceLoadBalancer.cs
+++ b/Engine/Source/Programs/Horde/Horde.Build/Perforce/PerforceLoadBalancer.cs
@@ -130,7 +130,7 @@ namespace Horde.Build.Perforce
 		/// <summary>
 		/// 
 		/// </summary>
-		[SingletonDocument("6046aec374a9283100967ee7")]
+		[SingletonDocument("perforce-server-list", "6046aec374a9283100967ee7")]
 		class PerforceServerList : SingletonBase
 		{
 			public List<PerforceServerEntry> Servers { get; set; } = new List<PerforceServerEntry>();
diff --git a/Engine/Source/Programs/Horde/Horde.Build/Server/Globals.cs b/Engine/Source/Programs/Horde/Horde.Build/Server/Globals.cs
--- a/Engine/Source/Programs/Horde/Horde.Build/Server/Globals.cs
+++ b/Engine/Source/Programs/Horde/Horde.Build/Server/Globals.cs
@@ -14,6 +14,8 @@ using MongoDB.Bson.Serialization.Attributes;
 
 namespace Horde.Build.Server
 {
+	using SingletonId = StringId<SingletonBase>;
+
 	/// <summary>
 	/// Base class for singleton documents
 	/// </summary>
@@ -23,7 +25,7 @@ namespace Horde.Build.Server
 		/// Unique id for this singleton
 		/// </summary>
 		[BsonId]
-		public ObjectId Id { get; set; }
+		public SingletonId Id { get; set; }
 
 		/// <summary>
 		/// The revision index of this document
@@ -276,14 +278,9 @@ namespace Horde.Build.Server
 	/// <summary>
 	/// Global server settings
 	/// </summary>
-	[SingletonDocument("5e3981cb28b8ec59cd07184a")]
+	[SingletonDocument("globals", "5e3981cb28b8ec59cd07184a")]
 	public class Globals : SingletonBase
 	{
-		/// <summary>
-		/// The unique id for all globals objects
-		/// </summary>
-		public static ObjectId StaticId { get; } = new ObjectId("5e3981cb28b8ec59cd07184a");
-
 		/// <summary>
 		/// Unique instance id of this database
 		/// </summary>
diff --git a/Engine/Source/Programs/Horde/Horde.Build/Server/MongoService.cs b/Engine/Source/Programs/Horde/Horde.Build/Server/MongoService.cs
--- a/Engine/Source/Programs/Horde/Horde.Build/Server/MongoService.cs
+++ b/Engine/Source/Programs/Horde/Horde.Build/Server/MongoService.cs
@@ -9,6 +9,7 @@ using System.Linq;
 using System.Net;
 using System.Net.NetworkInformation;
 using System.Net.Security;
+using System.Reflection;
 using System.Runtime.InteropServices;
 using System.Security.Cryptography.X509Certificates;
 using System.Text;
@@ -34,6 +35,8 @@ using StackExchange.Redis;
 
 namespace Horde.Build.Server
 {
+	using SingletonId = StringId<SingletonBase>;
+
 	/// <summary>
 	/// Constrained set of parameters for building a mongo index
 	/// </summary>
@@ -212,7 +215,12 @@ namespace Horde.Build.Server
 		/// <summary>
 		/// Collection of singleton documents
 		/// </summary>
-		public IMongoCollection<BsonDocument> Singletons { get; }
+		IMongoCollection<BsonDocument> SingletonsV1 { get; }
+
+		/// <summary>
+		/// Collection of singleton documents
+		/// </summary>
+		IMongoCollection<BsonDocument> SingletonsV2 { get; }
 
 		/// <summary>
 		/// Settings for the application
@@ -337,7 +345,8 @@ namespace Horde.Build.Server
 				MongoClient client = new MongoClient(mongoSettings);
 				Database = client.GetDatabase(Settings.DatabaseName);
 
-				Singletons = GetCollection<BsonDocument>("Singletons");
+				SingletonsV1 = GetCollection<BsonDocument>("Singletons");
+				SingletonsV2 = GetCollection<BsonDocument>("SingletonsV2");
 
 				Globals globals = GetGlobalsAsync().Result;
 				while (globals.JwtSigningKey == null)
@@ -771,39 +780,52 @@ namespace Horde.Build.Server
 			}
 		}
 
+		class SingletonInfo<T> where T : SingletonBase
+		{
+			public static readonly SingletonDocumentAttribute Attribute = GetAttribute();
+
+			static SingletonDocumentAttribute GetAttribute()
+			{
+				SingletonDocumentAttribute? attribute = typeof(T).GetCustomAttribute<SingletonDocumentAttribute>();
+				if (attribute == null)
+				{
+					throw new Exception($"Type {typeof(T).Name} is missing a {nameof(SingletonDocumentAttribute)} annotation");
+				}
+				return attribute;
+			}
+		}
+
 		/// <summary>
 		/// Gets a singleton document by id
 		/// </summary>
 		/// <returns>The globals document</returns>
 		public Task<Globals> GetGlobalsAsync()
 		{
-			return GetSingletonAsync(Globals.StaticId, () => new Globals() { SchemaVersion = UpgradeService.LatestSchemaVersion });
+			return GetSingletonAsync<Globals>(() => new Globals() { SchemaVersion = UpgradeService.LatestSchemaVersion });
 		}
 
 		/// <summary>
 		/// Gets a singleton document by id
 		/// </summary>
-		/// <param name="id">Id of the singleton document</param>
 		/// <returns>The document</returns>
-		public async Task<T> GetSingletonAsync<T>(ObjectId id) where T : SingletonBase, new()
+		public Task<T> GetSingletonAsync<T>() where T : SingletonBase, new()
 		{
-			T singleton = await GetSingletonAsync(id, () => new T());
-			singleton.PostLoad();
-			return singleton;
+			return GetSingletonAsync(() => new T());
 		}
 
 		/// <summary>
 		/// Gets a singleton document by id
 		/// </summary>
-		/// <param name="id">Id of the singleton document</param>
 		/// <param name="constructor">Method to use to construct a new object</param>
 		/// <returns>The document</returns>
-		public async Task<T> GetSingletonAsync<T>(ObjectId id, Func<T> constructor) where T : SingletonBase, new()
+		public async Task<T> GetSingletonAsync<T>(Func<T> constructor) where T : SingletonBase, new()
 		{
-			FilterDefinition<BsonDocument> filter = new BsonDocument(new BsonElement("_id", id));
+			SingletonDocumentAttribute Attribute = SingletonInfo<T>.Attribute;
+
+			FilterDefinition<BsonDocument> filter = new BsonDocument(new BsonElement("_id", Attribute.Id));
 			for (; ; )
 			{
-				BsonDocument? document = await Singletons.Find(filter).FirstOrDefaultAsync();
+				BsonDocument? document = await SingletonsV2.Find(filter).FirstOrDefaultAsync();
 				if (document != null)
 				{
 					T item = BsonSerializer.Deserialize<T>(document);
@@ -811,9 +833,21 @@ namespace Horde.Build.Server
 					return item;
 				}
 
-				T newItem = constructor();
-				newItem.Id = id;
-				await Singletons.InsertOneAsync(newItem.ToBsonDocument());
+				T? newItem = null;
+				if (Attribute.LegacyId != null)
+				{
+					BsonDocument? legacyDocument = await SingletonsV1.Find(new BsonDocument(new BsonElement("_id", ObjectId.Parse(Attribute.LegacyId)))).FirstOrDefaultAsync();
+					if (legacyDocument != null)
+					{
+						legacyDocument.Remove("_id");
+						newItem = BsonSerializer.Deserialize<T>(legacyDocument);
+						newItem.PostLoad();
+					}
+				}
+				newItem ??= constructor();
+
+				newItem.Id = new SingletonId(Attribute.Id);
+				await SingletonsV2.InsertOneIgnoreDuplicatesAsync(newItem.ToBsonDocument());
 			}
 		}
 
@@ -821,14 +855,13 @@ namespace Horde.Build.Server
 		/// Updates a singleton
 		/// </summary>
 		/// <typeparam name="T"></typeparam>
-		/// <param name="id"></param>
 		/// <param name="updater"></param>
 		/// <returns></returns>
-		public async Task UpdateSingletonAsync<T>(ObjectId id, Action<T> updater) where T : SingletonBase, new()
+		public async Task UpdateSingletonAsync<T>(Action<T> updater) where T : SingletonBase, new()
 		{
 			for (; ; )
 			{
-				T document = await GetSingletonAsync(id, () => new T());
+				T document = await GetSingletonAsync(() => new T());
 				updater(document);
 
 				if (await TryUpdateSingletonAsync(document))
@@ -847,10 +880,10 @@ namespace Horde.Build.Server
 		{
 			int prevRevision = singletonObject.Revision++;
 
-			BsonDocument filter = new BsonDocument { new BsonElement("_id", singletonObject.Id), new BsonElement(nameof(SingletonBase.Revision), prevRevision) };
+			BsonDocument filter = new BsonDocument { new BsonElement("_id", singletonObject.Id.ToString()), new BsonElement(nameof(SingletonBase.Revision), prevRevision) };
 			try
 			{
-				ReplaceOneResult result = await Singletons.ReplaceOneAsync(filter, singletonObject.ToBsonDocument(), new ReplaceOptions { IsUpsert = true });
+				ReplaceOneResult result = await SingletonsV2.ReplaceOneAsync(filter, singletonObject.ToBsonDocument(), new ReplaceOptions { IsUpsert = true });
 				return result.MatchedCount > 0;
 			}
 			catch (MongoWriteException ex)
@@ -871,7 +904,7 @@ namespace Horde.Build.Server
 	/// <summary>
 	/// Stores the version number of the latest server instance to have upgraded the database schema
 	/// </summary>
-	[SingletonDocument("62470d8508d48eddfac7b55d")]
+	[SingletonDocument("mongo-schema", "62470d8508d48eddfac7b55d")]
 	public class MongoSchemaDocument : SingletonBase
 	{
 		/// <summary>
diff --git a/Engine/Source/Programs/Horde/Horde.Build/Utilities/SingletonDocument.cs b/Engine/Source/Programs/Horde/Horde.Build/Utilities/SingletonDocument.cs
--- a/Engine/Source/Programs/Horde/Horde.Build/Utilities/SingletonDocument.cs
+++ b/Engine/Source/Programs/Horde/Horde.Build/Utilities/SingletonDocument.cs
@@ -8,6 +8,8 @@ using MongoDB.Bson;
 
 namespace Horde.Build.Utilities
 {
+	using SingletonId = StringId<SingletonBase>;
+
 	/// <summary>
 	/// Attribute specifying the unique id for a singleton document
 	/// </summary>
@@ -19,13 +21,20 @@ namespace Horde.Build.Utilities
 		/// </summary>
 		public string Id { get; }
 
+		/// <summary>
+		/// ObjectId for the legacy document
+		/// </summary>
+		public string? LegacyId { get; }
+
 		/// <summary>
 		/// Constructor
 		/// </summary>
-		/// <param name="id">Unique id for the singleton document</param>
-		public SingletonDocumentAttribute(string id)
+		/// <param name="id">Unique id for the singleton document. Should be a valid StringId.</param>
+		/// <param name="legacyId">Legacy identifier for the document</param>
+		public SingletonDocumentAttribute(string id, string? legacyId = null)
 		{
 			Id = id;
+			LegacyId = legacyId;
 		}
 	}
 
@@ -60,11 +69,6 @@ namespace Horde.Build.Utilities
 		/// </summary>
 		readonly MongoService _mongoService;
 
-		/// <summary>
-		/// Unique id for the singleton document
-		/// </summary>
-		readonly ObjectId _objectId;
-
 		/// <summary>
 		/// Constructor
 		/// </summary>
@@ -72,33 +76,12 @@ namespace Horde.Build.Utilities
 		public SingletonDocument(MongoService mongoService)
 		{
 			_mongoService = mongoService;
-
-			SingletonDocumentAttribute? attribute = typeof(T).GetCustomAttribute<SingletonDocumentAttribute>();
-			if (attribute == null)
-			{
-				throw new Exception($"Type {typeof(T).Name} is missing a {nameof(SingletonDocumentAttribute)} annotation");
-			}
-
-			_objectId = new ObjectId(attribute.Id);
-		}
-
-		/// <summary>
-		/// Constructor
-		/// </summary>
-		/// <param name="mongoService">The database service instance</param>
-		/// <param name="objectId">The singleton document object id</param>
-		public SingletonDocument(MongoService mongoService, ObjectId objectId)
-		{
-			_mongoService = mongoService;
-			_objectId = objectId;
 		}
 
 		/// <inheritdoc/>
 		public async Task<T> GetAsync()
 		{
-			T value = await _mongoService.GetSingletonAsync<T>(_objectId);
-			value.Id = _objectId;
-			return value;
+			return await _mongoService.GetSingletonAsync<T>();
 		}
 
 		/// <inheritdoc/>
