Horde: Fix exceptions deserializing paused template steps with a null user id.

diff --git a/Engine/Source/Programs/Horde/Horde.Build/Streams/StreamCollection.cs b/Engine/Source/Programs/Horde/Horde.Build/Streams/StreamCollection.cs
--- a/Engine/Source/Programs/Horde/Horde.Build/Streams/StreamCollection.cs
+++ b/Engine/Source/Programs/Horde/Horde.Build/Streams/StreamCollection.cs
@@ -178,6 +178,11 @@ namespace Horde.Build.Streams
 				_owner = owner;
 				Id = id;
 				Schedule?.PostLoad(this);
+
+				if (StepStates != null)
+				{
+					StepStates.RemoveAll(x => x.PausedByUserId == null);
+				}
 			}
 
 			#endregion
@@ -222,8 +227,11 @@ namespace Horde.Build.Streams
 		class TemplateStepDoc : ITemplateStep
 		{
 			public string Name { get; set; } = String.Empty;
-			public UserId PausedByUserId { get; set; }
-			public DateTime PauseTimeUtc { get; set; }
+			public UserId? PausedByUserId { get; set; }
+			public DateTime? PauseTimeUtc { get; set; }
+
+			UserId ITemplateStep.PausedByUserId => PausedByUserId ?? UserId.Empty;
+			DateTime ITemplateStep.PauseTimeUtc => PauseTimeUtc ?? DateTime.MinValue;
 
 			public TemplateStepDoc()
 			{
