Horde: Include UHT errors in compile issues, allowing Horde to triage the causer correctly.

diff --git a/Engine/Source/Programs/Horde/Horde.Build/Issues/Handlers/CompileIssueHandler.cs b/Engine/Source/Programs/Horde/Horde.Build/Issues/Handlers/CompileIssueHandler.cs
--- a/Engine/Source/Programs/Horde/Horde.Build/Issues/Handlers/CompileIssueHandler.cs
+++ b/Engine/Source/Programs/Horde/Horde.Build/Issues/Handlers/CompileIssueHandler.cs
@@ -39,7 +39,7 @@ namespace Horde.Build.Issues.Handlers
 		/// <returns>True if the given event id matches</returns>
 		public static bool IsMatchingEventId(EventId eventId)
 		{
-			return eventId == KnownLogEvents.Compiler || eventId == KnownLogEvents.AutomationTool_SourceFileLine || eventId == KnownLogEvents.MSBuild;
+			return eventId == KnownLogEvents.Compiler || eventId == KnownLogEvents.UHT || eventId == KnownLogEvents.AutomationTool_SourceFileLine || eventId == KnownLogEvents.MSBuild;
 		}
 
 		static bool IsMaskedEventId(EventId id) => id == KnownLogEvents.Generic || id == KnownLogEvents.ExitCode || id == KnownLogEvents.Systemic_Xge_BuildFailed;
