Horde: Don't include discussion thread in assignment DM when mentions aren't allowed
#rnx

diff --git a/Engine/Source/Programs/Horde/Horde.Build/Notifications/Sinks/SlackNotificationSink.cs b/Engine/Source/Programs/Horde/Horde.Build/Notifications/Sinks/SlackNotificationSink.cs
--- a/Engine/Source/Programs/Horde/Horde.Build/Notifications/Sinks/SlackNotificationSink.cs
+++ b/Engine/Source/Programs/Horde/Horde.Build/Notifications/Sinks/SlackNotificationSink.cs
@@ -1123,7 +1123,7 @@ namespace Horde.Build.Notifications.Sinks
 			if (workflowId != null)
 			{
 				IStream? stream = await _streamService.GetStreamAsync(span.StreamId);
-				if (stream != null && stream.Config.TryGetWorkflow(workflowId.Value, out WorkflowConfig? workflow) && workflow.TriageChannel != null)
+				if (stream != null && stream.Config.TryGetWorkflow(workflowId.Value, out WorkflowConfig? workflow) && workflow.TriageChannel != null && workflow.AllowMentions)
 				{
 					MessageStateDocument? state = await GetMessageStateAsync(workflow.TriageChannel, GetTriageThreadEventId(issue.Id));
 					if (state != null)
