Horde: Use the issue creation time to compute the length of time that an issue has been open. The span containing the first failure will exclude spans which have previously been fixed.

diff --git a/Engine/Source/Programs/Horde/Horde.Build/Notifications/Sinks/SlackNotificationSink.cs b/Engine/Source/Programs/Horde/Horde.Build/Notifications/Sinks/SlackNotificationSink.cs
--- a/Engine/Source/Programs/Horde/Horde.Build/Notifications/Sinks/SlackNotificationSink.cs
+++ b/Engine/Source/Programs/Horde/Horde.Build/Notifications/Sinks/SlackNotificationSink.cs
@@ -2115,7 +2115,7 @@ namespace Horde.Build.Notifications.Sinks
 					return null;
 				}
 
-				TimeSpan openTime = utcNow - span.FirstFailure.StepTime;
+				TimeSpan openTime = utcNow - issue.CreatedAt;
 
 				string openTimeStr;
 				if (openTime < TimeSpan.FromHours(1.0))
