Horde: Add autosubmit info to user home job panel
#rnx
#skipci

diff --git a/Engine/Source/Programs/Horde/HordeDashboard/src/components/UserHome.tsx b/Engine/Source/Programs/Horde/HordeDashboard/src/components/UserHome.tsx
--- a/Engine/Source/Programs/Horde/HordeDashboard/src/components/UserHome.tsx
+++ b/Engine/Source/Programs/Horde/HordeDashboard/src/components/UserHome.tsx
@@ -7,7 +7,7 @@ import moment from "moment-timezone";
 import { default as React, useEffect, useState } from 'react';
 import { Link, useLocation } from "react-router-dom";
 import backend, { useBackend } from "../backend";
-import { GetIssueResponse, GetStepResponse, JobData, JobQuery, JobStepOutcome, LabelData, LabelOutcome, LabelState, ProjectData, StepData, StreamData } from "../backend/Api";
+import { GetIssueResponse, GetStepResponse, JobData, JobQuery, JobState, JobStepOutcome, LabelData, LabelOutcome, LabelState, ProjectData, StepData, StreamData } from "../backend/Api";
 import dashboard from "../backend/Dashboard";
 import graphCache, { GraphQuery } from '../backend/GraphCache';
 import { useWindowSize } from '../base/utilities/hooks';
@@ -577,11 +577,12 @@ const JobsPanel: React.FC<{ includeOtherPreflights: boolean }> = observer(({ inc
    let columns: IColumn[] = [
       { key: 'jobview_column1', name: 'Status', minWidth: 16, maxWidth: 16 },
       { key: 'jobview_column2', name: 'Change', minWidth: 220, maxWidth: 220 },
-      { key: 'jobview_column3', name: 'Labels', minWidth: 440, maxWidth: 440 },
-      { key: 'jobview_column4', name: 'Steps', minWidth: 380, maxWidth: 380 },
-      { key: 'jobview_column5', name: 'StartedBy', minWidth: 140, maxWidth: 140 },
-      { key: 'jobview_column6', name: 'Time', minWidth: 130, maxWidth: 130 },
-      { key: 'jobview_column7', name: 'Dismiss', minWidth: 32, maxWidth: 32 },
+      { key: 'jobview_column3', name: 'Submit', minWidth: 140, maxWidth: 140 },
+      { key: 'jobview_column4', name: 'Labels', minWidth: 420, maxWidth: 420 },
+      { key: 'jobview_column5', name: 'Steps', minWidth: 380, maxWidth: 380 },
+      { key: 'jobview_column6', name: 'StartedBy', minWidth: 140, maxWidth: 140 },
+      { key: 'jobview_column7', name: 'Time', minWidth: 130, maxWidth: 130 },
+      { key: 'jobview_column8', name: 'Dismiss', minWidth: 32, maxWidth: 32 },
 
    ];
 
@@ -774,9 +775,36 @@ const JobsPanel: React.FC<{ includeOtherPreflights: boolean }> = observer(({ inc
 
       }
 
-      if (column!.name === "Job") {
-         return <Stack verticalAlign="center" verticalFill={true} tokens={{ childrenGap: 12 }}>
+      if (column!.name === "Submit") {
 
+         if (!job.autoSubmit) {
+            return null;
+         }
+
+         let message = "";
+         let url = "";
+
+         if (job.state !== JobState.Complete) {
+            message = "Submit Pending";
+         } else {
+
+            if (!job.autoSubmitChange) {
+               message = "Submit Failed";               
+            } else {
+               message = `Submitted - CL ${job.autoSubmitChange})`;
+               if (dashboard.swarmUrl) {
+                  url = `${dashboard.swarmUrl}/change/${job.autoSubmitChange}`;
+               }               
+            }
+         }
+
+         if (!message) {
+            return null;
+         }
+
+         return <Stack horizontalAlign="start" verticalAlign="center" verticalFill={true} tokens={{ childrenGap: 12 }}>
+            {!url && <Text style={{ fontSize: 11 }}>{message}</Text>}
+            {!!url && <a href={url} target="_blank" rel="noreferrer" onClick={ev => ev?.stopPropagation()}><Text variant="small">{message}</Text></a>}
          </Stack>
       }
 
@@ -976,7 +1004,7 @@ const JobsPanel: React.FC<{ includeOtherPreflights: boolean }> = observer(({ inc
                   </Stack>
                   <Stack grow />
                   <Stack>
-                     
+
                      {!!jobItems.length && <PrimaryButton text="Clear All" style={{ fontFamily: "Horde Open Sans SemiBold", fontSize: 10, padding: 8, height: "24px", minWidth: "48px" }} onClick={() => { dashboard.clearPinnedJobs() }} />}
                   </Stack>
                </Stack>
diff --git a/Engine/Source/Programs/Horde/HordeDashboard/src/components/jobDetailsV2/JobDetailViewSummary.tsx b/Engine/Source/Programs/Horde/HordeDashboard/src/components/jobDetailsV2/JobDetailViewSummary.tsx
--- a/Engine/Source/Programs/Horde/HordeDashboard/src/components/jobDetailsV2/JobDetailViewSummary.tsx
+++ b/Engine/Source/Programs/Horde/HordeDashboard/src/components/jobDetailsV2/JobDetailViewSummary.tsx
@@ -68,7 +68,7 @@ export const AutosubmitInfo: React.FC<{ jobDetails: JobDetailsV2 }> = observer((
    }
 
    if (jobData.autoSubmitChange) {
-      const url = `https://p4-swarm.epicgames.net/change/${jobData.autoSubmitChange}`;
+      const url = `${dashboard.swarmUrl}/change/${jobData.autoSubmitChange}`;
       return <Stack style={{ paddingTop: 24 }}><Text>Automatically submitted in <a href={url} target="_blank" rel="noreferrer" onClick={ev => ev?.stopPropagation()}>{`CL ${jobData.autoSubmitChange}`}</a></Text></Stack>;
    }
 
