Horde: Dress up the user home autosubmit column
#rnx
#skipci

diff --git a/Engine/Source/Programs/Horde/HordeDashboard/src/components/UserHome.tsx b/Engine/Source/Programs/Horde/HordeDashboard/src/components/UserHome.tsx
--- a/Engine/Source/Programs/Horde/HordeDashboard/src/components/UserHome.tsx
+++ b/Engine/Source/Programs/Horde/HordeDashboard/src/components/UserHome.tsx
@@ -1,6 +1,6 @@
 // Copyright Epic Games, Inc. All Rights Reserved.
 
-import { CollapseAllVisibility, DefaultButton, DetailsHeader, DetailsList, DetailsListLayoutMode, DetailsRow, FocusZone, FocusZoneDirection, IColumn, IconButton, IDetailsListProps, mergeStyleSets, PrimaryButton, ScrollablePane, ScrollbarVisibility, SelectionMode, Spinner, SpinnerSize, Stack, Text } from '@fluentui/react';
+import { CollapseAllVisibility, DefaultButton, DetailsHeader, DetailsList, DetailsListLayoutMode, DetailsRow, FocusZone, FocusZoneDirection, FontIcon, IColumn, IconButton, IDetailsListProps, mergeStyleSets, PrimaryButton, ScrollablePane, ScrollbarVisibility, SelectionMode, Spinner, SpinnerSize, Stack, Text } from '@fluentui/react';
 import { action, observable } from 'mobx';
 import { observer } from 'mobx-react-lite';
 import moment from "moment-timezone";
@@ -8,7 +8,7 @@ import { default as React, useEffect, useState } from 'react';
 import { Link, useLocation } from "react-router-dom";
 import backend, { useBackend } from "../backend";
 import { GetIssueResponse, GetStepResponse, JobData, JobQuery, JobState, JobStepOutcome, LabelData, LabelOutcome, LabelState, ProjectData, StepData, StreamData } from "../backend/Api";
-import dashboard from "../backend/Dashboard";
+import dashboard, { StatusColor } from "../backend/Dashboard";
 import graphCache, { GraphQuery } from '../backend/GraphCache';
 import { useWindowSize } from '../base/utilities/hooks';
 import { displayTimeZone, getElapsedString, getShortNiceTime } from '../base/utilities/timeUtils';
@@ -776,7 +776,7 @@ const JobsPanel: React.FC<{ includeOtherPreflights: boolean }> = observer(({ inc
       }
 
       if (column!.name === "Submit") {
-
+         
          if (!job.autoSubmit) {
             return null;
          }
@@ -784,17 +784,23 @@ const JobsPanel: React.FC<{ includeOtherPreflights: boolean }> = observer(({ inc
          let message = "";
          let url = "";
 
+         const colors = dashboard.getStatusColors();
+         let color = colors.get(StatusColor.Ready);
+
          if (job.state !== JobState.Complete) {
             message = "Submit Pending";
+            color = colors.get(StatusColor.Running);
          } else {
 
             if (!job.autoSubmitChange) {
-               message = "Submit Failed";               
+               message = "Submit Failed";
+               color = colors.get(StatusColor.Failure);
             } else {
-               message = `Submitted - CL ${job.autoSubmitChange})`;
+               color = colors.get(StatusColor.Success);
+               message = `Submitted\nCL ${job.autoSubmitChange}`;
                if (dashboard.swarmUrl) {
                   url = `${dashboard.swarmUrl}/change/${job.autoSubmitChange}`;
-               }               
+               }
             }
          }
 
@@ -802,9 +808,15 @@ const JobsPanel: React.FC<{ includeOtherPreflights: boolean }> = observer(({ inc
             return null;
          }
 
+         const style = { fontSize: 13, color: color, paddingTop: 2 };
+
+         const icon = <Stack className="horde-no-darktheme"><FontIcon style={style} iconName="Square" /></Stack>;
+
          return <Stack horizontalAlign="start" verticalAlign="center" verticalFill={true} tokens={{ childrenGap: 12 }}>
-            {!url && <Text style={{ fontSize: 11 }}>{message}</Text>}
-            {!!url && <a href={url} target="_blank" rel="noreferrer" onClick={ev => ev?.stopPropagation()}><Text variant="small">{message}</Text></a>}
+            <Stack horizontal verticalAlign='center' verticalFill={true} tokens={{ childrenGap: 6 }}>
+               {!url && <Stack horizontal verticalAlign='center' verticalFill={true} tokens={{ childrenGap: 8 }}>{icon}<Stack><Text style={{ fontSize: 11 }}>{message}</Text></Stack></Stack>}
+               {!!url && <Stack horizontal verticalAlign='center' verticalFill={true} tokens={{ childrenGap: 8 }}>{icon}<Stack style={{textAlign: "left"}}><a href={url} target="_blank" rel="noreferrer" onClick={ev => ev?.stopPropagation()}><Text variant="tiny" style={{whiteSpace: "pre" }}>{message}</Text></a></Stack></Stack>}
+            </Stack>
          </Stack>
       }
 
