Horde Add autosubmit to user home job query filter
#skipci
#rnx

diff --git a/Engine/Source/Programs/Horde/HordeDashboard/src/components/UserHome.tsx b/Engine/Source/Programs/Horde/HordeDashboard/src/components/UserHome.tsx
--- a/Engine/Source/Programs/Horde/HordeDashboard/src/components/UserHome.tsx
+++ b/Engine/Source/Programs/Horde/HordeDashboard/src/components/UserHome.tsx
@@ -405,7 +405,7 @@ class UserJobsHandler {
       const maxJobs = 100;
 
       try {
-         let filter = "id,streamId,name,change,preflightChange,templateId,templateHash,graphHash,startedByUserInfo,createTime,state,arguments,updateTime,labels,defaultLabel,batches";
+         let filter = "id,streamId,name,change,preflightChange,templateId,templateHash,graphHash,startedByUserInfo,createTime,state,arguments,updateTime,labels,defaultLabel,batches,autoSubmit,autoSubmitChange";
 
          const query: JobQuery = {
             filter: filter,
