Horde Storage - Fixed issue were cleanup of the filesystem would attempt to run even if no filesystem was mounted.
Also disabled last access tracking on worker deployments by default as it provides no value

diff --git a/Engine/Source/Programs/Horde/HordeStorage/Horde.Storage/Implementation/GC/BlobCleanupService.cs b/Engine/Source/Programs/Horde/HordeStorage/Horde.Storage/Implementation/GC/BlobCleanupService.cs
--- a/Engine/Source/Programs/Horde/HordeStorage/Horde.Storage/Implementation/GC/BlobCleanupService.cs
+++ b/Engine/Source/Programs/Horde/HordeStorage/Horde.Storage/Implementation/GC/BlobCleanupService.cs
@@ -33,7 +33,7 @@ namespace Horde.Storage.Implementation
             return _settings.CurrentValue.BlobCleanupServiceEnabled;
         }
 
-        public BlobCleanupService(IServiceProvider provider, IOptionsMonitor<GCSettings> settings) : base(serviceName: nameof(BlobCleanupService), settings.CurrentValue.BlobCleanupPollFrequency, new BlobCleanupState())
+        public BlobCleanupService(IServiceProvider provider, IOptionsMonitor<GCSettings> settings, IOptionsMonitor<FilesystemSettings> filesystemSettings) : base(serviceName: nameof(BlobCleanupService), settings.CurrentValue.BlobCleanupPollFrequency, new BlobCleanupState())
         {
             _settings = settings;
             
@@ -44,7 +44,7 @@ namespace Horde.Storage.Implementation
             }
 
             FileSystemStore? fileSystemStore = provider.GetService<FileSystemStore>();
-            if (fileSystemStore != null)
+            if (fileSystemStore != null && filesystemSettings.CurrentValue.RunCleanup)
             {
                 RegisterCleanup(fileSystemStore);
             }
diff --git a/Engine/Source/Programs/Horde/HordeStorage/Horde.Storage/Settings.cs b/Engine/Source/Programs/Horde/HordeStorage/Horde.Storage/Settings.cs
--- a/Engine/Source/Programs/Horde/HordeStorage/Horde.Storage/Settings.cs
+++ b/Engine/Source/Programs/Horde/HordeStorage/Horde.Storage/Settings.cs
@@ -130,6 +130,7 @@ namespace Horde.Storage
     public class FilesystemSettings
     {
         [Required] public string RootDir { get; set; } = "";
+        public bool RunCleanup { get; set; } = false;
         public ulong MaxSizeBytes { get; set; } = 500 * 1024 * 1024;
         public double TriggerThresholdPercentage { get; set; } = 0.95;
         public double TargetThresholdPercentage { get; set; } = 0.85;
