Horde Storage - Better attempt at  fixing issue were filesystem cleanup runs even for unmounted filesystems. This avoids accessing settings objects that might not be configured.

diff --git a/Engine/Source/Programs/Horde/HordeStorage/Horde.Storage/Implementation/GC/BlobCleanupService.cs b/Engine/Source/Programs/Horde/HordeStorage/Horde.Storage/Implementation/GC/BlobCleanupService.cs
--- a/Engine/Source/Programs/Horde/HordeStorage/Horde.Storage/Implementation/GC/BlobCleanupService.cs
+++ b/Engine/Source/Programs/Horde/HordeStorage/Horde.Storage/Implementation/GC/BlobCleanupService.cs
@@ -33,7 +33,7 @@ namespace Horde.Storage.Implementation
             return _settings.CurrentValue.BlobCleanupServiceEnabled;
         }
 
-        public BlobCleanupService(IServiceProvider provider, IOptionsMonitor<GCSettings> settings, IOptionsMonitor<FilesystemSettings> filesystemSettings) : base(serviceName: nameof(BlobCleanupService), settings.CurrentValue.BlobCleanupPollFrequency, new BlobCleanupState())
+        public BlobCleanupService(IServiceProvider provider, IOptionsMonitor<GCSettings> settings) : base(serviceName: nameof(BlobCleanupService), settings.CurrentValue.BlobCleanupPollFrequency, new BlobCleanupState())
         {
             _settings = settings;
             
@@ -43,10 +43,13 @@ namespace Horde.Storage.Implementation
                 RegisterCleanup(orphanBlobCleanupRefs);
             }
 
-            FileSystemStore? fileSystemStore = provider.GetService<FileSystemStore>();
-            if (fileSystemStore != null && filesystemSettings.CurrentValue.RunCleanup)
+            if (settings.CurrentValue.RunFilesystemCleanup)
             {
-                RegisterCleanup(fileSystemStore);
+                FileSystemStore? fileSystemStore = provider.GetService<FileSystemStore>();
+                if (fileSystemStore != null)
+                {
+                    RegisterCleanup(fileSystemStore);
+                }
             }
         }
 
diff --git a/Engine/Source/Programs/Horde/HordeStorage/Horde.Storage/Settings.cs b/Engine/Source/Programs/Horde/HordeStorage/Horde.Storage/Settings.cs
--- a/Engine/Source/Programs/Horde/HordeStorage/Horde.Storage/Settings.cs
+++ b/Engine/Source/Programs/Horde/HordeStorage/Horde.Storage/Settings.cs
@@ -130,7 +130,6 @@ namespace Horde.Storage
     public class FilesystemSettings
     {
         [Required] public string RootDir { get; set; } = "";
-        public bool RunCleanup { get; set; } = false;
         public ulong MaxSizeBytes { get; set; } = 500 * 1024 * 1024;
         public double TriggerThresholdPercentage { get; set; } = 0.95;
         public double TargetThresholdPercentage { get; set; } = 0.85;
@@ -157,6 +156,7 @@ namespace Horde.Storage
         
         public bool CleanOldRefRecords { get; set; } = false;
         public bool CleanOldBlobs { get; set; } = true;
+        public bool RunFilesystemCleanup { get; set; } = false;
 
         public TimeSpan LastAccessCutoff { get; set; } = TimeSpan.FromDays(14);
 
@@ -164,6 +164,7 @@ namespace Horde.Storage
         public TimeSpan RefCleanupPollFrequency { get; set; } = TimeSpan.FromMinutes(60);
         public int OrphanGCMaxParallelOperations { get; set; } = 8;
         public int OrphanRefMaxParallelOperations { get; set; } = 8;
+
     }
 
     public class UpstreamRelaySettings
