Horde Storage - Fixed issue were blobs would get deleted even if they were referenced by some namespace but not others under a storage pool

diff --git a/Engine/Source/Programs/Horde/HordeStorage/Horde.Storage/Implementation/GC/OrphanBlobCleanupRefs.cs b/Engine/Source/Programs/Horde/HordeStorage/Horde.Storage/Implementation/GC/OrphanBlobCleanupRefs.cs
--- a/Engine/Source/Programs/Horde/HordeStorage/Horde.Storage/Implementation/GC/OrphanBlobCleanupRefs.cs
+++ b/Engine/Source/Programs/Horde/HordeStorage/Horde.Storage/Implementation/GC/OrphanBlobCleanupRefs.cs
@@ -136,7 +136,7 @@ namespace Horde.Storage.Implementation
 
                 if (blobIndex == null)
                 {
-                    break;
+                    continue;
                 }
 
                 foreach ((BucketId, IoHashKey) tuple in blobIndex.References!)
