Horde: Remove unused scopes from Okta-based auth

diff --git a/Engine/Source/Programs/Horde/Horde.Build/ServerSettings.cs b/Engine/Source/Programs/Horde/Horde.Build/ServerSettings.cs
--- a/Engine/Source/Programs/Horde/Horde.Build/ServerSettings.cs
+++ b/Engine/Source/Programs/Horde/Horde.Build/ServerSettings.cs
@@ -194,6 +194,10 @@ namespace Horde.Build
 	/// </summary>
 	public class FeatureFlagSettings
 	{
+		/// <summary>
+		/// Use new auth config for custom auth servers in Okta
+		/// </summary>
+		public bool AuthSettingsV2 { get; set; } = false;
 	}
 
 	/// <summary>
@@ -321,7 +325,7 @@ namespace Horde.Build
 		public string? JwtSecret { get; set; } = null!;
 
 		/// <summary>
-		/// Length of time before JWT tokens expire, in hourse
+		/// Length of time before JWT tokens expire, in hours
 		/// </summary>
 		public int JwtExpiryTimeHours { get; set; } = 4;
 
diff --git a/Engine/Source/Programs/Horde/Horde.Build/Startup.cs b/Engine/Source/Programs/Horde/Horde.Build/Startup.cs
--- a/Engine/Source/Programs/Horde/Horde.Build/Startup.cs
+++ b/Engine/Source/Programs/Horde/Horde.Build/Startup.cs
@@ -554,6 +554,11 @@ namespace Horde.Build
 							options.Authority = settings.OidcAuthority;
 							options.ClientId = settings.OidcClientId;
 
+							if (settings.FeatureFlags.AuthSettingsV2)
+							{
+								options.Scope.Remove("groups");								
+							}
+
 							if (!String.IsNullOrEmpty(settings.OidcSigninRedirect))
 							{
 								options.Events = new OpenIdConnectEvents
