Horde: Add commit tags to stream configuration, allowing commits modifying files matching various wildcards to be tagged as different types (eg. code, content, modifies project X, etc...).

diff --git a/Engine/Source/Programs/Horde/Horde.Build/Jobs/Schedules/ScheduleService.cs b/Engine/Source/Programs/Horde/Horde.Build/Jobs/Schedules/ScheduleService.cs
--- a/Engine/Source/Programs/Horde/Horde.Build/Jobs/Schedules/ScheduleService.cs
+++ b/Engine/Source/Programs/Horde/Horde.Build/Jobs/Schedules/ScheduleService.cs
@@ -126,8 +126,8 @@ namespace Horde.Build.Jobs.Schedules
 						return 0;
 					}
 
-					ChangeContentFlags contentFlags = await details.GetContentFlagsAsync(CancellationToken.None);
-					if ((contentFlags & ChangeContentFlags.ContainsCode) != 0)
+					IReadOnlyList<CommitTag> commitTags = await details.GetTagsAsync(CancellationToken.None);
+					if (commitTags.Contains(CommitTag.Code))
 					{
 						return details.Number;
 					}
@@ -421,7 +421,6 @@ namespace Horde.Build.Jobs.Schedules
 
 			// Create a timer to limit the amount we look back through P4 history
 			Stopwatch timer = Stopwatch.StartNew();
-			ChangeContentFlags? filterFlags = schedule.Config.GetFilterFlags();
 
 			// Create a file filter
 			FileFilter? fileFilter = null;
@@ -466,7 +465,7 @@ namespace Horde.Build.Jobs.Schedules
 
 				// Adjust the changelist for the desired filter
 				int change = changeDetails.Number;
-				if (await ShouldBuildChangeAsync(changeDetails, filterFlags, fileFilter))
+				if (await ShouldBuildChangeAsync(changeDetails, schedule.Config.Commits, fileFilter))
 				{
 					int codeChange = await history.GetCodeChange(change);
 					if (codeChange == -1)
@@ -534,21 +533,21 @@ namespace Horde.Build.Jobs.Schedules
 		/// Tests whether a schedule should build a particular change, based on its requested change filters
 		/// </summary>
 		/// <param name="details">The change details</param>
-		/// <param name="filterFlags"></param>
+		/// <param name="filterTags"></param>
 		/// <param name="fileFilter">Filter for the files to trigger a build</param>
 		/// <returns></returns>
-		private async ValueTask<bool> ShouldBuildChangeAsync(ICommit details, ChangeContentFlags? filterFlags, FileFilter? fileFilter)
+		private async ValueTask<bool> ShouldBuildChangeAsync(ICommit details, List<CommitTag>? filterTags, FileFilter? fileFilter)
 		{
 			if (Regex.IsMatch(details.Description, @"^\s*#\s*skipci", RegexOptions.Multiline))
 			{
 				return false;
 			}
-			if (filterFlags != null && filterFlags.Value != 0)
+			if (filterTags != null && filterTags.Count > 0)
 			{
-				ChangeContentFlags contentFlags = await details.GetContentFlagsAsync(CancellationToken.None);
-				if ((contentFlags & filterFlags.Value) == 0)
+				IReadOnlyList<CommitTag> commitTags = await details.GetTagsAsync(CancellationToken.None);
+				if (!commitTags.Any(x => filterTags.Contains(x)))
 				{
-					_logger.LogDebug("Not building change {Change} ({ChangeFlags}) due to filter flags ({FilterFlags})", details.Number, contentFlags.ToString(), filterFlags.Value.ToString());
+					_logger.LogDebug("Not building change {Change} ({ChangeTags}) due to filter tags ({FilterTags})", details.Number, String.Join(", ", commitTags.Select(x => x.ToString())), String.Join(", ", filterTags.Select(x => x.ToString())));
 					return false;
 				}
 			}
diff --git a/Engine/Source/Programs/Horde/Horde.Build/Jobs/Schedules/Schedules.cs b/Engine/Source/Programs/Horde/Horde.Build/Jobs/Schedules/Schedules.cs
--- a/Engine/Source/Programs/Horde/Horde.Build/Jobs/Schedules/Schedules.cs
+++ b/Engine/Source/Programs/Horde/Horde.Build/Jobs/Schedules/Schedules.cs
@@ -40,9 +40,15 @@ namespace Horde.Build.Jobs.Schedules
 		/// </summary>
 		public ScheduleGateConfig? Gate { get; set; }
 
+		/// <summary>
+		/// Which commits to run this job for
+		/// </summary>
+		public List<CommitTag>? Commits { get; set; }
+
 		/// <summary>
 		/// The types of changes to run for
 		/// </summary>
+		[Obsolete("Use Commits instead")]
 		public List<ChangeContentFlags>? Filter { get; set; }
 
 		/// <summary>
@@ -81,7 +87,10 @@ namespace Horde.Build.Jobs.Schedules
 			MaxChanges = schedule.Config.MaxChanges;
 			RequireSubmittedChange = schedule.Config.RequireSubmittedChange;
 			Gate = schedule.Config.Gate;
+			Commits = schedule.Config.Commits;
+#pragma warning disable CS0618 // Type or member is obsolete
 			Filter = schedule.Config.Filter;
+#pragma warning restore CS0618 // Type or member is obsolete
 			TemplateParameters = schedule.Config.TemplateParameters;
 			Patterns = schedule.Config.Patterns;
 			LastTriggerChange = schedule.LastTriggerChange;
diff --git a/Engine/Source/Programs/Horde/Horde.Build/Perforce/Change.cs b/Engine/Source/Programs/Horde/Horde.Build/Perforce/Change.cs
deleted file mode 100644
diff --git a/Engine/Source/Programs/Horde/Horde.Build/Perforce/ChangeMessages.cs b/Engine/Source/Programs/Horde/Horde.Build/Perforce/ChangeMessages.cs
--- a/Engine/Source/Programs/Horde/Horde.Build/Perforce/ChangeMessages.cs
+++ b/Engine/Source/Programs/Horde/Horde.Build/Perforce/ChangeMessages.cs
@@ -30,17 +30,24 @@ namespace Horde.Build.Perforce
 		/// </summary>
 		public string Description { get; set; }
 
+		/// <summary>
+		/// Tags for this commit
+		/// </summary>
+		public List<CommitTag> Tags { get; set; }
+
 		/// <summary>
 		/// Constructor
 		/// </summary>
 		/// <param name="commit">The commit to construct from</param>
 		/// <param name="author">Author of the commit</param>
-		public GetChangeSummaryResponse(ICommit commit, IUser author)
+		/// <param name="tags">Tags for the commit</param>
+		public GetChangeSummaryResponse(ICommit commit, IUser author, IReadOnlyList<CommitTag> tags)
 		{
 			Number = commit.Number;
 			Author = author.Name;
 			AuthorInfo = new GetThinUserInfoResponse(author);
 			Description = commit.Description;
+			Tags = new List<CommitTag>(tags);
 		}
 	}
 
@@ -69,6 +76,11 @@ namespace Horde.Build.Perforce
 		/// </summary>
 		public string Description { get; set; }
 
+		/// <summary>
+		/// Tags for this commit
+		/// </summary>
+		public List<CommitTag> Tags { get; set; }
+
 		/// <summary>
 		/// List of files that were modified, relative to the stream base
 		/// </summary>
@@ -79,13 +91,15 @@ namespace Horde.Build.Perforce
 		/// </summary>
 		/// <param name="commit">The commit to construct from</param>
 		/// <param name="author">Author of the change</param>
+		/// <param name="tags">Tags for the commit</param>
 		/// <param name="files">Files modified by the commit</param>
-		public GetChangeDetailsResponse(ICommit commit, IUser author, IReadOnlyList<string> files)
+		public GetChangeDetailsResponse(ICommit commit, IUser author, IReadOnlyList<CommitTag> tags, IReadOnlyList<string> files)
 		{
 			Number = commit.Number;
 			Author = author.Name;
 			AuthorInfo = new GetThinUserInfoResponse(author);
 			Description = commit.Description;
+			Tags = new List<CommitTag>(tags);
 			Files = new List<string>(files);
 		}
 	}
diff --git a/Engine/Source/Programs/Horde/Horde.Build/Perforce/CommitTag.cs b/Engine/Source/Programs/Horde/Horde.Build/Perforce/CommitTag.cs
new file mode 100644
--- /dev/null
+++ b/Engine/Source/Programs/Horde/Horde.Build/Perforce/CommitTag.cs
@@ -0,0 +1,122 @@
+// Copyright Epic Games, Inc. All Rights Reserved.
+
+using System;
+using System.Collections.Generic;
+using System.Linq;
+using System.Text.Json;
+using System.Text.Json.Serialization;
+using EpicGames.Horde;
+using EpicGames.Perforce;
+using MongoDB.Bson.Serialization;
+using MongoDB.Bson.Serialization.Attributes;
+using MongoDB.Bson.Serialization.Serializers;
+
+namespace Horde.Build.Perforce
+{
+	/// <summary>
+	/// Constants for known commit tags
+	/// </summary>
+	[JsonConverter(typeof(CommitTagJsonConverter))]
+	[BsonSerializer(typeof(CommitTagBsonSerializer))]
+	public readonly struct CommitTag : IEquatable<CommitTag>
+	{
+		/// <summary>
+		/// Predefined filter name for commits containing code
+		/// </summary>
+		public static CommitTag Code { get; } = new CommitTag("code");
+
+		/// <summary>
+		/// Predefined filter for code
+		/// </summary>
+		public static IReadOnlyList<string> CodeFilter { get; } = PerforceUtils.CodeExtensions.Select(x => $"*{x}").ToList();
+
+		/// <summary>
+		/// Predefined filter name for commits containing content
+		/// </summary>
+		public static CommitTag Content { get; } = new CommitTag("content");
+
+		/// <summary>
+		/// Predefined filter name for content
+		/// </summary>
+		public static IReadOnlyList<string> ContentFilter { get; } = Enumerable.Concat(new[] { "*" }, PerforceUtils.CodeExtensions.Select(x => $"-*{x}")).ToArray();
+
+		/// <summary>
+		/// The tag text
+		/// </summary>
+		public string Text { get; }
+
+		/// <summary>
+		/// Constructor
+		/// </summary>
+		/// <param name="text"></param>
+		public CommitTag(string text)
+		{
+			Text = text;
+			StringId.ValidateArgument(text, nameof(text));
+		}
+
+		/// <inheritdoc/>
+		public override bool Equals(object? obj) => obj is CommitTag id && Equals(id);
+
+		/// <inheritdoc/>
+		public override int GetHashCode() => String.GetHashCode(Text, StringComparison.Ordinal);
+
+		/// <inheritdoc/>
+		public bool Equals(CommitTag other) => String.Equals(Text, other.Text, StringComparison.Ordinal);
+
+		/// <inheritdoc/>
+		public override string ToString() => Text;
+
+		/// <summary>
+		/// Compares two string ids for equality
+		/// </summary>
+		/// <param name="left">The first string id</param>
+		/// <param name="right">Second string id</param>
+		/// <returns>True if the two string ids are equal</returns>
+		public static bool operator ==(CommitTag left, CommitTag right) => left.Equals(right);
+
+		/// <summary>
+		/// Compares two string ids for inequality
+		/// </summary>
+		/// <param name="left">The first string id</param>
+		/// <param name="right">Second string id</param>
+		/// <returns>True if the two string ids are not equal</returns>
+		public static bool operator !=(CommitTag left, CommitTag right) => !left.Equals(right);
+	}
+
+	/// <summary>
+	/// Converts <see cref="CommitTag"/> values to and from JSON
+	/// </summary>
+	public class CommitTagJsonConverter : JsonConverter<CommitTag>
+	{
+		/// <inheritdoc/>
+		public override CommitTag Read(ref Utf8JsonReader reader, Type typeToConvert, JsonSerializerOptions options)
+		{
+			return new CommitTag(reader.GetString() ?? String.Empty);
+		}
+
+		/// <inheritdoc/>
+		public override void Write(Utf8JsonWriter writer, CommitTag value, JsonSerializerOptions options)
+		{
+			writer.WriteStringValue(value.Text);
+		}
+	}
+
+	/// <summary>
+	/// Serializer for StringId objects
+	/// </summary>
+	public sealed class CommitTagBsonSerializer : SerializerBase<CommitTag>
+	{
+		/// <inheritdoc/>
+		public override CommitTag Deserialize(BsonDeserializationContext context, BsonDeserializationArgs args)
+		{
+			return new CommitTag(context.Reader.ReadString());
+		}
+
+		/// <inheritdoc/>
+		public override void Serialize(BsonSerializationContext context, BsonSerializationArgs args, CommitTag value)
+		{
+			context.Writer.WriteString(value.ToString());
+		}
+	}
+}
diff --git a/Engine/Source/Programs/Horde/Horde.Build/Perforce/ICommit.cs b/Engine/Source/Programs/Horde/Horde.Build/Perforce/ICommit.cs
--- a/Engine/Source/Programs/Horde/Horde.Build/Perforce/ICommit.cs
+++ b/Engine/Source/Programs/Horde/Horde.Build/Perforce/ICommit.cs
@@ -61,45 +61,17 @@ namespace Horde.Build.Perforce
 		DateTime DateUtc { get; }
 
 		/// <summary>
-		/// Gets the files for this change, relative to the root of the stream
+		/// Gets the list of tags for the commit
 		/// </summary>
 		/// <param name="cancellationToken">Cancellation token for the operation</param>
-		/// <returns>List of files modified by this commit</returns>
-		ValueTask<IReadOnlyList<string>> GetFilesAsync(CancellationToken cancellationToken);
-	}
+		/// <returns>True if the commit has the given tag</returns>
+		ValueTask<IReadOnlyList<CommitTag>> GetTagsAsync(CancellationToken cancellationToken);
 
-	/// <summary>
-	/// Extension methods for commits
-	/// </summary>
-	public static class CommitExtensions
-	{
 		/// <summary>
-		/// Determines if this change is a code change
+		/// Gets the files for this change, relative to the root of the stream
 		/// </summary>
-		/// <returns>True if this change is a code change</returns>
-		public static async ValueTask<ChangeContentFlags> GetContentFlagsAsync(this ICommit commit, CancellationToken cancellationToken)
-		{
-			ChangeContentFlags scope = 0;
-
-			// Check whether the files are code or content
-			IReadOnlyList<string> files = await commit.GetFilesAsync(cancellationToken);
-			foreach (string file in files)
-			{
-				if (PerforceUtils.CodeExtensions.Any(extension => file.EndsWith(extension, StringComparison.OrdinalIgnoreCase)))
-				{
-					scope |= ChangeContentFlags.ContainsCode;
-				}
-				else
-				{
-					scope |= ChangeContentFlags.ContainsContent;
-				}
-
-				if (scope == (ChangeContentFlags.ContainsCode | ChangeContentFlags.ContainsContent))
-				{
-					break;
-				}
-			}
-			return scope;
-		}
+		/// <param name="cancellationToken">Cancellation token for the operation</param>
+		/// <returns>List of files modified by this commit</returns>
+		ValueTask<IReadOnlyList<string>> GetFilesAsync(CancellationToken cancellationToken);
 	}
 }
diff --git a/Engine/Source/Programs/Horde/Horde.Build/Perforce/PerforceService.cs b/Engine/Source/Programs/Horde/Horde.Build/Perforce/PerforceService.cs
--- a/Engine/Source/Programs/Horde/Horde.Build/Perforce/PerforceService.cs
+++ b/Engine/Source/Programs/Horde/Horde.Build/Perforce/PerforceService.cs
@@ -19,6 +19,7 @@ using Horde.Build.Server;
 using Horde.Build.Streams;
 using Horde.Build.Users;
 using Horde.Build.Utilities;
+using Microsoft.AspNetCore.Connections.Features;
 using Microsoft.Extensions.Caching.Memory;
 using Microsoft.Extensions.Logging;
 using Microsoft.Extensions.Options;
@@ -128,27 +129,24 @@ namespace Horde.Build.Perforce
 		class Commit : ICommit
 		{
 			readonly PerforceService _owner;
-			readonly string _clusterName;
-			readonly string _streamName;
-
-			public StreamId StreamId { get; set; }
-			public int Number { get; set; }
-			public int OriginalChange { get; set; }
-			public UserId AuthorId { get; set; }
-			public UserId OwnerId { get; set; }
-			public string Description { get; set; }
-			public string BasePath { get; set; }
-			public DateTime DateUtc { get; set; }
+			readonly IStream _stream;
+
+			public StreamId StreamId => _stream.Id;
+			public int Number { get; }
+			public int OriginalChange { get; }
+			public UserId AuthorId { get; }
+			public UserId OwnerId { get; }
+			public string Description { get; }
+			public string BasePath { get; }
+			public DateTime DateUtc { get; }
 
 			IReadOnlyList<string>? _files;
 
-			public Commit(PerforceService owner, string clusterName, string streamName, StreamId streamId, int number, int originalChange, UserId authorId, UserId ownerId, string description, string basePath, DateTime dateUtc)
+			public Commit(PerforceService owner, IStream stream, int number, int originalChange, UserId authorId, UserId ownerId, string description, string basePath, DateTime dateUtc)
 			{
 				_owner = owner;
-				_clusterName = clusterName;
-				_streamName = streamName;
-	
-				StreamId = streamId;
+				_stream = stream;
+
 				Number = number;
 				OriginalChange = originalChange;
 				AuthorId = authorId;
@@ -163,14 +161,30 @@ namespace Horde.Build.Perforce
 				_files = files;
 			}
 
+			public async ValueTask<IReadOnlyList<CommitTag>> GetTagsAsync(CancellationToken cancellationToken)
+			{
+				List<CommitTag> commitTags = new List<CommitTag>();
+
+				IReadOnlyList<string> files = await GetFilesAsync(cancellationToken);
+				foreach (CommitTagConfig tagConfig in _stream.Config.GetAllCommitTags())
+				{
+					if (tagConfig.CreateFileFilter().ApplyTo(files).Any())
+					{
+						commitTags.Add(tagConfig.Name);
+					}
+				}
+
+				return commitTags;
+			}
+
 			public async ValueTask<IReadOnlyList<string>> GetFilesAsync(CancellationToken cancellationToken)
 			{
 				if (_files == null)
 				{
-					using (IPerforceConnection perforce = await _owner.ConnectAsync(_clusterName, null, cancellationToken))
+					using (IPerforceConnection perforce = await _owner.ConnectAsync(_stream.Config.ClusterName, null, cancellationToken))
 					{
 						DescribeRecord describeRecord = await perforce.DescribeAsync(Number, cancellationToken);
-						_files = GetStreamFiles(describeRecord, _streamName);
+						_files = GetStreamFiles(describeRecord, _stream.Config.Name);
 					}
 				}
 				return _files;
@@ -635,7 +649,7 @@ namespace Horde.Build.Perforce
 				ownerUser = await FindOrAddUserAsync(stream.Config.ClusterName, owner, cancellationToken);
 			}
 
-			return new Commit(this, stream.Config.ClusterName, stream.Name, stream.Id, number, originalChange, authorUser.Id, ownerUser.Id, description, basePath ?? String.Empty, dateUtc);
+			return new Commit(this, stream, number, originalChange, authorUser.Id, ownerUser.Id, description, basePath ?? String.Empty, dateUtc);
 		}
 
 		public async ValueTask<ICommit> CreateCommitAsync(IStream stream, DescribeRecord describeRecord, InfoRecord serverInfo, CancellationToken cancellationToken)
@@ -1063,9 +1077,9 @@ namespace Horde.Build.Perforce
 				List<ICommit> detailsList = await GetChangeDetailsAsync(stream, changes.ConvertAll(x => x.Number), cancellationToken);
 				foreach (ICommit details in detailsList.OrderByDescending(x => x.Number))
 				{
-					ChangeContentFlags contentFlags = await details.GetContentFlagsAsync(cancellationToken);
-					_logger.LogInformation("Change {Change} = {Flags}", details.Number, contentFlags.ToString());
-					if ((contentFlags & ChangeContentFlags.ContainsCode) != 0)
+					IReadOnlyList<CommitTag> commitTags = await details.GetTagsAsync(cancellationToken);
+					_logger.LogInformation("Change {Change} = {Flags}", details.Number, String.Join(", ", commitTags.Select(x => x.ToString())));
+					if (commitTags.Contains(CommitTag.Code))
 					{
 						return details.Number;
 					}
diff --git a/Engine/Source/Programs/Horde/Horde.Build/Perforce/PerforceServiceCache.cs b/Engine/Source/Programs/Horde/Horde.Build/Perforce/PerforceServiceCache.cs
--- a/Engine/Source/Programs/Horde/Horde.Build/Perforce/PerforceServiceCache.cs
+++ b/Engine/Source/Programs/Horde/Horde.Build/Perforce/PerforceServiceCache.cs
@@ -7,6 +7,7 @@ using System.Linq;
 using System.Runtime.CompilerServices;
 using System.Threading;
 using System.Threading.Tasks;
+using EpicGames.Core;
 using EpicGames.Perforce;
 using Horde.Build.Server;
 using Horde.Build.Streams;
@@ -57,15 +58,29 @@ namespace Horde.Build.Perforce
 			public Dictionary<StreamId, int> MinChanges { get; set; } = new Dictionary<StreamId, int>();
 		}
 
+		class CommitTagInfo
+		{
+			public CommitTag Name { get; }
+			public FileFilter Filter { get; }
+
+			public CommitTagInfo(CommitTagConfig config)
+			{
+				Name = config.Name;
+				Filter = config.CreateFileFilter();
+			}
+		}
+
 		class StreamInfo
 		{
 			public IStream Stream { get; set; }
 			public PerforceViewMap View { get; }
+			public List<CommitTagInfo> CommitTags { get; set; }
 
 			public StreamInfo(IStream stream, PerforceViewMap view)
 			{
 				Stream = stream;
 				View = view;
+				CommitTags = stream.Config.GetAllCommitTags().Select(x => new CommitTagInfo(x)).ToList();
 			}
 		}
 
@@ -102,14 +117,18 @@ namespace Horde.Build.Perforce
 			public string BasePath { get; set; }
 			public DateTime DateUtc { get; set; }
 
+			public List<CommitTag> CommitTags { get; set; } = new List<CommitTag>();
+
 			[BsonConstructor]
 			CachedCommitDoc()
 			{
 				Description = null!;
 				BasePath = null!;
+
+				CommitTags = new List<CommitTag>();
 			}
 
-			public CachedCommitDoc(ICommit commit)
+			public CachedCommitDoc(ICommit commit, List<CommitTag> commitTags)
 			{
 				StreamId = commit.StreamId;
 				Number = commit.Number;
@@ -119,6 +138,8 @@ namespace Horde.Build.Perforce
 				Description = commit.Description;
 				BasePath = commit.BasePath;
 				DateUtc = commit.DateUtc;
+
+				CommitTags = commitTags;
 			}
 
 			public void PostLoad(PerforceServiceCache owner, IStream stream)
@@ -127,6 +148,11 @@ namespace Horde.Build.Perforce
 				_stream = stream;
 			}
 
+			public ValueTask<IReadOnlyList<CommitTag>> GetTagsAsync(CancellationToken cancellationToken)
+			{
+				return new ValueTask<IReadOnlyList<CommitTag>>(CommitTags);
+			}
+
 			public async ValueTask<IReadOnlyList<string>> GetFilesAsync(CancellationToken cancellationToken)
 			{
 				ICommit? other = await _owner!.GetChangeDetailsAsync(_stream, Number, cancellationToken);
@@ -154,7 +180,11 @@ namespace Horde.Build.Perforce
 			Options = options.Value;
 
 			_mongoService = mongoService;
-			_commits = mongoService.GetCollection<CachedCommitDoc>("CommitsV2", keys => keys.Ascending(x => x.StreamId).Descending(x => x.Number), unique: true);
+
+			List<MongoIndex<CachedCommitDoc>> indexes = new List<MongoIndex<CachedCommitDoc>>();
+			indexes.Add(MongoIndex.Create<CachedCommitDoc>(keys => keys.Ascending(x => x.StreamId).Descending(x => x.Number), true));
+			indexes.Add(MongoIndex.Create<CachedCommitDoc>(keys => keys.Ascending(x => x.StreamId).Ascending(x => x.CommitTags).Descending(x => x.Number), true));
+			_commits = mongoService.GetCollection<CachedCommitDoc>("CommitsV2", indexes);
 
 			_streamCollection = streamCollection;
 			_logger = logger;
@@ -364,8 +394,17 @@ namespace Horde.Build.Perforce
 						if (files.Count > 0)
 						{
 							ICommit commit = await CreateCommitAsync(streamInfo.Stream, describeRecord, info, cancellationToken);
-							CachedCommitDoc commitDoc = new CachedCommitDoc(commit);
 
+							List<CommitTag> commitTags = new List<CommitTag>();
+							foreach (CommitTagInfo commitTagInfo in streamInfo.CommitTags)
+							{
+								if (commitTagInfo.Filter.ApplyTo(files.Select(x => "/" + x)).Any())
+								{
+									commitTags.Add(commitTagInfo.Name);
+								}
+							}
+
+							CachedCommitDoc commitDoc = new CachedCommitDoc(commit, commitTags);
 							FilterDefinition<CachedCommitDoc> filter = Builders<CachedCommitDoc>.Filter.Expr(x => x.StreamId == commit.StreamId && x.Number == commit.Number);
 							await _commits.ReplaceOneAsync(filter, commitDoc, new ReplaceOptions { IsUpsert = true }, cancellationToken);
 						}
@@ -376,12 +415,14 @@ namespace Horde.Build.Perforce
 				Dictionary<StreamId, int> minChanges = new Dictionary<StreamId, int>(streamInfos.Count);
 				foreach (StreamInfo streamInfo in streamInfos)
 				{
-					if (!state.MinChanges.ContainsKey(streamInfo.Stream.Id))
+					int minChange;
+					if (!state.MinChanges.TryGetValue(streamInfo.Stream.Id, out minChange))
 					{
-						int minChange = changes[^1].Number;
-						state.MinChanges.Add(streamInfo.Stream.Id, minChange);
+						minChange = changes[^1].Number;
 					}
+					minChanges.Add(streamInfo.Stream.Id, minChange);
 				}
+				state.MinChanges = minChanges;
 				state.MaxChange = changes[0].Number;
 				return state;
 			}
diff --git a/Engine/Source/Programs/Horde/Horde.Build/Streams/StreamConfig.cs b/Engine/Source/Programs/Horde/Horde.Build/Streams/StreamConfig.cs
--- a/Engine/Source/Programs/Horde/Horde.Build/Streams/StreamConfig.cs
+++ b/Engine/Source/Programs/Horde/Horde.Build/Streams/StreamConfig.cs
@@ -5,7 +5,9 @@ using System.Collections.Generic;
 using System.ComponentModel.DataAnnotations;
 using System.Diagnostics.CodeAnalysis;
 using System.Linq;
+using Datadog.Trace;
 using EpicGames.Core;
+using EpicGames.Perforce;
 using Horde.Build.Acls;
 using Horde.Build.Agents.Pools;
 using Horde.Build.Issues;
@@ -22,6 +24,24 @@ namespace Horde.Build.Streams
 	using TemplateId = StringId<ITemplateRef>;
 	using WorkflowId = StringId<WorkflowConfig>;
 
+	/// <summary>
+	/// Flags identifying content of a changelist
+	/// </summary>
+	[Flags]
+	[Obsolete]
+	public enum ChangeContentFlags
+	{
+		/// <summary>
+		/// The change contains code
+		/// </summary>
+		ContainsCode = 1,
+
+		/// <summary>
+		/// The change contains content
+		/// </summary>
+		ContainsContent = 2,
+	}
+
 	/// <summary>
 	/// How to replicate data for this stream
 	/// </summary>
@@ -96,6 +116,11 @@ namespace Horde.Build.Streams
 		/// </summary>
 		public DefaultPreflightConfig? DefaultPreflight { get; set; }
 
+		/// <summary>
+		/// List of tags to apply to commits. Allows fast searching and classification of different commit types (eg. code vs content).
+		/// </summary>
+		public List<CommitTagConfig> CommitTags { get; set; } = new List<CommitTagConfig>();
+
 		/// <summary>
 		/// List of tabs to show for the new stream
 		/// </summary>
@@ -154,7 +179,61 @@ namespace Horde.Build.Streams
 		/// <summary>
 		/// Tokens to create for each job step
 		/// </summary>
-		public List<TokenConfig> Tokens { get; set; } = new List<TokenConfig>(); 
+		public List<TokenConfig> Tokens { get; set; } = new List<TokenConfig>();
+
+		/// <summary>
+		/// Enumerates all commit tags, including the default tags for code and content.
+		/// </summary>
+		/// <returns></returns>
+		public IEnumerable<CommitTagConfig> GetAllCommitTags()
+		{
+			if (TryGetCommitTag(CommitTag.Code, out CommitTagConfig? codeConfig))
+			{
+				yield return codeConfig;
+			}
+			if (TryGetCommitTag(CommitTag.Content, out CommitTagConfig? contentConfig))
+			{
+				yield return contentConfig;
+			}
+			foreach (CommitTagConfig config in CommitTags)
+			{
+				if (config.Name != CommitTag.Code && config.Name != CommitTag.Content)
+				{
+					yield return config;
+				}
+			}
+		}
+
+		/// <summary>
+		/// Gets the configuration for a commit tag
+		/// </summary>
+		/// <param name="tag">Tag to search for</param>
+		/// <param name="config">Receives the tag configuration</param>
+		/// <returns>True if a match was found</returns>
+		public bool TryGetCommitTag(CommitTag tag, [NotNullWhen(true)] out CommitTagConfig? config)
+		{
+			CommitTagConfig? first = CommitTags.FirstOrDefault(x => x.Name == tag);
+			if (first != null)
+			{
+				config = first;
+				return true;
+			}
+			else if (tag == CommitTag.Code)
+			{
+				config = CommitTagConfig.CodeDefault;
+				return true;
+			}
+			else if (tag == CommitTag.Content)
+			{
+				config = CommitTagConfig.ContentDefault;
+				return true;
+			}
+			else
+			{
+				config = null;
+				return false;
+			}
+		}
 
 		/// <summary>
 		/// Tries to find a template with the given id
@@ -640,10 +719,43 @@ namespace Horde.Build.Streams
 		/// </summary>
 		public ScheduleGateConfig? Gate { get; set; }
 
+		/// <summary>
+		/// Commit tags for this schedule
+		/// </summary>
+		public List<CommitTag> Commits { get; set; } = new List<CommitTag>();
+
 		/// <summary>
 		/// The types of changes to run for
 		/// </summary>
-		public List<ChangeContentFlags>? Filter { get; set; }
+		[Obsolete("Use Tags instead")]
+		public List<ChangeContentFlags> Filter
+		{
+			get
+			{
+				List<ChangeContentFlags> flags = new List<ChangeContentFlags>();
+				if (Commits.Contains(CommitTag.Code))
+				{
+					flags.Add(ChangeContentFlags.ContainsCode);
+				}
+				if (Commits.Contains(CommitTag.Content))
+				{
+					flags.Add(ChangeContentFlags.ContainsContent);
+				}
+				return flags;
+			}
+			set
+			{
+				Commits.Clear();
+				if (value.Contains(ChangeContentFlags.ContainsCode))
+				{
+					Commits.Add(CommitTag.Code);
+				}
+				if (value.Contains(ChangeContentFlags.ContainsContent))
+				{
+					Commits.Add(CommitTag.Content);
+				}
+			}
+		}
 
 		/// <summary>
 		/// Files that should cause the job to trigger
@@ -660,25 +772,6 @@ namespace Horde.Build.Streams
 		/// </summary>
 		public List<SchedulePatternConfig> Patterns { get; set; } = new List<SchedulePatternConfig>();
 
-		/// <summary>
-		/// Gets the flags to filter changes by
-		/// </summary>
-		/// <returns>Set of filter flags</returns>
-		public ChangeContentFlags? GetFilterFlags()
-		{
-			if (Filter == null)
-			{
-				return null;
-			}
-
-			ChangeContentFlags flags = 0;
-			foreach (ChangeContentFlags flag in Filter)
-			{
-				flags |= flag;
-			}
-			return flags;
-		}
-
 		/// <summary>
 		/// Get the next time that the schedule will trigger
 		/// </summary>
@@ -732,4 +825,42 @@ namespace Horde.Build.Streams
 		[Required]
 		public string EnvVar { get; set; } = String.Empty;
 	}
+
+	/// <summary>
+	/// Configuration for custom commit filters
+	/// </summary>
+	public class CommitTagConfig
+	{
+		/// <summary>
+		/// Name of the tag
+		/// </summary>
+		[Required]
+		public CommitTag Name { get; set; }
+
+		/// <summary>
+		/// List of files to be included in this filter
+		/// </summary>
+		public List<string> Filter { get; set; } = new List<string>();
+
+		/// <summary>
+		/// Default config for code filters
+		/// </summary>
+		public static CommitTagConfig CodeDefault { get; } = new CommitTagConfig { Name = CommitTag.Code, Filter = CommitTag.CodeFilter.ToList() };
+
+		/// <summary>
+		/// Default config for content filters
+		/// </summary>
+		public static CommitTagConfig ContentDefault { get; } = new CommitTagConfig { Name = CommitTag.Content, Filter = CommitTag.ContentFilter.ToList() };
+
+		/// <summary>
+		/// Constructs a <see cref="FileFilter"/> from the rules in this configuration object
+		/// </summary>
+		/// <returns>Filter object</returns>
+		public FileFilter CreateFileFilter()
+		{
+			FileFilter filter = new FileFilter(FileFilterType.Exclude);
+			filter.AddRules(Filter);
+			return filter;
+		}
+	}
 }
diff --git a/Engine/Source/Programs/Horde/Horde.Build/Streams/StreamsController.cs b/Engine/Source/Programs/Horde/Horde.Build/Streams/StreamsController.cs
--- a/Engine/Source/Programs/Horde/Horde.Build/Streams/StreamsController.cs
+++ b/Engine/Source/Programs/Horde/Horde.Build/Streams/StreamsController.cs
@@ -185,7 +185,8 @@ namespace Horde.Build.Streams
 			foreach (ICommit commit in commits)
 			{
 				IUser? author = await _userCollection.GetCachedUserAsync(commit.AuthorId);
-				responses.Add(new GetChangeSummaryResponse(commit, author!));
+				IReadOnlyList<CommitTag> tags = await commit.GetTagsAsync(HttpContext.RequestAborted);
+				responses.Add(new GetChangeSummaryResponse(commit, author!, tags));
 			}
 			return responses.ConvertAll(x => PropertyFilter.Apply(x, filter));
 		}
@@ -219,9 +220,10 @@ namespace Horde.Build.Streams
 			}
 
 			IUser? author = await _userCollection.GetCachedUserAsync(changeDetails.AuthorId);
+			IReadOnlyList<CommitTag> tags = await changeDetails.GetTagsAsync(HttpContext.RequestAborted);
 			IReadOnlyList<string> files = await changeDetails.GetFilesAsync(CancellationToken.None);
 
-			return PropertyFilter.Apply(new GetChangeDetailsResponse(changeDetails, author!, files), filter);
+			return PropertyFilter.Apply(new GetChangeDetailsResponse(changeDetails, author!, tags, files), filter);
 		}
 
 		/// <summary>
diff --git a/Engine/Source/Programs/Shared/EpicGames.Horde/StringId.cs b/Engine/Source/Programs/Shared/EpicGames.Horde/StringId.cs
--- a/Engine/Source/Programs/Shared/EpicGames.Horde/StringId.cs
+++ b/Engine/Source/Programs/Shared/EpicGames.Horde/StringId.cs
@@ -7,7 +7,7 @@ namespace EpicGames.Horde
 	/// <summary>
 	/// Normalized string identifier for a resource
 	/// </summary>
-	struct StringId : IEquatable<StringId>
+	public struct StringId : IEquatable<StringId>
 	{
 		/// <summary>
 		/// The text representing this id
@@ -17,45 +17,57 @@ namespace EpicGames.Horde
 		/// <summary>
 		/// Constructor
 		/// </summary>
-		/// <param name="input">Unique id for the string</param>
-		public StringId(string input)
+		/// <param name="text">Unique id for the string</param>
+		public StringId(string text)
 		{
-			Text = input;
+			Text = text;
+			ValidateArgument(text, nameof(text));
+		}
+
+		/// <summary>
+		/// Checks whether this StringId is set
+		/// </summary>
+		public bool IsEmpty => String.IsNullOrEmpty(Text);
 
-			if (Text.Length == 0)
+		/// <summary>
+		/// Validates the given string as a StringId, normalizing it if necessary.
+		/// </summary>
+		/// <param name="text">Text to validate as a StringId</param>
+		/// <param name="paramName">Name of the parameter to show if invalid characters are returned.</param>
+		/// <returns></returns>
+		public static string ValidateArgument(string text, string paramName)
+		{
+			if (text.Length == 0)
 			{
-				throw new ArgumentException("String id may not be empty");
+				throw new ArgumentException("String id may not be empty", paramName);
 			}
 
 			const int MaxLength = 64;
-			if (Text.Length > MaxLength)
+			if (text.Length > MaxLength)
 			{
-				throw new ArgumentException($"String id may not be longer than {MaxLength} characters");
+				throw new ArgumentException($"String id may not be longer than {MaxLength} characters", paramName);
 			}
 
-			for (int idx = 0; idx < Text.Length; idx++)
+			for (int idx = 0; idx < text.Length; idx++)
 			{
-				char character = Text[idx];
+				char character = text[idx];
 				if (!IsValidCharacter(character))
 				{
 					if (character >= 'A' && character <= 'Z')
 					{
 #pragma warning disable CA1308 // Normalize strings to uppercase
-						Text = Text.ToLowerInvariant();
+						text = text.ToLowerInvariant();
 #pragma warning restore CA1308 // Normalize strings to uppercase
 					}
 					else
 					{
-						throw new ArgumentException($"{Text} is not a valid string id");
+						throw new ArgumentException($"{text} is not a valid string id", paramName);
 					}
 				}
 			}
-		}
 
-		/// <summary>
-		/// Checks whether this StringId is set
-		/// </summary>
-		public bool IsEmpty => String.IsNullOrEmpty(Text);
+			return text;
+		}
 
 		/// <summary>
 		/// Checks whether the given character is valid within a string id
