Horde: Move commit responses into stream response file.

diff --git a/Engine/Source/Programs/Horde/Horde.Build/Perforce/ChangeMessages.cs b/Engine/Source/Programs/Horde/Horde.Build/Perforce/ChangeMessages.cs
--- a/Engine/Source/Programs/Horde/Horde.Build/Perforce/ChangeMessages.cs
+++ b/Engine/Source/Programs/Horde/Horde.Build/Perforce/ChangeMessages.cs
@@ -5,102 +5,5 @@ using Horde.Build.Users;
 
 namespace Horde.Build.Perforce
 {
-	/// <summary>
-	/// Information about a submitted changelist
-	/// </summary>
-	public class GetChangeSummaryResponse
-	{
-		/// <summary>
-		/// The source changelist number
-		/// </summary>
-		public int Number { get; set; }
 
-		/// <summary>
-		/// Name of the user that authored this change [DEPRECATED]
-		/// </summary>
-		public string Author { get; set; }
-
-		/// <summary>
-		/// Information about the change author
-		/// </summary>
-		public GetThinUserInfoResponse AuthorInfo { get; set; }
-
-		/// <summary>
-		/// The description text
-		/// </summary>
-		public string Description { get; set; }
-
-		/// <summary>
-		/// Tags for this commit
-		/// </summary>
-		public List<CommitTag> Tags { get; set; }
-
-		/// <summary>
-		/// Constructor
-		/// </summary>
-		/// <param name="commit">The commit to construct from</param>
-		/// <param name="author">Author of the commit</param>
-		/// <param name="tags">Tags for the commit</param>
-		public GetChangeSummaryResponse(ICommit commit, IUser author, IReadOnlyList<CommitTag> tags)
-		{
-			Number = commit.Number;
-			Author = author.Name;
-			AuthorInfo = new GetThinUserInfoResponse(author);
-			Description = commit.Description;
-			Tags = new List<CommitTag>(tags);
-		}
-	}
-
-	/// <summary>
-	/// Information about a submitted changelist
-	/// </summary>
-	public class GetChangeDetailsResponse
-	{
-		/// <summary>
-		/// The source changelist number
-		/// </summary>
-		public int Number { get; set; }
-
-		/// <summary>
-		/// Name of the user that authored this change [DEPRECATED]
-		/// </summary>
-		public string Author { get; set; }
-
-		/// <summary>
-		/// Information about the user that authored this change
-		/// </summary>
-		public GetThinUserInfoResponse AuthorInfo { get; set; }
-
-		/// <summary>
-		/// The description text
-		/// </summary>
-		public string Description { get; set; }
-
-		/// <summary>
-		/// Tags for this commit
-		/// </summary>
-		public List<CommitTag> Tags { get; set; }
-
-		/// <summary>
-		/// List of files that were modified, relative to the stream base
-		/// </summary>
-		public List<string> Files { get; set; }
-
-		/// <summary>
-		/// Constructor
-		/// </summary>
-		/// <param name="commit">The commit to construct from</param>
-		/// <param name="author">Author of the change</param>
-		/// <param name="tags">Tags for the commit</param>
-		/// <param name="files">Files modified by the commit</param>
-		public GetChangeDetailsResponse(ICommit commit, IUser author, IReadOnlyList<CommitTag> tags, IReadOnlyList<string> files)
-		{
-			Number = commit.Number;
-			Author = author.Name;
-			AuthorInfo = new GetThinUserInfoResponse(author);
-			Description = commit.Description;
-			Tags = new List<CommitTag>(tags);
-			Files = new List<string>(files);
-		}
-	}
 }
diff --git a/Engine/Source/Programs/Horde/Horde.Build/Streams/StreamMessages.cs b/Engine/Source/Programs/Horde/Horde.Build/Streams/StreamMessages.cs
--- a/Engine/Source/Programs/Horde/Horde.Build/Streams/StreamMessages.cs
+++ b/Engine/Source/Programs/Horde/Horde.Build/Streams/StreamMessages.cs
@@ -10,6 +10,7 @@ using Horde.Build.Jobs.Schedules;
 using Horde.Build.Users;
 using Horde.Build.Utilities;
 using Horde.Build.Issues;
+using Horde.Build.Perforce;
 
 namespace Horde.Build.Streams
 {
@@ -192,7 +193,7 @@ namespace Horde.Build.Streams
 		/// </summary>
 		public GetTemplateStepStateResponse(ITemplateStep state, GetThinUserInfoResponse? pausedByUserInfo)
 		{
-			Name = state.Name;			
+			Name = state.Name;
 			PauseTimeUtc = state.PauseTimeUtc;
 			PausedByUserInfo = pausedByUserInfo;
 		}
@@ -278,7 +279,7 @@ namespace Horde.Build.Streams
 			Schedule = (templateRef.Schedule != null) ? new GetScheduleResponse(templateRef.Schedule) : null;
 			ChainedJobs = (templateRef.Config.ChainedJobs != null && templateRef.Config.ChainedJobs.Count > 0) ? templateRef.Config.ChainedJobs : null;
 			StepStates = stepStates;
-			Acl = (bIncludeAcl && templateRef.Acl != null)? new GetAclResponse(templateRef.Acl) : null;
+			Acl = (bIncludeAcl && templateRef.Acl != null) ? new GetAclResponse(templateRef.Acl) : null;
 		}
 	}
 
@@ -366,12 +367,12 @@ namespace Horde.Build.Streams
 		/// Custom permissions for this object
 		/// </summary>
 		public GetAclResponse? Acl { get; set; }
-		
+
 		/// <summary>
 		/// Stream paused for new builds until this date
 		/// </summary>
 		public DateTime? PausedUntil { get; set; }
-		
+
 		/// <summary>
 		/// Reason for stream being paused
 		/// </summary>
@@ -424,4 +425,61 @@ namespace Horde.Build.Streams
 			Workflows = workflows;
 		}
 	}
+
+	/// <summary>
+	/// Information about a commit
+	/// </summary>
+	public class GetCommitResponse
+	{
+		/// <summary>
+		/// The source changelist number
+		/// </summary>
+		public int Number { get; set; }
+
+		/// <summary>
+		/// Name of the user that authored this change [DEPRECATED]
+		/// </summary>
+		public string Author { get; set; }
+
+		/// <summary>
+		/// Information about the user that authored this change
+		/// </summary>
+		public GetThinUserInfoResponse AuthorInfo { get; set; }
+
+		/// <summary>
+		/// The description text
+		/// </summary>
+		public string Description { get; set; }
+
+		/// <summary>
+		/// Tags for this commit
+		/// </summary>
+		public List<CommitTag> Tags { get; set; }
+
+		/// <summary>
+		/// List of files that were modified, relative to the stream base
+		/// </summary>
+		public List<string>? Files { get; set; }
+
+		/// <summary>
+		/// Constructor
+		/// </summary>
+		/// <param name="commit">The commit to construct from</param>
+		/// <param name="author">Author of the change</param>
+		/// <param name="tags">Tags for the commit</param>
+		/// <param name="files">Files modified by the commit</param>
+		public GetCommitResponse(ICommit commit, IUser author, IReadOnlyList<CommitTag> tags, IReadOnlyList<string>? files)
+		{
+			Number = commit.Number;
+			Author = author.Name;
+			AuthorInfo = new GetThinUserInfoResponse(author);
+			Description = commit.Description;
+			Tags = new List<CommitTag>(tags);
+
+			if (files != null)
+			{
+				Files = new List<string>(files);
+			}
+		}
+	}
 }
diff --git a/Engine/Source/Programs/Horde/Horde.Build/Streams/StreamsController.cs b/Engine/Source/Programs/Horde/Horde.Build/Streams/StreamsController.cs
--- a/Engine/Source/Programs/Horde/Horde.Build/Streams/StreamsController.cs
+++ b/Engine/Source/Programs/Horde/Horde.Build/Streams/StreamsController.cs
@@ -166,7 +166,7 @@ namespace Horde.Build.Streams
 		/// <returns>Http result code</returns>
 		[HttpGet]
 		[Route("/api/v1/streams/{streamId}/changes")]
-		[ProducesResponseType(typeof(List<GetChangeSummaryResponse>), 200)]
+		[ProducesResponseType(typeof(List<GetCommitResponse>), 200)]
 		public async Task<ActionResult<List<object>>> GetChangesAsync(StreamId streamId, [FromQuery] int? min = null, [FromQuery] int? max = null, [FromQuery] int results = 50, PropertyFilter? filter = null)
 		{
 			IStream? stream = await _streamService.GetStreamAsync(streamId);
@@ -181,12 +181,12 @@ namespace Horde.Build.Streams
 
 			List<ICommit> commits = await _commitService.GetCollection(stream).FindAsync(min, max, results).ToListAsync();
 
-			List<GetChangeSummaryResponse> responses = new List<GetChangeSummaryResponse>();
+			List<GetCommitResponse> responses = new List<GetCommitResponse>();
 			foreach (ICommit commit in commits)
 			{
 				IUser? author = await _userCollection.GetCachedUserAsync(commit.AuthorId);
 				IReadOnlyList<CommitTag> tags = await commit.GetTagsAsync(HttpContext.RequestAborted);
-				responses.Add(new GetChangeSummaryResponse(commit, author!, tags));
+				responses.Add(new GetCommitResponse(commit, author!, tags, null));
 			}
 			return responses.ConvertAll(x => PropertyFilter.Apply(x, filter));
 		}
@@ -200,7 +200,7 @@ namespace Horde.Build.Streams
 		/// <returns>Http result code</returns>
 		[HttpGet]
 		[Route("/api/v1/streams/{streamId}/changes/{changeNumber}")]
-		[ProducesResponseType(typeof(GetChangeDetailsResponse), 200)]
+		[ProducesResponseType(typeof(GetCommitResponse), 200)]
 		public async Task<ActionResult<object>> GetChangeDetailsAsync(StreamId streamId, int changeNumber, PropertyFilter? filter = null)
 		{
 			IStream? stream = await _streamService.GetStreamAsync(streamId);
@@ -223,7 +223,7 @@ namespace Horde.Build.Streams
 			IReadOnlyList<CommitTag> tags = await changeDetails.GetTagsAsync(HttpContext.RequestAborted);
 			IReadOnlyList<string> files = await changeDetails.GetFilesAsync(CancellationToken.None);
 
-			return PropertyFilter.Apply(new GetChangeDetailsResponse(changeDetails, author!, tags, files), filter);
+			return PropertyFilter.Apply(new GetCommitResponse(changeDetails, author!, tags, files), filter);
 		}
 
 		/// <summary>
