Horde: Fix build error.

diff --git a/Engine/Source/Programs/Horde/Horde.Build.Tests/Stubs/Services/PerforceServiceStub.cs b/Engine/Source/Programs/Horde/Horde.Build.Tests/Stubs/Services/PerforceServiceStub.cs
--- a/Engine/Source/Programs/Horde/Horde.Build.Tests/Stubs/Services/PerforceServiceStub.cs
+++ b/Engine/Source/Programs/Horde/Horde.Build.Tests/Stubs/Services/PerforceServiceStub.cs
@@ -58,6 +58,11 @@ namespace Horde.Build.Tests.Stubs.Services
 				Files = files;
 			}
 
+			public ValueTask<IReadOnlyList<CommitTag>> GetTagsAsync(CancellationToken cancellationToken)
+			{
+				return new ValueTask<IReadOnlyList<CommitTag>>(Array.Empty<CommitTag>());
+			}
+
 			public ValueTask<IReadOnlyList<string>> GetFilesAsync(CancellationToken cancellationToken)
 			{
 				return new ValueTask<IReadOnlyList<string>>(Files);
