Horde: Fix more build warnings.

diff --git a/Engine/Source/Programs/Horde/Horde.Build.Tests/SchedulerTests.cs b/Engine/Source/Programs/Horde/Horde.Build.Tests/SchedulerTests.cs
--- a/Engine/Source/Programs/Horde/Horde.Build.Tests/SchedulerTests.cs
+++ b/Engine/Source/Programs/Horde/Horde.Build.Tests/SchedulerTests.cs
@@ -270,7 +270,7 @@ namespace Horde.Build.Tests
 			schedule.Enabled = true;
 			schedule.Patterns.Add(new SchedulePatternConfig { MinTime = 13 * 60, MaxTime = 14 * 60, Interval = 15 });
 			schedule.MaxChanges = 2;
-			schedule.Filter = new List<ChangeContentFlags> { ChangeContentFlags.ContainsCode };
+			schedule.Commits.Add(CommitTag.Code);
 			await SetScheduleAsync(schedule);
 
 			// Initial tick
@@ -307,7 +307,7 @@ namespace Horde.Build.Tests
 			schedule.Enabled = true;
 			schedule.Patterns.Add(new SchedulePatternConfig { MinTime = 13 * 60, MaxTime = 14 * 60, Interval = 15 });
 			schedule.MaxChanges = 2;
-			schedule.Filter = new List<ChangeContentFlags> { ChangeContentFlags.ContainsCode };
+			schedule.Commits.Add(CommitTag.Code);
 			await SetScheduleAsync(schedule);
 
 			// Initial tick
@@ -498,7 +498,7 @@ namespace Horde.Build.Tests
 			newTemplate2.Name = "Test template 2";
 			newTemplate2.Schedule = new ScheduleConfig();
 			newTemplate2.Schedule.MaxChanges = 4;
-			newTemplate2.Schedule.Filter = new List<ChangeContentFlags> { ChangeContentFlags.ContainsCode };
+			newTemplate2.Schedule.Commits.Add(CommitTag.Code);
 			newTemplate2.Schedule.Gate = new ScheduleGateConfig { TemplateId = newTemplateRefId1, Target = "TriggerNext" };
 			newTemplate2.Schedule.Patterns.Add(new SchedulePatternConfig { Interval = 10 });// (null, 0, null, 10));
 //			NewTemplate2.Schedule.LastTriggerTime = StartTime;
diff --git a/Engine/Source/Programs/Horde/Horde.Build.Tests/Stubs/Services/PerforceServiceStub.cs b/Engine/Source/Programs/Horde/Horde.Build.Tests/Stubs/Services/PerforceServiceStub.cs
--- a/Engine/Source/Programs/Horde/Horde.Build.Tests/Stubs/Services/PerforceServiceStub.cs
+++ b/Engine/Source/Programs/Horde/Horde.Build.Tests/Stubs/Services/PerforceServiceStub.cs
@@ -60,7 +60,16 @@ namespace Horde.Build.Tests.Stubs.Services
 
 			public ValueTask<IReadOnlyList<CommitTag>> GetTagsAsync(CancellationToken cancellationToken)
 			{
-				return new ValueTask<IReadOnlyList<CommitTag>>(Array.Empty<CommitTag>());
+				List<CommitTag> tags = new List<CommitTag>();
+				if (CommitTagConfig.CodeDefault.CreateFileFilter().ApplyTo(Files).Any())
+				{
+					tags.Add(CommitTag.Code);
+				}
+				if (CommitTagConfig.ContentDefault.CreateFileFilter().ApplyTo(Files).Any())
+				{
+					tags.Add(CommitTag.Content);
+				}
+				return new ValueTask<IReadOnlyList<CommitTag>>(tags);
 			}
 
 			public ValueTask<IReadOnlyList<string>> GetFilesAsync(CancellationToken cancellationToken)
