Horde: Use indexed queries on cached P4 commits to determine the last code change for a particular commit.

diff --git a/Engine/Source/Programs/Horde/Horde.Build.Tests/Stubs/Services/PerforceServiceStub.cs b/Engine/Source/Programs/Horde/Horde.Build.Tests/Stubs/Services/PerforceServiceStub.cs
--- a/Engine/Source/Programs/Horde/Horde.Build.Tests/Stubs/Services/PerforceServiceStub.cs
+++ b/Engine/Source/Programs/Horde/Horde.Build.Tests/Stubs/Services/PerforceServiceStub.cs
@@ -231,15 +231,27 @@ namespace Horde.Build.Tests.Stubs.Services
 				_stream = stream;
 			}
 
-			public async IAsyncEnumerable<ICommit> FindAsync(int? minChange, int? maxChange, int? maxResults, [EnumeratorCancellation] CancellationToken cancellationToken = default)
+			public async IAsyncEnumerable<ICommit> FindAsync(int? minChange, int? maxChange, int? maxResults, IReadOnlyList<CommitTag>? tags, [EnumeratorCancellation] CancellationToken cancellationToken = default)
 			{
-				foreach (ICommit commit in await _owner.GetChangesAsync(_stream, minChange, maxChange, maxResults, cancellationToken))
+				foreach (ICommit commit in await _owner.GetChangesAsync(_stream, minChange, maxChange, null, cancellationToken))
 				{
+					if (tags != null && tags.Count > 0)
+					{
+						IReadOnlyList<CommitTag> commitTags = await commit.GetTagsAsync(cancellationToken);
+						if (!tags.Any(x => commitTags.Contains(x)))
+						{
+							continue;
+						}
+					}
 					yield return commit;
 				}
 			}
 
-			public Task<ICommit> GetAsync(int changeNumber, CancellationToken cancellationToken = default) => _owner.GetChangeDetailsAsync(_stream, changeNumber);
+			public async Task<ICommit> GetAsync(int changeNumber, CancellationToken cancellationToken = default)
+			{
+				List<ICommit> commits = await _owner.GetChangeDetailsAsync(_stream, new[] { changeNumber }, cancellationToken);
+				return commits[0];
+			}
 
 			public async Task<int> LatestNumberAsync(CancellationToken cancellationToken = default)
 			{
diff --git a/Engine/Source/Programs/Horde/Horde.Build.Tests/TestSetup.cs b/Engine/Source/Programs/Horde/Horde.Build.Tests/TestSetup.cs
--- a/Engine/Source/Programs/Horde/Horde.Build.Tests/TestSetup.cs
+++ b/Engine/Source/Programs/Horde/Horde.Build.Tests/TestSetup.cs
@@ -97,6 +97,7 @@ namespace Horde.Build.Tests
 		public AclService AclService => ServiceProvider.GetRequiredService<AclService>();
 		public AgentSoftwareService AgentSoftwareService => ServiceProvider.GetRequiredService<AgentSoftwareService>();
 		public AgentService AgentService => ServiceProvider.GetRequiredService<AgentService>();
+		public ICommitService CommitService => ServiceProvider.GetRequiredService<ICommitService>();
 		public ComputeService ComputeService => ServiceProvider.GetRequiredService<ComputeService>();
 		public MongoService MongoService => ServiceProvider.GetRequiredService<MongoService>();
 		public ITemplateCollection TemplateCollection => ServiceProvider.GetRequiredService<ITemplateCollection>();
@@ -268,7 +269,7 @@ namespace Horde.Build.Tests
 		private JobsController GetJobsController()
         {
 			ILogger<JobsController> logger = ServiceProvider.GetRequiredService<ILogger<JobsController>>();
-			JobsController jobsCtrl = new JobsController(GraphCollection, PerforceService, StreamService, JobService,
+			JobsController jobsCtrl = new JobsController(GraphCollection, CommitService, PerforceService, StreamService, JobService,
 		        TemplateCollection, ArtifactCollection, UserCollection, NotificationService, AgentService, logger);
 	        jobsCtrl.ControllerContext = GetControllerContext();
 	        return jobsCtrl;
diff --git a/Engine/Source/Programs/Horde/Horde.Build/Jobs/JobsController.cs b/Engine/Source/Programs/Horde/Horde.Build/Jobs/JobsController.cs
--- a/Engine/Source/Programs/Horde/Horde.Build/Jobs/JobsController.cs
+++ b/Engine/Source/Programs/Horde/Horde.Build/Jobs/JobsController.cs
@@ -41,6 +41,7 @@ namespace Horde.Build.Jobs
 	public class JobsController : HordeControllerBase
 	{
 		private readonly IGraphCollection _graphs;
+		private readonly ICommitService _commitService;
 		private readonly IPerforceService _perforce;
 		private readonly StreamService _streamService;
 		private readonly JobService _jobService;
@@ -54,9 +55,10 @@ namespace Horde.Build.Jobs
 		/// <summary>
 		/// Constructor
 		/// </summary>
-		public JobsController(IGraphCollection graphs, IPerforceService perforce, StreamService streamService, JobService jobService, ITemplateCollection templateCollection, IArtifactCollection artifactCollection, IUserCollection userCollection, INotificationService notificationService, AgentService agentService, ILogger<JobsController> logger)
+		public JobsController(IGraphCollection graphs, ICommitService commitService, IPerforceService perforce, StreamService streamService, JobService jobService, ITemplateCollection templateCollection, IArtifactCollection artifactCollection, IUserCollection userCollection, INotificationService notificationService, AgentService agentService, ILogger<JobsController> logger)
 		{
 			_graphs = graphs;
+			_commitService = commitService;
 			_perforce = perforce;
 			_streamService = streamService;
 			_jobService = jobService;
@@ -132,6 +134,9 @@ namespace Horde.Build.Jobs
 			// New groups for the job
 			IGraph graph = await _graphs.AddAsync(template);
 
+			// Get the commits for this stream
+			ICommitCollection commits = _commitService.GetCollection(stream);
+
 			// Get the change to build
 			int change;
 			if (create.Change.HasValue)
@@ -140,19 +145,20 @@ namespace Horde.Build.Jobs
 			}
 			else if (create.ChangeQuery != null)
 			{
-				change = await ExecuteChangeQueryAsync(stream, create.ChangeQuery.TemplateId ?? create.TemplateId, create.ChangeQuery.Target, create.ChangeQuery.Outcomes ?? new List<JobStepOutcome> { JobStepOutcome.Success });
+				change = await ExecuteChangeQueryAsync(stream, create.ChangeQuery.TemplateId ?? create.TemplateId, create.ChangeQuery.Target, create.ChangeQuery.Outcomes ?? new List<JobStepOutcome> { JobStepOutcome.Success }, commits);
 			}
 			else if (create.PreflightChange == null && template.SubmitNewChange != null)
 			{
-				change = await _perforce.CreateNewChangeForTemplateAsync(stream, template);
+				change = await commits.CreateNewAsync(template, HttpContext.RequestAborted);
 			}
 			else
 			{
-				change = await _perforce.GetLatestChangeAsync(stream);
+				change = await commits.GetLatestNumberAsync(HttpContext.RequestAborted);
 			}
 
 			// And get the matching code changelist
-			int codeChange = await _perforce.GetCodeChangeAsync(stream, change);
+			ICommit? lastCodeCommit = await commits.GetLastCodeChange(change, HttpContext.RequestAborted);
+			int codeChange = lastCodeCommit?.Number ?? change;
 
 			// New properties for the job
 			List<string> arguments = create.Arguments ?? template.GetDefaultArguments();
@@ -202,14 +208,15 @@ namespace Horde.Build.Jobs
 		/// <param name="templateId"></param>
 		/// <param name="target"></param>
 		/// <param name="outcomes"></param>
+		/// <param name="commits"></param>
 		/// <returns></returns>
-		async Task<int> ExecuteChangeQueryAsync(IStream stream, TemplateId templateId, string? target, List<JobStepOutcome> outcomes)
+		async Task<int> ExecuteChangeQueryAsync(IStream stream, TemplateId templateId, string? target, List<JobStepOutcome> outcomes, ICommitCollection commits)
 		{
 			IList<IJob> jobs = await _jobService.FindJobsAsync(streamId: stream.Id, templates: new[] { templateId }, target: target, state: new[] { JobStepState.Completed }, outcome: outcomes.ToArray(), count: 1, excludeUserJobs: true);
 			if (jobs.Count == 0)
 			{
 				_logger.LogInformation("Unable to find successful build of {TemplateId} target {Target}. Using latest change instead", templateId, target);
-				return await _perforce.GetLatestChangeAsync(stream);
+				return await commits.GetLatestNumberAsync();
 			}
 			else
 			{
diff --git a/Engine/Source/Programs/Horde/Horde.Build/Jobs/Schedules/ScheduleService.cs b/Engine/Source/Programs/Horde/Horde.Build/Jobs/Schedules/ScheduleService.cs
--- a/Engine/Source/Programs/Horde/Horde.Build/Jobs/Schedules/ScheduleService.cs
+++ b/Engine/Source/Programs/Horde/Horde.Build/Jobs/Schedules/ScheduleService.cs
@@ -61,83 +61,7 @@ namespace Horde.Build.Jobs.Schedules
 			public static double GetScoreFromTime(DateTime time) => (time.ToUniversalTime() - DateTime.UnixEpoch).TotalSeconds;
 		}
 
-		class PerforceHistory
-		{
-			readonly IStream _stream;
-			readonly ICommitService _commitService;
-			int _maxResults;
-			readonly List<ICommit> _changes = new List<ICommit>();
-			int _nextIndex;
-
-			public PerforceHistory(IStream stream, ICommitService commitService)
-			{
-				_stream = stream;
-				_commitService = commitService;
-				_maxResults = 10;
-			}
-
-			async ValueTask<ICommit?> GetChangeAtIndexAsync(int index)
-			{
-				while (index >= _changes.Count)
-				{
-					List<ICommit> newChanges = await _commitService.GetCollection(_stream).FindAsync(null, null, _maxResults).ToListAsync();
-
-					int numResults = newChanges.Count;
-					if (_changes.Count > 0)
-					{
-						newChanges.RemoveAll(x => x.Number >= _changes[^1].Number);
-					}
-					if (newChanges.Count == 0 && numResults < _maxResults)
-					{
-						return null;
-					}
-					if(newChanges.Count > 0)
-					{
-						_changes.AddRange(newChanges);
-					}
-					_maxResults += 10;
-				}
-				return _changes[index];
-			}
-
-			public async ValueTask<ICommit?> GetNextChangeAsync()
-			{
-				ICommit? details = await GetChangeAtIndexAsync(_nextIndex);
-				if (details != null)
-				{
-					_nextIndex++;
-				}
-				return details;
-			}
-
-			public async ValueTask<int> GetCodeChange(int change)
-			{
-				int index = _changes.BinarySearch(x => -x.Number, -change);
-				if (index < 0)
-				{
-					index = ~index;
-				}
-
-				for (; ; )
-				{
-					ICommit? details = await GetChangeAtIndexAsync(index);
-					if (details == null)
-					{
-						return 0;
-					}
-
-					IReadOnlyList<CommitTag> commitTags = await details.GetTagsAsync(CancellationToken.None);
-					if (commitTags.Contains(CommitTag.Code))
-					{
-						return details.Number;
-					}
-					index++;
-				}
-			}
-		}
-
 		readonly IGraphCollection _graphs;
-		readonly IPerforceService _perforce;
 		readonly ICommitService _commitService;
 		readonly IJobCollection _jobCollection;
 		readonly JobService _jobService;
@@ -155,10 +79,9 @@ namespace Horde.Build.Jobs.Schedules
 		/// <summary>
 		/// Constructor
 		/// </summary>
-		public ScheduleService(RedisService redis, IGraphCollection graphs, IPerforceService perforce, ICommitService commitService, IJobCollection jobCollection, JobService jobService, StreamService streamService, ITemplateCollection templateCollection, MongoService mongoService, IClock clock, IOptionsMonitor<ServerSettings> settings, ILogger<ScheduleService> logger)
+		public ScheduleService(RedisService redis, IGraphCollection graphs, ICommitService commitService, IJobCollection jobCollection, JobService jobService, StreamService streamService, ITemplateCollection templateCollection, MongoService mongoService, IClock clock, IOptionsMonitor<ServerSettings> settings, ILogger<ScheduleService> logger)
 		{
 			_graphs = graphs;
-			_perforce = perforce;
 			_commitService = commitService;
 			_jobCollection = jobCollection;
 			_jobService = jobService;
@@ -430,8 +353,10 @@ namespace Horde.Build.Jobs.Schedules
 			}
 
 			// Cache the Perforce history as we're iterating through changes to improve query performance
-			PerforceHistory history = new PerforceHistory(stream, _commitService);
-			
+			ICommitCollection commits = _commitService.GetCollection(stream);
+			IAsyncEnumerable<ICommit> commitEnumerable = commits.FindAsync(minChangeNumber, null, null, schedule.Config.Commits, cancellationToken);
+			IAsyncEnumerator<ICommit> commitEnumerator = commitEnumerable.GetAsyncEnumerator(cancellationToken);
+
 			// Start as many jobs as possible
 			List<(int Change, int CodeChange)> triggerChanges = new List<(int, int)>();
 			while (triggerChanges.Count < maxNewChanges)
@@ -444,9 +369,9 @@ namespace Horde.Build.Jobs.Schedules
 				{
 					changeDetails = await GetNextChangeForGateAsync(stream, templateId, schedule.Config.Gate, minChangeNumber, maxChangeNumber, cancellationToken);
 				}
-				else
+				else if (await commitEnumerator.MoveNextAsync(cancellationToken))
 				{
-					changeDetails = await history.GetNextChangeAsync();
+					changeDetails = commitEnumerator.Current;
 				}
 
 				// Quit if we didn't find anything
@@ -467,13 +392,13 @@ namespace Horde.Build.Jobs.Schedules
 				int change = changeDetails.Number;
 				if (await ShouldBuildChangeAsync(changeDetails, schedule.Config.Commits, fileFilter))
 				{
-					int codeChange = await history.GetCodeChange(change);
-					if (codeChange == -1)
+					ICommit? codeChange = await commits.GetLastCodeChange(changeDetails.Number, cancellationToken);
+					if (codeChange == null)
 					{
 						_logger.LogWarning("Unable to find code change for CL {Change}", change);
-						codeChange = change;
+						codeChange = changeDetails;
 					}
-					triggerChanges.Add((change, codeChange));
+					triggerChanges.Add((change, codeChange.Number));
 				}
 
 				// Check we haven't exceeded the time limit
@@ -512,9 +437,9 @@ namespace Horde.Build.Jobs.Schedules
 			// We may need to submit a new change for any new jobs. This only makes sense if there's one change.
 			if (template.SubmitNewChange != null)
 			{
-				int newChange = await _perforce.CreateNewChangeForTemplateAsync(stream, template);
-				int newCodeChange = await _perforce.GetCodeChangeAsync(stream, newChange, cancellationToken);
-				triggerChanges = new List<(int, int)> { (newChange, newCodeChange) };
+				int newChange = await commits.CreateNewAsync(template, cancellationToken);
+				ICommit? newCodeChange = await commits.GetLastCodeChange(newChange, cancellationToken);
+				triggerChanges = new List<(int, int)> { (newChange, newCodeChange?.Number ?? newChange) };
 			}
 
 			// Try to start all the new jobs
@@ -590,7 +515,7 @@ namespace Horde.Build.Jobs.Schedules
 						JobStepOutcome outcome = state.Value.Item2;
 						if (outcome == JobStepOutcome.Success || outcome == JobStepOutcome.Warnings)
 						{
-							return await _perforce.GetChangeDetailsAsync(stream, job.Change, cancellationToken);
+							return await _commitService.GetCollection(stream).GetAsync(job.Change, cancellationToken);
 						}
 						_logger.LogInformation("Skipping trigger of {StreamName} template {TemplateId} - last {OtherTemplateRefId} job ({JobId}) ended with errors", stream.Id, templateRefId, gate.TemplateId, job.Id);
 					}
diff --git a/Engine/Source/Programs/Horde/Horde.Build/Perforce/ICommitCollection.cs b/Engine/Source/Programs/Horde/Horde.Build/Perforce/ICommitCollection.cs
--- a/Engine/Source/Programs/Horde/Horde.Build/Perforce/ICommitCollection.cs
+++ b/Engine/Source/Programs/Horde/Horde.Build/Perforce/ICommitCollection.cs
@@ -5,9 +5,7 @@ using System.Collections.Generic;
 using System.Linq;
 using System.Threading;
 using System.Threading.Tasks;
-using Horde.Build.Streams;
-using Horde.Build.Users;
-using Horde.Build.Utilities;
+using Horde.Build.Jobs.Templates;
 
 namespace Horde.Build.Perforce
 {
@@ -16,6 +14,15 @@ namespace Horde.Build.Perforce
 	/// </summary>
 	public interface ICommitCollection
 	{
+		/// <summary>
+		/// Creates a new change
+		/// </summary>
+		/// <param name="path">Path to modify in the change</param>
+		/// <param name="description">Description of the change</param>
+		/// <param name="cancellationToken">Cancellation token for the operation</param>
+		/// <returns>New commit information</returns>
+		Task<int> CreateNewAsync(string path, string description, CancellationToken cancellationToken = default);
+
 		/// <summary>
 		/// Gets the latest change for a particular stream
 		/// </summary>
@@ -30,16 +37,10 @@ namespace Horde.Build.Perforce
 		/// <param name="minChange">The minimum changelist number</param>
 		/// <param name="maxChange">The maximum changelist number</param>
 		/// <param name="maxResults">Maximum number of results to return</param>
+		/// <param name="tags">Tags for the commits to return</param>
 		/// <param name="cancellationToken">Cancellation token for the operation</param>
 		/// <returns>Changelist information</returns>
-		IAsyncEnumerable<ICommit> FindAsync(int? minChange, int? maxChange, int? maxResults, CancellationToken cancellationToken = default);
-
-		/// <summary>
-		/// Gets the latest change number
-		/// </summary>
-		/// <param name="cancellationToken">Cancellation token for the operation</param>
-		/// <returns></returns>
-		Task<int> LatestNumberAsync(CancellationToken cancellationToken = default);
+		IAsyncEnumerable<ICommit> FindAsync(int? minChange = null, int? maxChange = null, int? maxResults = null, IReadOnlyList<CommitTag>? tags = null, CancellationToken cancellationToken = default);
 	}
 
 	/// <summary>
@@ -47,15 +48,52 @@ namespace Horde.Build.Perforce
 	/// </summary>
 	public static class CommitCollectionExtensions
 	{
+		/// <summary>
+		/// Creates a new change for a template
+		/// </summary>
+		/// <param name="perforce">The Perforce service instance</param>
+		/// <param name="template">The template being built</param>
+		/// <param name="cancellationToken">Cancellation token for the operation</param>
+		/// <returns>New changelist number</returns>
+		public static Task<int> CreateNewAsync(this ICommitCollection perforce, ITemplate template, CancellationToken cancellationToken)
+		{
+			string description = (template.SubmitDescription ?? "[Horde] New change for $(TemplateName)").Replace("$(TemplateName)", template.Name, StringComparison.OrdinalIgnoreCase);
+			return perforce.CreateNewAsync(template.SubmitNewChange!, description, cancellationToken);
+		}
+
+		/// <summary>
+		/// Gets the last code code equal or before the given change number
+		/// </summary>
+		/// <param name="source">The commit source to query</param>
+		/// <param name="maxChange">Maximum code change to query</param>
+		/// <param name="cancellationToken">Cancellation token for the operation</param>
+		/// <returns>The last code change</returns>
+		public static async ValueTask<ICommit?> GetLastCodeChange(this ICommitCollection source, int? maxChange, CancellationToken cancellationToken = default)
+		{
+			return await source.FindAsync(null, maxChange, 1, new[] { CommitTag.Code }, cancellationToken).FirstOrDefaultAsync(cancellationToken);
+		}
+
+		/// <summary>
+		/// Finds the latest commit from a source
+		/// </summary>
+		/// <param name="source">The commit source to query</param>
+		/// <param name="cancellationToken">Cancellation token for the operation</param>
+		/// <returns>The latest commit</returns>
+		public static async Task<ICommit> GetLatestAsync(this ICommitCollection source, CancellationToken cancellationToken = default)
+		{
+			return await source.FindAsync(null, null, 1, null, cancellationToken).FirstAsync(cancellationToken);
+		}
+
 		/// <summary>
 		/// Finds the latest commit from a source
 		/// </summary>
 		/// <param name="source">The commit source to query</param>
 		/// <param name="cancellationToken">Cancellation token for the operation</param>
 		/// <returns>The latest commit</returns>
-		public static async Task<ICommit> LatestAsync(this ICommitCollection source, CancellationToken cancellationToken)
+		public static async Task<int> GetLatestNumberAsync(this ICommitCollection source, CancellationToken cancellationToken = default)
 		{
-			return await source.FindAsync(null, null, 1, cancellationToken).FirstAsync(cancellationToken);
+			ICommit commit = await GetLatestAsync(source, cancellationToken);
+			return commit.Number;
 		}
 	}
 }
diff --git a/Engine/Source/Programs/Horde/Horde.Build/Perforce/IPerforceService.cs b/Engine/Source/Programs/Horde/Horde.Build/Perforce/IPerforceService.cs
--- a/Engine/Source/Programs/Horde/Horde.Build/Perforce/IPerforceService.cs
+++ b/Engine/Source/Programs/Horde/Horde.Build/Perforce/IPerforceService.cs
@@ -84,37 +84,6 @@ namespace Horde.Build.Perforce
 		/// <returns></returns>
 		public ValueTask<IUser> FindOrAddUserAsync(string clusterName, string userName, CancellationToken cancellationToken = default);
 
-		/// <summary>
-		/// Create a new changelist by submitting the given file
-		/// </summary>
-		/// <param name="clusterName">Name of the Perforce cluster</param>
-		/// <param name="streamName">The stream to query</param>
-		/// <param name="path">Path for the file to submit</param>
-		/// <param name="description">Description for the changelist</param>
-		/// <param name="cancellationToken">Cancellation token for the operation</param>
-		/// <returns>New changelist number</returns>
-		public Task<int> CreateNewChangeAsync(string clusterName, string streamName, string path, string description, CancellationToken cancellationToken = default);
-
-		/// <summary>
-		/// Gets the code change corresponding to an actual change submitted to a stream
-		/// </summary>
-		/// <param name="stream">Stream to query</param>
-		/// <param name="change">The changelist number to query</param>
-		/// <param name="cancellationToken">Cancellation token for the operation</param>
-		/// <returns>Code change for the latest change</returns>
-		public Task<int> GetCodeChangeAsync(IStream stream, int change, CancellationToken cancellationToken = default);
-
-		/// <summary>
-		/// Gets the latest change for a particular stream
-		/// </summary>
-		/// <param name="stream">Stream to query</param>
-		/// <param name="minChange">The minimum changelist number</param>
-		/// <param name="maxChange">The maximum changelist number</param>
-		/// <param name="results">Number of results to return</param>
-		/// <param name="cancellationToken">Cancellation token for the operation</param>
-		/// <returns>Latest changelist number</returns>
-		public Task<List<ICommit>> GetChangesAsync(IStream stream, int? minChange, int? maxChange, int? results, CancellationToken cancellationToken = default);
-
 		/// <summary>
 		/// Checks a shelf is valid for the given stream
 		/// </summary>
@@ -134,15 +103,6 @@ namespace Horde.Build.Perforce
 		/// <returns>Commit details</returns>
 		public Task<ICommit> GetChangeDetailsAsync(IStream stream, int changeNumber, CancellationToken cancellationToken = default);
 
-		/// <summary>
-		/// Gets the latest change for a particular stream
-		/// </summary>
-		/// <param name="stream">Stream to query</param>
-		/// <param name="changeNumbers">Change numbers to query</param>
-		/// <param name="cancellationToken">Cancellation token for the operation</param>
-		/// <returns>Commit details</returns>
-		public Task<List<ICommit>> GetChangeDetailsAsync(IStream stream, IReadOnlyList<int> changeNumbers, CancellationToken cancellationToken = default);
-
 		/// <summary>
 		/// Duplicates a shelved changelist
 		/// </summary>
@@ -188,76 +148,4 @@ namespace Horde.Build.Perforce
 		/// <returns>Commit source instance</returns>
 		public ICommitCollection GetCommits(IStream stream);
 	}
-
-	/// <summary>
-	/// Extension methods for IPerforceService implementations
-	/// </summary>
-	public static class PerforceServiceExtensions
-	{
-		/// <summary>
-		/// Creates a new change for a template
-		/// </summary>
-		/// <param name="perforce">The Perforce service instance</param>
-		/// <param name="stream">Stream containing the template</param>
-		/// <param name="template">The template being built</param>
-		/// <returns>New changelist number</returns>
-		public static Task<int> CreateNewChangeForTemplateAsync(this IPerforceService perforce, IStream stream, ITemplate template)
-		{
-			string description = (template.SubmitDescription ?? "[Horde] New change for $(TemplateName)").Replace("$(TemplateName)", template.Name, StringComparison.OrdinalIgnoreCase);
-
-			Match match = Regex.Match(template.SubmitNewChange!, @"^(//[^/]+/[^/]+)/(.+)$");
-			if (match.Success)
-			{
-				return perforce.CreateNewChangeAsync(stream.Config.ClusterName, match.Groups[1].Value, match.Groups[2].Value, description);
-			}
-			else
-			{
-				return perforce.CreateNewChangeAsync(stream.Config.ClusterName, stream.Name, template.SubmitNewChange!, description);
-			}
-		}
-
-		/// <summary>
-		/// Gets the latest change for a particular stream
-		/// </summary>
-		/// <param name="perforce">The perforce implementation</param>
-		/// <param name="stream">The stream to query</param>
-		/// <param name="changeNumber">Change number to query</param>
-		/// <returns>Commit details</returns>
-		public static async Task<ICommit> GetChangeDetailsAsync(this IPerforceService perforce, IStream stream, int changeNumber)
-		{
-			List<ICommit> results = await perforce.GetChangeDetailsAsync(stream, new[] { changeNumber });
-			return results[0];
-		}
-
-		/// <summary>
-		/// Gets the latest change for a particular stream
-		/// </summary>
-		/// <param name="perforce">The perforce implementation</param>
-		/// <param name="stream">The stream to query</param>
-		/// <param name="minChange">The minimum changelist number</param>
-		/// <param name="maxChange">The maximum changelist number</param>
-		/// <param name="results">Number of results to return</param>
-		/// <returns>Commit details</returns>
-		public static async Task<List<ICommit>> GetChangeDetailsAsync(this IPerforceService perforce, IStream stream, int? minChange, int? maxChange, int results)
-		{
-			List<ICommit> changes = await perforce.GetChangesAsync(stream, minChange, maxChange, results);
-			return await perforce.GetChangeDetailsAsync(stream, changes.ConvertAll(x => x.Number));
-		}
-
-		/// <summary>
-		/// Get the latest submitted change to the stream
-		/// </summary>
-		/// <param name="perforce">The perforce implementation</param>
-		/// <param name="stream">The stream to query</param>
-		/// <returns>Latest changelist number</returns>
-		public static async Task<int> GetLatestChangeAsync(this IPerforceService perforce, IStream stream)
-		{
-			List<ICommit> changes = await perforce.GetChangesAsync(stream, null, null, 1);
-			if (changes.Count == 0)
-			{
-				throw new Exception($"No changes have been submitted to stream {stream.Name}");
-			}
-			return changes[0].Number;
-		}
-	}
 }
diff --git a/Engine/Source/Programs/Horde/Horde.Build/Perforce/PerforceService.cs b/Engine/Source/Programs/Horde/Horde.Build/Perforce/PerforceService.cs
--- a/Engine/Source/Programs/Horde/Horde.Build/Perforce/PerforceService.cs
+++ b/Engine/Source/Programs/Horde/Horde.Build/Perforce/PerforceService.cs
@@ -1054,42 +1054,6 @@ namespace Horde.Build.Perforce
 			}
 		}
 
-		/// <inheritdoc/>
-		public async Task<int> GetCodeChangeAsync(IStream stream, int change, CancellationToken cancellationToken)
-		{
-			using IScope scope = GlobalTracer.Instance.BuildSpan("PerforceService.GetCodeChangeAsync").StartActive();
-			scope.Span.SetTag("ClusterName", stream.Config.ClusterName);
-			scope.Span.SetTag("StreamName", stream.Name);
-			scope.Span.SetTag("Change", change);
-			
-			int maxChange = change;
-			for (; ; )
-			{
-				// Query for the changes before this point
-				List<ICommit> changes = await GetChangesAsync(stream, null, maxChange, 10, cancellationToken);
-				_logger.LogInformation("Finding last code change in {Stream} before {MaxChange}: {NumResults}", stream.Name, maxChange, changes.Count);
-				if (changes.Count == 0)
-				{
-					return 0;
-				}
-
-				// Get the details for them
-				List<ICommit> detailsList = await GetChangeDetailsAsync(stream, changes.ConvertAll(x => x.Number), cancellationToken);
-				foreach (ICommit details in detailsList.OrderByDescending(x => x.Number))
-				{
-					IReadOnlyList<CommitTag> commitTags = await details.GetTagsAsync(cancellationToken);
-					_logger.LogInformation("Change {Change} = {Flags}", details.Number, String.Join(", ", commitTags.Select(x => x.ToString())));
-					if (commitTags.Contains(CommitTag.Code))
-					{
-						return details.Number;
-					}
-				}
-
-				// Loop round again, adjusting our maximum changelist number
-				maxChange = changes.Min(x => x.Number) - 1;
-			}
-		}
-
 		/// <summary>
 		/// Gets the wildcard filter for a particular range of changes
 		/// </summary>
@@ -1160,32 +1124,73 @@ namespace Horde.Build.Perforce
 		/// </summary>
 		protected class CommitSource : ICommitCollection
 		{
-			protected readonly IPerforceService _perforceService;
+			protected readonly PerforceService _perforceService;
 			protected readonly IStream _stream;
 
-			public CommitSource(IPerforceService perforceService, IStream stream)
+			public CommitSource(PerforceService perforceService, IStream stream)
 			{
 				_perforceService = perforceService;
 				_stream = stream;
 			}
 
-			public virtual async IAsyncEnumerable<ICommit> FindAsync(int? minChange, int? maxChange, int? maxResults, [EnumeratorCancellation] CancellationToken cancellationToken = default)
+			public async Task<int> CreateNewAsync(string path, string description, CancellationToken cancellationToken = default)
 			{
-				List<ICommit> commits = await _perforceService.GetChangesAsync(_stream, minChange, maxChange, maxResults, cancellationToken);
-				foreach (ICommit commit in commits)
+				Match match = Regex.Match(path, @"^(//[^/]+/[^/]+)/(.+)$");
+				if (match.Success)
 				{
-					yield return commit;
+					return await _perforceService.CreateNewChangeAsync(_stream.Config.ClusterName, match.Groups[1].Value, match.Groups[2].Value, description, cancellationToken);
+				}
+				else
+				{
+					return await _perforceService.CreateNewChangeAsync(_stream.Config.ClusterName, _stream.Name, path, description, cancellationToken);
 				}
 			}
 
-			public virtual Task<ICommit> GetAsync(int changeNumber, CancellationToken cancellationToken = default)
+			public virtual async IAsyncEnumerable<ICommit> FindAsync(int? minChange, int? maxChange, int? maxResults, IReadOnlyList<CommitTag>? tags, [EnumeratorCancellation] CancellationToken cancellationToken = default)
 			{
-				return _perforceService.GetChangeDetailsAsync(_stream, changeNumber, cancellationToken);
+				int maxBatchResults = 20;
+				if (tags == null && maxResults != null)
+				{
+					maxBatchResults = Math.Min(maxBatchResults, maxResults.Value);
+				}
+
+				int numResults = 0;
+				while (maxResults == null || numResults < maxResults.Value)
+				{
+					List<ICommit> commits = await _perforceService.GetChangesAsync(_stream, minChange, maxChange, maxBatchResults, cancellationToken);
+					if (commits.Count == 0)
+					{
+						break;
+					}
+
+					foreach (ICommit commit in commits)
+					{
+						if (tags != null)
+						{
+							IReadOnlyList<CommitTag> commitTags = await commit.GetTagsAsync(cancellationToken);
+							if (!tags.Any(x => commitTags.Contains(x)))
+							{
+								continue;
+							}
+						}
+
+						yield return commit;
+
+						numResults++;
+
+						if (maxResults != null && numResults >= maxResults.Value)
+						{
+							break;
+						}
+					}
+
+					maxChange = commits[^1].Number;
+				}
 			}
 
-			public virtual Task<int> LatestNumberAsync(CancellationToken cancellationToken = default)
+			public virtual Task<ICommit> GetAsync(int changeNumber, CancellationToken cancellationToken = default)
 			{
-				return _perforceService.GetLatestChangeAsync(_stream);
+				return _perforceService.GetChangeDetailsAsync(_stream, changeNumber, cancellationToken);
 			}
 		}
 
diff --git a/Engine/Source/Programs/Horde/Horde.Build/Perforce/PerforceServiceCache.cs b/Engine/Source/Programs/Horde/Horde.Build/Perforce/PerforceServiceCache.cs
--- a/Engine/Source/Programs/Horde/Horde.Build/Perforce/PerforceServiceCache.cs
+++ b/Engine/Source/Programs/Horde/Horde.Build/Perforce/PerforceServiceCache.cs
@@ -445,7 +445,7 @@ namespace Horde.Build.Perforce
 				_owner = owner;
 			}
 
-			public override async IAsyncEnumerable<ICommit> FindAsync(int? minChange, int? maxChange, int? maxResults, [EnumeratorCancellation] CancellationToken cancellationToken = default)
+			public override async IAsyncEnumerable<ICommit> FindAsync(int? minChange, int? maxChange, int? maxResults, IReadOnlyList<CommitTag>? tags, [EnumeratorCancellation] CancellationToken cancellationToken = default)
 			{
 				CacheState state = await _owner._mongoService.GetSingletonAsync<CacheState>();
 				if (state.Clusters.TryGetValue(_stream.Config.ClusterName, out ClusterState? clusterState))
@@ -455,6 +455,18 @@ namespace Horde.Build.Perforce
 					{
 						FilterDefinition<CachedCommitDoc> filter = Builders<CachedCommitDoc>.Filter.Eq(x => x.StreamId, _stream.Id);
 
+						if (tags != null && tags.Count > 0)
+						{
+							if (tags.Count == 1)
+							{
+								filter &= Builders<CachedCommitDoc>.Filter.AnyEq(x => x.CommitTags, tags[0]);
+							}
+							else
+							{
+								filter &= Builders<CachedCommitDoc>.Filter.AnyIn(x => x.CommitTags, tags);
+							}
+						}
+
 						if (maxChange != null)
 						{
 							filter &= Builders<CachedCommitDoc>.Filter.Lte(x => x.Number, Math.Min(maxChange.Value, clusterState.MaxChange));
@@ -496,7 +508,7 @@ namespace Horde.Build.Perforce
 
 				if (maxResults == null || maxResults > 0)
 				{
-					await foreach (ICommit commit in base.FindAsync(minChange, maxChange, maxResults, cancellationToken))
+					await foreach (ICommit commit in base.FindAsync(minChange, maxChange, maxResults, tags, cancellationToken))
 					{
 						yield return commit;
 					}
@@ -513,16 +525,6 @@ namespace Horde.Build.Perforce
 				}
 				return await base.GetAsync(changeNumber, cancellationToken);
 			}
-
-			public override async Task<int> LatestNumberAsync(CancellationToken cancellationToken = default)
-			{
-				CacheState state = await _owner._mongoService.GetSingletonAsync<CacheState>();
-				if (state.Clusters.TryGetValue(_stream.Config.ClusterName, out ClusterState? clusterState))
-				{
-					return clusterState.MaxChange;
-				}
-				return await base.LatestNumberAsync(cancellationToken);
-			}
 		}
 
 		/// <inheritdoc/>
