Horde: Missing file.

diff --git a/Engine/Source/Programs/Horde/Horde.Build.Tests/Stubs/Services/PerforceServiceStub.cs b/Engine/Source/Programs/Horde/Horde.Build.Tests/Stubs/Services/PerforceServiceStub.cs
--- a/Engine/Source/Programs/Horde/Horde.Build.Tests/Stubs/Services/PerforceServiceStub.cs
+++ b/Engine/Source/Programs/Horde/Horde.Build.Tests/Stubs/Services/PerforceServiceStub.cs
@@ -231,6 +231,11 @@ namespace Horde.Build.Tests.Stubs.Services
 				_stream = stream;
 			}
 
+			public Task<int> CreateNewAsync(string path, string description, CancellationToken cancellationToken = default)
+			{
+				throw new NotImplementedException();
+			}
+
 			public async IAsyncEnumerable<ICommit> FindAsync(int? minChange, int? maxChange, int? maxResults, IReadOnlyList<CommitTag>? tags, [EnumeratorCancellation] CancellationToken cancellationToken = default)
 			{
 				foreach (ICommit commit in await _owner.GetChangesAsync(_stream, minChange, maxChange, null, cancellationToken))
