Horde: Add tests for schema generation, and fix current serialization issues.

diff --git a/Engine/Source/Programs/Horde/Horde.Build.Tests/SchemaTests.cs b/Engine/Source/Programs/Horde/Horde.Build.Tests/SchemaTests.cs
new file mode 100644
--- /dev/null
+++ b/Engine/Source/Programs/Horde/Horde.Build.Tests/SchemaTests.cs
@@ -0,0 +1,27 @@
+// Copyright Epic Games, Inc. All Rights Reserved.
+
+using EpicGames.Core;
+using Horde.Build.Streams;
+using Horde.Build.Utilities;
+using Microsoft.VisualStudio.TestTools.UnitTesting;
+
+namespace Horde.Build.Tests
+{
+	[TestClass]
+	public class SchemaTests
+	{
+		[TestMethod]
+		public void StreamSchema()
+		{
+			JsonSchema schema = Schemas.CreateSchema(typeof(StreamConfig));
+			_ = schema;
+		}
+
+		[TestMethod]
+		public void ProjectSchema()
+		{
+			JsonSchema schema = Schemas.CreateSchema(typeof(StreamConfig));
+			_ = schema;
+		}
+	}
+}
diff --git a/Engine/Source/Programs/Horde/Horde.Build/Perforce/CommitTag.cs b/Engine/Source/Programs/Horde/Horde.Build/Perforce/CommitTag.cs
--- a/Engine/Source/Programs/Horde/Horde.Build/Perforce/CommitTag.cs
+++ b/Engine/Source/Programs/Horde/Horde.Build/Perforce/CommitTag.cs
@@ -5,6 +5,7 @@ using System.Collections.Generic;
 using System.Linq;
 using System.Text.Json;
 using System.Text.Json.Serialization;
+using EpicGames.Core;
 using EpicGames.Horde;
 using EpicGames.Perforce;
 using MongoDB.Bson.Serialization;
@@ -16,6 +17,7 @@ namespace Horde.Build.Perforce
 	/// <summary>
 	/// Constants for known commit tags
 	/// </summary>
+	[JsonSchemaString]
 	[JsonConverter(typeof(CommitTagJsonConverter))]
 	[BsonSerializer(typeof(CommitTagBsonSerializer))]
 	public readonly struct CommitTag : IEquatable<CommitTag>
diff --git a/Engine/Source/Programs/Shared/EpicGames.Core/JsonSchema.cs b/Engine/Source/Programs/Shared/EpicGames.Core/JsonSchema.cs
--- a/Engine/Source/Programs/Shared/EpicGames.Core/JsonSchema.cs
+++ b/Engine/Source/Programs/Shared/EpicGames.Core/JsonSchema.cs
@@ -360,6 +360,10 @@ namespace EpicGames.Core
 			{
 				return new JsonSchemaString(JsonSchemaStringFormat.DateTime);
 			}
+			if (type == typeof(TimeSpan))
+			{
+				return new JsonSchemaString();
+			}
 
 			Type[] interfaceTypes = type.GetInterfaces();
 			foreach (Type interfaceType in interfaceTypes)
