Horde: Add issue triage aliases for specific issue types
#rnx

diff --git a/Engine/Source/Programs/Horde/Horde.Build/Issues/WorkflowConfig.cs b/Engine/Source/Programs/Horde/Horde.Build/Issues/WorkflowConfig.cs
--- a/Engine/Source/Programs/Horde/Horde.Build/Issues/WorkflowConfig.cs
+++ b/Engine/Source/Programs/Horde/Horde.Build/Issues/WorkflowConfig.cs
@@ -65,6 +65,11 @@ namespace Horde.Build.Issues
 		/// </summary>
 		public string? TriageAlias { get; set; }
 
+		/// <summary>
+		/// Slack user/alias to ping for specific issue types (such as Systemic), if there is nobody assigned to an issue by default.
+		/// </summary>
+		public Dictionary<string, string>? TriageTypeAliases { get; set; }
+
 		/// <summary>
 		/// Alias to ping if an issue has not been resolved for a certain amount of time
 		/// </summary>
diff --git a/Engine/Source/Programs/Horde/Horde.Build/Notifications/Sinks/SlackNotificationSink.cs b/Engine/Source/Programs/Horde/Horde.Build/Notifications/Sinks/SlackNotificationSink.cs
--- a/Engine/Source/Programs/Horde/Horde.Build/Notifications/Sinks/SlackNotificationSink.cs
+++ b/Engine/Source/Programs/Horde/Horde.Build/Notifications/Sinks/SlackNotificationSink.cs
@@ -973,10 +973,20 @@ namespace Horde.Build.Notifications.Sinks
 					}
 				}
 
-				if (workflow.TriageAlias != null && issue.OwnerId == null && (suspects.All(x => x.DeclinedAt != null) || notifyTriageAlias) && !closed)
+				if ((workflow.TriageAlias != null || workflow.TriageTypeAliases != null) && issue.OwnerId == null && (suspects.All(x => x.DeclinedAt != null) || notifyTriageAlias) && !closed)
 				{
-					string triageMessage = $"(cc {FormatUserOrGroupMention(workflow.TriageAlias)} for triage).";
-					await SendOrUpdateMessageAsync(triageChannel, state.Ts, eventId + "_triage", null, triageMessage);
+					string? triageAlias;
+					
+					if (workflow.TriageTypeAliases == null || issue.Fingerprints.Count == 0 || !workflow.TriageTypeAliases.TryGetValue(issue.Fingerprints[0].Type, out triageAlias))
+					{
+						triageAlias = workflow.TriageAlias;
+					}
+
+					if (triageAlias != null)
+					{
+						string triageMessage = $"(cc {FormatUserOrGroupMention(triageAlias)} for triage).";
+						await SendOrUpdateMessageAsync(triageChannel, state.Ts, eventId + "_triage", null, triageMessage);
+					}
 				}
 
 				if (issue.AcknowledgedAt != null)
