Horde: Omit status of an issue from reports if the issue is already resolved.

diff --git a/Engine/Source/Programs/Horde/Horde.Build/Notifications/Sinks/SlackNotificationSink.cs b/Engine/Source/Programs/Horde/Horde.Build/Notifications/Sinks/SlackNotificationSink.cs
--- a/Engine/Source/Programs/Horde/Horde.Build/Notifications/Sinks/SlackNotificationSink.cs
+++ b/Engine/Source/Programs/Horde/Horde.Build/Notifications/Sinks/SlackNotificationSink.cs
@@ -1627,12 +1627,7 @@ namespace Horde.Build.Notifications.Sinks
 					body.Append('\n');
 				}
 
-				string text = await FormatIssueAsync(issue, span, channel, reportTime);
-				if (!open)
-				{
-					text = $"~{text}~";
-				}
-
+				string text = await FormatIssueAsync(issue, span, channel, reportTime, open);
 				body.Append(text);
 			}
 
@@ -1647,7 +1642,7 @@ namespace Horde.Build.Notifications.Sinks
 			return (severity == IssueSeverity.Warning) ? _settings.SlackWarningPrefix : _settings.SlackErrorPrefix;
 		}
 
-		async ValueTask<string> FormatIssueAsync(IIssue issue, IIssueSpan? span, string? triageChannel, DateTime reportTime)
+		async ValueTask<string> FormatIssueAsync(IIssue issue, IIssueSpan? span, string? triageChannel, DateTime reportTime, bool open)
 		{
 			Uri issueUrl = _settings.DashboardUrl;
 			if (span != null)
@@ -1718,9 +1713,17 @@ namespace Horde.Build.Notifications.Sinks
 			{
 				body.Append($" ({FormatExternalIssue(issue.ExternalIssueKey)})");
 			}
+			if (open)
+			{
+				body.Append($" - {status}");
+			}
 
-			body.Append($" - {status}");
-			return body.ToString();
+			string text = body.ToString();
+			if (!open)
+			{
+				text = $"~{text}~";
+			}
+			return text;
 		}
 
 		static string FormatReadableTimeSpan(TimeSpan timeSpan)
