Horde: Remove redundant config information from project documents in the DB.

diff --git a/Engine/Source/Programs/Horde/Horde.Build/Configuration/ConfigUpdateService.cs b/Engine/Source/Programs/Horde/Horde.Build/Configuration/ConfigUpdateService.cs
--- a/Engine/Source/Programs/Horde/Horde.Build/Configuration/ConfigUpdateService.cs
+++ b/Engine/Source/Programs/Horde/Horde.Build/Configuration/ConfigUpdateService.cs
@@ -182,7 +182,7 @@ namespace Horde.Build.Configuration
 				}
 
 				IProject? project = projects.FirstOrDefault(x => x.Id == projectRef.Id);
-				bool update = project == null || project.ConfigPath != projectPath.ToString() || project.ConfigRevision != revision;
+				bool update = project == null || project.ConfigRevision != revision;
 
 				ProjectConfig? projectConfig;
 				if (!update && prevCachedProjectConfigs.TryGetValue(projectRef.Id, out (ProjectConfig Config, string Revision) result) && result.Revision == revision)
diff --git a/Engine/Source/Programs/Horde/Horde.Build/Projects/IProject.cs b/Engine/Source/Programs/Horde/Horde.Build/Projects/IProject.cs
--- a/Engine/Source/Programs/Horde/Horde.Build/Projects/IProject.cs
+++ b/Engine/Source/Programs/Horde/Horde.Build/Projects/IProject.cs
@@ -3,90 +3,15 @@
 using System.Collections.Generic;
 using System.Linq;
 using System.Text.RegularExpressions;
-using EpicGames.Core;
 using Horde.Build.Acls;
 using Horde.Build.Streams;
 using Horde.Build.Utilities;
-using MongoDB.Bson.Serialization.Attributes;
 using MongoDB.Driver;
 
 namespace Horde.Build.Projects
 {
 	using ProjectId = StringId<IProject>;
 
-	/// <summary>
-	/// Specifies a category of streams to display on the dashboard
-	/// </summary>
-	public class StreamCategory
-	{
-		/// <summary>
-		/// Name of this group
-		/// </summary>
-		[BsonRequired]
-		public string Name { get; set; }
-
-		/// <summary>
-		/// Row to display this category on
-		/// </summary>
-		public int Row { get; set; }
-
-		/// <summary>
-		/// Whether to show this category on the nav menu
-		/// </summary>
-		public bool ShowOnNavMenu { get; set; } = true;
-
-		/// <summary>
-		/// List of stream name patterns to include for this project
-		/// </summary>
-		public List<string> IncludePatterns { get; set; } = new List<string>();
-
-		/// <summary>
-		/// Patterns for stream names to be excluded
-		/// </summary>
-		public List<string> ExcludePatterns { get; set; } = new List<string>();
-
-		/// <summary>
-		/// Private constructor
-		/// </summary>
-		private StreamCategory()
-		{
-			Name = null!;
-		}
-
-		/// <summary>
-		/// Constructor
-		/// </summary>
-		/// <param name="name">Name of this group</param>
-		public StreamCategory(string name)
-		{
-			Name = name;
-		}
-
-		/// <summary>
-		/// Constructs a category from a request object
-		/// </summary>
-		/// <param name="request">The request object</param>
-		public StreamCategory(CreateProjectCategoryRequest request)
-		{
-			Name = request.Name;
-			Row = request.Row;
-			ShowOnNavMenu = request.ShowOnNavMenu;
-			IncludePatterns = request.IncludePatterns;
-			ExcludePatterns = request.ExcludePatterns;
-		}
-
-		/// <summary>
-		/// Tests if a given name matches a pattern
-		/// </summary>
-		/// <param name="name">The name to test</param>
-		/// <param name="pattern">The pattern to match against</param>
-		/// <returns>True if the pattern matches, false otherwise</returns>
-		public static bool MatchPattern(string name, string pattern)
-		{
-			return Regex.IsMatch(name, pattern);
-		}
-	}
-
 	/// <summary>
 	/// Represents a project
 	/// </summary>
@@ -95,22 +20,17 @@ namespace Horde.Build.Projects
 		/// <summary>
 		/// Identifier for the project.
 		/// </summary>
-		public ProjectId Id { get; }
+		ProjectId Id { get; }
 
 		/// <summary>
 		/// Name of the project
 		/// </summary>
-		public string Name { get; }
-
-		/// <summary>
-		/// Path to the config file used to configure this project
-		/// </summary>
-		public string? ConfigPath { get; }
+		string Name { get; }
 
 		/// <summary>
 		/// Revision of the config file used to configure this project
 		/// </summary>
-		public string? ConfigRevision { get; }
+		string? ConfigRevision { get; }
 
 		/// <summary>
 		/// Order to display on the dashboard
@@ -118,14 +38,14 @@ namespace Horde.Build.Projects
 		public int Order { get; }
 
 		/// <summary>
-		/// The ACL for this object
+		/// Configuration settings for the stream
 		/// </summary>
-		public Acl? Acl { get; }
+		ProjectConfig Config { get; }
 
 		/// <summary>
-		/// List of stream categories for this project. Controls how streams are displayed on the dashboard.
+		/// The ACL for this object
 		/// </summary>
-		public IReadOnlyList<StreamCategory> Categories { get; }
+		Acl? Acl { get; }
 	}
 	
 	/// <summary>
@@ -184,7 +104,7 @@ namespace Horde.Build.Projects
 			List<GetProjectCategoryResponse>? categoryResponses = null;
 			if (bIncludeCategories)
 			{
-				categoryResponses = project.Categories.ConvertAll(x => new GetProjectCategoryResponse(x));
+				categoryResponses = project.Config.Categories.ConvertAll(x => new GetProjectCategoryResponse(x));
 				if (streams != null)
 				{
 					foreach (IStream stream in streams)
@@ -198,7 +118,8 @@ namespace Horde.Build.Projects
 								row++;
 							}
 
-							StreamCategory otherCategory = new StreamCategory("Other");
+							ProjectCategoryConfig otherCategory = new ProjectCategoryConfig();
+							otherCategory.Name = "Other";
 							otherCategory.Row = row;
 							otherCategory.IncludePatterns.Add(".*");
 
@@ -222,9 +143,9 @@ namespace Horde.Build.Projects
 		/// <returns>True if the category matches</returns>
 		static bool MatchCategory(string name, GetProjectCategoryResponse category)
 		{
-			if (category.IncludePatterns.Any(x => StreamCategory.MatchPattern(name, x)))
+			if (category.IncludePatterns.Any(x => Regex.IsMatch(name, x)))
 			{
-				if (!category.ExcludePatterns.Any(x => StreamCategory.MatchPattern(name, x)))
+				if (!category.ExcludePatterns.Any(x => Regex.IsMatch(name, x)))
 				{
 					return true;
 				}
diff --git a/Engine/Source/Programs/Horde/Horde.Build/Projects/ProjectCollection.cs b/Engine/Source/Programs/Horde/Horde.Build/Projects/ProjectCollection.cs
--- a/Engine/Source/Programs/Horde/Horde.Build/Projects/ProjectCollection.cs
+++ b/Engine/Source/Programs/Horde/Horde.Build/Projects/ProjectCollection.cs
@@ -6,8 +6,10 @@ using System.Diagnostics.CodeAnalysis;
 using System.Linq;
 using System.Threading.Tasks;
 using Horde.Build.Acls;
+using Horde.Build.Configuration;
 using Horde.Build.Server;
 using Horde.Build.Utilities;
+using Microsoft.Extensions.Logging;
 using MongoDB.Bson.Serialization.Attributes;
 using MongoDB.Driver;
 
@@ -31,31 +33,43 @@ namespace Horde.Build.Projects
 			public ProjectId Id { get; set; }
 
 			[BsonRequired]
-			public string Name { get; set; }
-
-			public string? ConfigPath { get; set; }
-			public string? ConfigRevision { get; set; }
+			public string ConfigRevision { get; set; } = null!;
 
 			public int Order { get; set; } = DefaultOrder;
 			public Acl? Acl { get; set; }
-			public List<StreamCategory> Categories { get; set; } = new List<StreamCategory>();
 
 			[BsonIgnoreIfDefault, BsonDefaultValue(false)]
 			public bool Deleted { get; set; }
 
-			IReadOnlyList<StreamCategory> IProject.Categories => Categories;
+			#region IProject interface
 
-			[BsonConstructor]
-			private ProjectDocument()
-			{
-				Name = null!;
-			}
+			[BsonIgnore]
+			public string Name => Config?.Name!;
+
+			[BsonIgnore]
+			public ProjectConfig Config { get; private set; } = null!;
 
-			public ProjectDocument(ProjectId id, string name)
+			public async Task PostLoadAsync(ConfigCollection configCollection, ILogger logger)
 			{
-				Id = id;
-				Name = name;
+				try
+				{
+					if (Deleted)
+					{
+						Config = new ProjectConfig();
+					}
+					else
+					{
+						Config = await configCollection.GetConfigAsync<ProjectConfig>(ConfigRevision);
+					}
+				}
+				catch (Exception)
+				{
+					logger.LogError("Unable to get stream config for {StreamId} at {Revision}", Id, ConfigRevision);
+					throw;
+				}
 			}
+
+			#endregion
 		}
 
 		/// <summary>
@@ -89,51 +103,56 @@ namespace Horde.Build.Projects
 			public static readonly ProjectionDefinition<ProjectDocument> Projection = Builders<ProjectDocument>.Projection.Include(x => x.Acl);
 		}
 
-		/// <summary>
-		/// Collection of project documents
-		/// </summary>
 		readonly IMongoCollection<ProjectDocument> _projects;
-
-		/// <summary>
-		/// Collection of project logo documents
-		/// </summary>
 		readonly IMongoCollection<ProjectLogoDocument> _projectLogos;
+		readonly ConfigCollection _configCollection;
+		readonly ILogger _logger;
 
 		/// <summary>
 		/// Constructor
 		/// </summary>
-		/// <param name="mongoService">The database service instance</param>
-		public ProjectCollection(MongoService mongoService)
+		public ProjectCollection(MongoService mongoService, ConfigCollection configCollection, ILogger<ProjectCollection> logger)
 		{
 			_projects = mongoService.GetCollection<ProjectDocument>("Projects");
 			_projectLogos = mongoService.GetCollection<ProjectLogoDocument>("ProjectLogos");
+			_configCollection = configCollection;
+			_logger = logger;
 		}
 
 		/// <inheritdoc/>
 		public async Task<IProject?> AddOrUpdateAsync(ProjectId id, string configPath, string revision, int order, ProjectConfig config)
 		{
-			ProjectDocument newProject = new ProjectDocument(id, config.Name);
-			newProject.ConfigPath = configPath;
+			ProjectDocument newProject = new ProjectDocument();
+			newProject.Id = id;
 			newProject.ConfigRevision = revision;
 			newProject.Order = order;
-			newProject.Categories = config.Categories.ConvertAll(x => new StreamCategory(x));
 			newProject.Acl = Acl.Merge(new Acl(), config.Acl);
 
-			await _projects.FindOneAndReplaceAsync<ProjectDocument>(x => x.Id == id, newProject, new FindOneAndReplaceOptions<ProjectDocument> { IsUpsert = true });
-			return newProject;
+			ProjectDocument result = await _projects.FindOneAndReplaceAsync<ProjectDocument>(x => x.Id == id, newProject, new FindOneAndReplaceOptions<ProjectDocument> { IsUpsert = true });
+			await result.PostLoadAsync(_configCollection, _logger);
+			return result;
 		}
 
 		/// <inheritdoc/>
 		public async Task<List<IProject>> FindAllAsync()
 		{
 			List<ProjectDocument> results = await _projects.Find(x => !x.Deleted).ToListAsync();
+			foreach (ProjectDocument result in results)
+			{
+				await result.PostLoadAsync(_configCollection, _logger);
+			}
 			return results.OrderBy(x => x.Order).ThenBy(x => x.Name).Select<ProjectDocument, IProject>(x => x).ToList();
 		}
 
 		/// <inheritdoc/>
 		public async Task<IProject?> GetAsync(ProjectId projectId)
 		{
-			return await _projects.Find<ProjectDocument>(x => x.Id == projectId).FirstOrDefaultAsync();
+			ProjectDocument? project = await _projects.Find<ProjectDocument>(x => x.Id == projectId).FirstOrDefaultAsync();
+			if (project != null)
+			{
+				await project.PostLoadAsync(_configCollection, _logger);
+			}
+			return project;
 		}
 
 		/// <inheritdoc/>
diff --git a/Engine/Source/Programs/Horde/Horde.Build/Projects/ProjectConfig.cs b/Engine/Source/Programs/Horde/Horde.Build/Projects/ProjectConfig.cs
--- a/Engine/Source/Programs/Horde/Horde.Build/Projects/ProjectConfig.cs
+++ b/Engine/Source/Programs/Horde/Horde.Build/Projects/ProjectConfig.cs
@@ -1,7 +1,9 @@
 // Copyright Epic Games, Inc. All Rights Reserved.
 
+using System;
 using System.Collections.Generic;
 using System.ComponentModel.DataAnnotations;
+using System.Text.RegularExpressions;
 using EpicGames.Core;
 using Horde.Build.Acls;
 using Horde.Build.Streams;
@@ -32,7 +34,7 @@ namespace Horde.Build.Projects
 		/// <summary>
 		/// Categories to include in this project
 		/// </summary>
-		public List<CreateProjectCategoryRequest> Categories { get; set; } = new List<CreateProjectCategoryRequest>();
+		public List<ProjectCategoryConfig> Categories { get; set; } = new List<ProjectCategoryConfig>();
 
 		/// <summary>
 		/// List of streams
@@ -62,4 +64,36 @@ namespace Horde.Build.Projects
 		[Required]
 		public string Path { get; set; } = null!;
 	}
+
+	/// <summary>
+	/// Information about a category to display for a stream
+	/// </summary>
+	public class ProjectCategoryConfig
+	{
+		/// <summary>
+		/// Name of this category
+		/// </summary>
+		[Required]
+		public string Name { get; set; } = String.Empty;
+
+		/// <summary>
+		/// Index of the row to display this category on
+		/// </summary>
+		public int Row { get; set; }
+
+		/// <summary>
+		/// Whether to show this category on the nav menu
+		/// </summary>
+		public bool ShowOnNavMenu { get; set; }
+
+		/// <summary>
+		/// Patterns for stream names to include
+		/// </summary>
+		public List<string> IncludePatterns { get; set; } = new List<string>();
+
+		/// <summary>
+		/// Patterns for stream names to exclude
+		/// </summary>
+		public List<string> ExcludePatterns { get; set; } = new List<string>();
+	}
 }
diff --git a/Engine/Source/Programs/Horde/Horde.Build/Projects/ProjectMessages.cs b/Engine/Source/Programs/Horde/Horde.Build/Projects/ProjectMessages.cs
--- a/Engine/Source/Programs/Horde/Horde.Build/Projects/ProjectMessages.cs
+++ b/Engine/Source/Programs/Horde/Horde.Build/Projects/ProjectMessages.cs
@@ -6,46 +6,6 @@ using Horde.Build.Acls;
 
 namespace Horde.Build.Projects
 {
-	/// <summary>
-	/// Information about a category to display for a stream
-	/// </summary>
-	public class CreateProjectCategoryRequest
-	{
-		/// <summary>
-		/// Name of this category
-		/// </summary>
-		[Required]
-		public string Name { get; set; }
-
-		/// <summary>
-		/// Index of the row to display this category on
-		/// </summary>
-		public int Row { get; set; }
-
-		/// <summary>
-		/// Whether to show this category on the nav menu
-		/// </summary>
-		public bool ShowOnNavMenu { get; set; }
-
-		/// <summary>
-		/// Patterns for stream names to include
-		/// </summary>
-		public List<string> IncludePatterns { get; set; } = new List<string>();
-
-		/// <summary>
-		/// Patterns for stream names to exclude
-		/// </summary>
-		public List<string> ExcludePatterns { get; set; } = new List<string>();
-
-		/// <summary>
-		/// Constructor
-		/// </summary>
-		public CreateProjectCategoryRequest(string name)
-		{
-			Name = name;
-		}
-	}
-
 	/// <summary>
 	/// Information about a stream within a project
 	/// </summary>
@@ -112,7 +72,7 @@ namespace Horde.Build.Projects
 		/// Constructor
 		/// </summary>
 		/// <param name="streamCategory">The category to construct from</param>
-		public GetProjectCategoryResponse(StreamCategory streamCategory)
+		public GetProjectCategoryResponse(ProjectCategoryConfig streamCategory)
 		{
 			Name = streamCategory.Name;
 			Row = streamCategory.Row;
