Horde: Test fixes

diff --git a/Engine/Source/Programs/Horde/Horde.Build.Tests/DeviceServiceTests.cs b/Engine/Source/Programs/Horde/Horde.Build.Tests/DeviceServiceTests.cs
--- a/Engine/Source/Programs/Horde/Horde.Build.Tests/DeviceServiceTests.cs
+++ b/Engine/Source/Programs/Horde/Horde.Build.Tests/DeviceServiceTests.cs
@@ -55,7 +55,9 @@ namespace Horde.Build.Tests
 		{
 			DeviceConfig Devices = new DeviceConfig();
 
-			await ProjectCollection.AddOrUpdateAsync(new ProjectId("ue5"), "", "", 0, new ProjectConfig { Name = "UE5" });
+			const string ConfigRevision = "device-test-config";
+			await ConfigCollection.AddConfigAsync(ConfigRevision, new ProjectConfig { Name = "UE5" });
+			await ProjectCollection.AddOrUpdateAsync(new ProjectId("ue5"), ConfigRevision, 0);
 
 			// create 2 pools
 			for (int i = 1; i < 3; i++)
diff --git a/Engine/Source/Programs/Horde/Horde.Build.Tests/IssueServiceTests.cs b/Engine/Source/Programs/Horde/Horde.Build.Tests/IssueServiceTests.cs
--- a/Engine/Source/Programs/Horde/Horde.Build.Tests/IssueServiceTests.cs
+++ b/Engine/Source/Programs/Horde/Horde.Build.Tests/IssueServiceTests.cs
@@ -142,7 +142,9 @@ namespace Horde.Build.Tests
 				_workspaceDir = new DirectoryReference("/Horde");
 			}
 
-			IProject project = ProjectCollection.AddOrUpdateAsync(new ProjectId("ue4"), "", "", 0, new ProjectConfig { Name = "UE4" }).Result!;
+			const string ProjectConfigRevision = "projectconfig";
+			ConfigCollection.AddConfigAsync(ProjectConfigRevision, new ProjectConfig { Name = "UE4" }).Wait();
+			IProject project = ProjectCollection.AddOrUpdateAsync(new ProjectId("ue4"), ProjectConfigRevision, 0).Result!;
 
 			IStream mainStream = CreateStreamAsync(project.Id, _mainStreamId, MainStreamName).Result;
 			IStream releaseStream = CreateStreamAsync(project.Id, _releaseStreamId, ReleaseStreamName).Result;
diff --git a/Engine/Source/Programs/Horde/Horde.Build.Tests/JobServiceTests.cs b/Engine/Source/Programs/Horde/Horde.Build.Tests/JobServiceTests.cs
--- a/Engine/Source/Programs/Horde/Horde.Build.Tests/JobServiceTests.cs
+++ b/Engine/Source/Programs/Horde/Horde.Build.Tests/JobServiceTests.cs
@@ -34,7 +34,10 @@ namespace Horde.Build.Tests
 		public async Task TestChainedJobs()
 		{
 			ProjectId projectId = new ProjectId("ue5");
-			IProject? project = await ProjectService.Collection.AddOrUpdateAsync(projectId, "", "", 0, new ProjectConfig { Name = "UE5" });
+
+			const string ProjectConfigRevision = "projectconfig";
+			await ConfigCollection.AddConfigAsync(ProjectConfigRevision, new ProjectConfig { Name = "UE4" });
+			IProject? project = await ProjectService.Collection.AddOrUpdateAsync(projectId, ProjectConfigRevision, 0);
 			Assert.IsNotNull(project);
 
 			ITemplate template = await TemplateCollection.AddAsync("Test template");
@@ -158,7 +161,9 @@ namespace Horde.Build.Tests
 			IAgent? agent = await AgentService.CreateAgentAsync("TestAgent", true, null, new List<StringId<IPool>> { new StringId<IPool>("win") });
 			await AgentService.CreateSessionAsync(agent, AgentStatus.Ok, new List<string>(), new Dictionary<string, int>(), null);
 
-			IProject? project = await ProjectService.Collection.AddOrUpdateAsync(new ProjectId("ue5"), "", "", 0, new ProjectConfig { Name = "UE5" });
+			const string ProjectConfigRevision = "projectconfig";
+			await ConfigCollection.AddConfigAsync(ProjectConfigRevision, new ProjectConfig { Name = "UE4" });
+			IProject? project = await ProjectService.Collection.AddOrUpdateAsync(new ProjectId("ue5"), ProjectConfigRevision, 0);
 			Assert.IsNotNull(project);
 
 			StreamId streamId = new StreamId("ue5-main");
diff --git a/Engine/Source/Programs/Horde/Horde.Build.Tests/SchedulerTests.cs b/Engine/Source/Programs/Horde/Horde.Build.Tests/SchedulerTests.cs
--- a/Engine/Source/Programs/Horde/Horde.Build.Tests/SchedulerTests.cs
+++ b/Engine/Source/Programs/Horde/Horde.Build.Tests/SchedulerTests.cs
@@ -41,7 +41,9 @@ namespace Horde.Build.Tests
 		{
 			IUser bob = UserCollection.FindOrAddUserByLoginAsync("Bob").Result;
 
-			IProject ? project = ProjectService.Collection.AddOrUpdateAsync(ProjectId, "", "", 0, new ProjectConfig { Name = "UE5" }).Result;
+			const string ProjectConfigRevision = "projectconfig";
+			ConfigCollection.AddConfigAsync(ProjectConfigRevision, new ProjectConfig { Name = "UE4" }).Wait();
+			IProject? project = ProjectService.Collection.AddOrUpdateAsync(ProjectId, ProjectConfigRevision, 0).Result;
 			Assert.IsNotNull(project);
 
 			_template = TemplateCollection.AddAsync("Test template").Result;
diff --git a/Engine/Source/Programs/Horde/Horde.Build/Configuration/ConfigUpdateService.cs b/Engine/Source/Programs/Horde/Horde.Build/Configuration/ConfigUpdateService.cs
--- a/Engine/Source/Programs/Horde/Horde.Build/Configuration/ConfigUpdateService.cs
+++ b/Engine/Source/Programs/Horde/Horde.Build/Configuration/ConfigUpdateService.cs
@@ -219,7 +219,7 @@ namespace Horde.Build.Configuration
 						if (update)
 						{
 							_logger.LogInformation("Updating configuration for project {ProjectId} ({Revision})", projectRef.Id, revision);
-							await _projectService.Collection.AddOrUpdateAsync(projectRef.Id, projectPath.ToString(), revision, idx, projectConfig);
+							await _projectService.Collection.AddOrUpdateAsync(projectRef.Id, revision, idx);
 						}
 					}
 					catch (Exception ex)
diff --git a/Engine/Source/Programs/Horde/Horde.Build/Projects/IProjectCollection.cs b/Engine/Source/Programs/Horde/Horde.Build/Projects/IProjectCollection.cs
--- a/Engine/Source/Programs/Horde/Horde.Build/Projects/IProjectCollection.cs
+++ b/Engine/Source/Programs/Horde/Horde.Build/Projects/IProjectCollection.cs
@@ -17,12 +17,10 @@ namespace Horde.Build.Projects
 		/// Updates the project configuration
 		/// </summary>
 		/// <param name="id">The project id</param>
-		/// <param name="configPath">Path to the config file used to configure this project</param>
-		/// <param name="revision">The config file revision</param>
+		/// <param name="configRevision">The config file revision. Can be referenced withing the config collection.</param>
 		/// <param name="order">Order of the project</param>
-		/// <param name="config">The configuration</param>
 		/// <returns>New project instance</returns>
-		Task<IProject?> AddOrUpdateAsync(ProjectId id, string configPath, string revision, int order, ProjectConfig config);
+		Task<IProject?> AddOrUpdateAsync(ProjectId id, string configRevision, int order);
 
 		/// <summary>
 		/// Gets all the available projects
diff --git a/Engine/Source/Programs/Horde/Horde.Build/Projects/ProjectCollection.cs b/Engine/Source/Programs/Horde/Horde.Build/Projects/ProjectCollection.cs
--- a/Engine/Source/Programs/Horde/Horde.Build/Projects/ProjectCollection.cs
+++ b/Engine/Source/Programs/Horde/Horde.Build/Projects/ProjectCollection.cs
@@ -120,15 +120,17 @@ namespace Horde.Build.Projects
 		}
 
 		/// <inheritdoc/>
-		public async Task<IProject?> AddOrUpdateAsync(ProjectId id, string configPath, string revision, int order, ProjectConfig config)
+		public async Task<IProject?> AddOrUpdateAsync(ProjectId id, string configRevision, int order)
 		{
+			ProjectConfig config = await _configCollection.GetConfigAsync<ProjectConfig>(configRevision);
+
 			ProjectDocument newProject = new ProjectDocument();
 			newProject.Id = id;
-			newProject.ConfigRevision = revision;
+			newProject.ConfigRevision = configRevision;
 			newProject.Order = order;
 			newProject.Acl = Acl.Merge(new Acl(), config.Acl);
 
-			ProjectDocument result = await _projects.FindOneAndReplaceAsync<ProjectDocument>(x => x.Id == id, newProject, new FindOneAndReplaceOptions<ProjectDocument> { IsUpsert = true });
+			ProjectDocument result = await _projects.FindOneAndReplaceAsync<ProjectDocument>(x => x.Id == id, newProject, new FindOneAndReplaceOptions<ProjectDocument> { IsUpsert = true, ReturnDocument = ReturnDocument.After });
 			await result.PostLoadAsync(_configCollection, _logger);
 			return result;
 		}
