Horde: Hide acknowledge/not me/mark fixed buttons once an issue is closed.

diff --git a/Engine/Source/Programs/Horde/Horde.Build/Notifications/Sinks/SlackNotificationSink.cs b/Engine/Source/Programs/Horde/Horde.Build/Notifications/Sinks/SlackNotificationSink.cs
--- a/Engine/Source/Programs/Horde/Horde.Build/Notifications/Sinks/SlackNotificationSink.cs
+++ b/Engine/Source/Programs/Horde/Horde.Build/Notifications/Sinks/SlackNotificationSink.cs
@@ -870,48 +870,51 @@ namespace Horde.Build.Notifications.Sinks
 						message.AddSection(workflow.TriageInstructions);
 					}
 
-					ActionsBlock actions = message.AddActions();
-					actions.AddButton("Assign to Me", value: $"issue_{issue.Id}_ack", style: ButtonStyle.Primary);
-					actions.AddButton("Not Me", value: $"issue_{issue.Id}_decline", style: ButtonStyle.Danger);
-					actions.AddButton("Mark Fixed", value: $"issue_{issue.Id}_markfixed");
-
-					string? context = null;
-					if (issue.OwnerId != null)
+					if (!closed)
 					{
-						string user = await FormatNameAsync(issue.OwnerId.Value);
-						if (issue.AcknowledgedAt == null)
+						ActionsBlock actions = message.AddActions();
+						actions.AddButton("Assign to Me", value: $"issue_{issue.Id}_ack", style: ButtonStyle.Primary);
+						actions.AddButton("Not Me", value: $"issue_{issue.Id}_decline", style: ButtonStyle.Danger);
+						actions.AddButton("Mark Fixed", value: $"issue_{issue.Id}_markfixed");
+
+						string? context = null;
+						if (issue.OwnerId != null)
 						{
-							context = $"Assigned to {user} (unacknowledged).";
+							string user = await FormatNameAsync(issue.OwnerId.Value);
+							if (issue.AcknowledgedAt == null)
+							{
+								context = $"Assigned to {user} (unacknowledged).";
+							}
+							else
+							{
+								context = $"Acknowledged by {user}.";
+							}
 						}
-						else
+						else if (suspects.Any(x => x.DeclinedAt != null))
 						{
-							context = $"Acknowledged by {user}.";
-						}
-					}
-					else if(suspects.Any(x => x.DeclinedAt != null))
-					{
-						HashSet<UserId> userIds = new HashSet<UserId>();
-						foreach (IIssueSuspect suspect in suspects)
-						{
-							if (suspect.DeclinedAt != null)
+							HashSet<UserId> userIds = new HashSet<UserId>();
+							foreach (IIssueSuspect suspect in suspects)
+							{
+								if (suspect.DeclinedAt != null)
+								{
+									userIds.Add(suspect.AuthorId);
+								}
+							}
+
+							List<string> users = new List<string>();
+							foreach (UserId userId in userIds)
 							{
-								userIds.Add(suspect.AuthorId);
+								users.Add(await FormatNameAsync(userId));
 							}
+							users.Sort(StringComparer.OrdinalIgnoreCase);
+
+							context = $"Declined by {StringUtils.FormatList(users)}.";
 						}
 
-						List<string> users = new List<string>();
-						foreach (UserId userId in userIds)
+						if (context != null)
 						{
-							users.Add(await FormatNameAsync(userId));
+							message.AddContext(context);
 						}
-						users.Sort(StringComparer.OrdinalIgnoreCase);
-
-						context = $"Declined by {StringUtils.FormatList(users)}.";
-					}
-
-					if (context != null)
-					{
-						message.AddContext(context);
 					}
 
 					await SendOrUpdateMessageAsync(triageChannel, state.Ts, $"{eventId}_buttons", null, message);
