Horde: Fix missing audit log message when a user declines an issue.

diff --git a/Engine/Source/Programs/Horde/Horde.Build/Issues/IssueCollection.cs b/Engine/Source/Programs/Horde/Horde.Build/Issues/IssueCollection.cs
--- a/Engine/Source/Programs/Horde/Horde.Build/Issues/IssueCollection.cs
+++ b/Engine/Source/Programs/Horde/Horde.Build/Issues/IssueCollection.cs
@@ -986,6 +986,7 @@ namespace Horde.Build.Issues
 
 			if (newDeclinedById != null)
 			{
+				GetLogger(issue.Id).LogInformation("Declined by {UserId}", newDeclinedById.Value);
 				await _issueSuspects.UpdateManyAsync(x => x.IssueId == issue.Id && x.AuthorId == newDeclinedById.Value, Builders<IssueSuspect>.Update.Set(x => x.DeclinedAt, DateTime.UtcNow));
 			}
 			if (newQuarantinedById != null)
