Horde: Add date range to graph and summary, add device ip's to platform summaries, present issues vs reservations
#rnx

diff --git a/Engine/Source/Programs/Horde/HordeDashboard/src/components/DevicePoolTelemetry.tsx b/Engine/Source/Programs/Horde/HordeDashboard/src/components/DevicePoolTelemetry.tsx
--- a/Engine/Source/Programs/Horde/HordeDashboard/src/components/DevicePoolTelemetry.tsx
+++ b/Engine/Source/Programs/Horde/HordeDashboard/src/components/DevicePoolTelemetry.tsx
@@ -172,7 +172,7 @@ class PoolTelemetryHandler {
 
    getProblemDevices() {
 
-      const platforms: Map<string, { deviceId: string, deviceName: string, problems: number, latestProblem?: string }[]> = new Map();
+      const platforms: Map<string, { deviceId: string, deviceName: string, problems: number, problemsDesc: string, latestProblem?: string }[]> = new Map();
 
       this.deviceTelemetry.forEach((telemetry, deviceId) => {
 
@@ -198,7 +198,7 @@ class PoolTelemetryHandler {
             platforms.set(platform.id, []);
          }
 
-         platforms.get(platform.id)!.push({ deviceId: deviceId, deviceName: device.name, problems: count, latestProblem: latest ? getHumanTime(latest) : "" });
+         platforms.get(platform.id)!.push({ deviceId: deviceId, deviceName: `${device.name} / ${device.address}`, problems: count, problemsDesc: `${count} / ${telemetry.length}`, latestProblem: latest ? getHumanTime(latest) : "" });
 
       });
 
@@ -491,7 +491,7 @@ class PoolTelemetryHandler {
    @observable
    updated: number = 0;
 
-   private data?: PoolTelemetryData[];
+   data?: PoolTelemetryData[];
 
    platforms = new Map<string, GetDevicePlatformResponse>();
    pools = new Map<string, GetDevicePoolResponse>();
@@ -550,7 +550,7 @@ class PoolTelemetryGraph {
 
       const schemaColor = d3.scaleOrdinal()
          .domain(streams)
-         .range( dashboard.darktheme ? d3.schemeDark2 : d3.schemeSet2);
+         .range(dashboard.darktheme ? d3.schemeDark2 : d3.schemeSet2);
 
       const scolors = dashboard.getStatusColors();
       const colors: Record<string, string> = {
@@ -581,7 +581,7 @@ class PoolTelemetryGraph {
          .attr("height", height - margin.top - margin.bottom);
 
       const updateChart = (event: any) => {
-      
+
          if (this.freezeTooltip || (event && !event.sourceEvent)) {
             return;
          }
@@ -1009,8 +1009,8 @@ const SummaryModal: React.FC<{ handler: PoolTelemetryHandler, onClose: () => voi
       platformProblems = platformProblems.filter(p => p.problems > 8);
 
       const problemColumns = [
-         { key: 'column1', name: 'Problem Devices', fieldName: 'deviceName', minWidth: 150, maxWidth: 150 },
-         { key: 'column2', name: 'Issues', fieldName: 'problems', minWidth: 120, maxWidth: 120 },
+         { key: 'column1', name: 'Problem Devices', fieldName: 'deviceName', minWidth: 320, maxWidth: 320 },
+         { key: 'column2', name: 'Issues vs Reservations', fieldName: 'problemsDesc', minWidth: 180, maxWidth: 180 },
          { key: 'column3', name: 'Latest Issue', fieldName: 'latestProblem', minWidth: 120, maxWidth: 120 },
       ];
 
@@ -1045,7 +1045,7 @@ const SummaryModal: React.FC<{ handler: PoolTelemetryHandler, onClose: () => voi
             </Stack>
             <Stack tokens={{ childrenGap: 12 }}>
                {!!platformProblems?.length &&
-                  <Stack style={{ paddingLeft: 12, paddingTop: 12, width: 480 }}>
+                  <Stack style={{ paddingLeft: 12, paddingTop: 12, width: 800 }}>
                      <DetailsList
                         isHeaderVisible={true}
                         items={platformProblems}
@@ -1072,13 +1072,19 @@ const SummaryModal: React.FC<{ handler: PoolTelemetryHandler, onClose: () => voi
       </Stack>
    };
 
-   return <Modal className={hordeClasses.modal} isOpen={true} isBlocking={true} topOffsetFixed={true} styles={{ main: { padding: 8, width: 940, backgroundColor: dashboard.darktheme ? modeColors.content : modeColors.background, hasBeenOpened: false, top: "72px", position: "absolute", height: "820px" } }} onDismiss={() => { onClose() }}>
+   let title = "Device Platform Summary";
 
-      <Stack>
+   if (handler.data && handler.data.length > 1) {
+      const minDate = getHumanTime(handler.data[0].date);
+      const maxDate = getHumanTime(handler.data[handler.data.length - 1].date);
+      title = `Device Platform Summary / ${minDate} - ${maxDate}`;
+   }
 
+   return <Modal className={hordeClasses.modal} isOpen={true} isBlocking={true} topOffsetFixed={true} styles={{ main: { padding: 8, width: 940, backgroundColor: dashboard.darktheme ? modeColors.content : modeColors.background, hasBeenOpened: false, top: "72px", position: "absolute", height: "820px" } }} onDismiss={() => { onClose() }}>
+      <Stack>
          <Stack horizontal verticalAlign="center" style={{ paddingTop: 12 }}>
             <Stack style={{ paddingLeft: 16 }}>
-               <Text variant="mediumPlus" styles={{ root: { fontWeight: "unset", fontFamily: "Horde Open Sans SemiBold" } }}>Device Platform Summary</Text>
+               <Text variant="mediumPlus" styles={{ root: { fontWeight: "unset", fontFamily: "Horde Open Sans SemiBold" } }}>{title}</Text>
             </Stack>
             <Stack grow />
             <Stack horizontalAlign="end" style={{ paddingRight: 12 }}>
@@ -1157,12 +1163,20 @@ export const DevicePoolGraph: React.FC = observer(() => {
    const poolOptions: IDropdownOption[] = pools.map(p => { return { key: `pool_key_${p.id}`, text: p.name, data: p } });
    const platformOptions: IDropdownOption[] = platforms.map(p => { return { key: `platform_key_${p.id}`, text: p.name, data: p } });
 
+   let title = "Device Pool Telemetry";
+
+   if (handler.data && handler.data.length > 1) {
+      const minDate = getHumanTime(handler.data[0].date);
+      const maxDate = getHumanTime(handler.data[handler.data.length - 1].date);
+      title = `Device Pool Telemetry / ${minDate} - ${maxDate}`;
+   }
+   
    return <Stack>
       {showSummary && !!handler && <SummaryModal handler={handler} onClose={() => setShowSummary(false)} />}
       <Stack style={{ position: "relative" }}>
          <Stack horizontal verticalAlign="start" style={{ width: 1412, paddingBottom: 18 }} tokens={{ childrenGap: 32 }}>
             <Stack style={{ paddingLeft: 4 }}>
-               <Text variant="mediumPlus" styles={{ root: { fontWeight: "unset", fontFamily: "Horde Open Sans SemiBold" } }}>Device Pool Telemetry</Text>
+               <Text variant="mediumPlus" styles={{ root: { fontWeight: "unset", fontFamily: "Horde Open Sans SemiBold" } }}>{title}</Text>
             </Stack>
 
             <Stack grow />
