Horde: Add a test the quarantined issues stay assigned, which they do
#rnx

diff --git a/Engine/Source/Programs/Horde/Horde.Build.Tests/IssueServiceTests.cs b/Engine/Source/Programs/Horde/Horde.Build.Tests/IssueServiceTests.cs
--- a/Engine/Source/Programs/Horde/Horde.Build.Tests/IssueServiceTests.cs
+++ b/Engine/Source/Programs/Horde/Horde.Build.Tests/IssueServiceTests.cs
@@ -1927,6 +1927,9 @@ namespace Horde.Build.Tests
 				lastSeenAt = issues[0].LastSeenAt;
 			}
 
+			// assign to bob
+			await IssueService.UpdateIssueAsync(issueId, ownerId: _bobId);
+
 			// Mark issue as quarantined
 			await IssueService.UpdateIssueAsync(issueId, quarantinedById: _jerryId);
 
@@ -1939,6 +1942,7 @@ namespace Horde.Build.Tests
 
 				IIssue? issue = await IssueCollection.GetIssueAsync(issueId);
 				Assert.IsNull(issue!.ResolvedAt);
+				Assert.AreEqual(issue!.OwnerId, _bobId);
 
 				List<IIssueSpan> spans = await IssueCollection.FindSpansAsync(issue.Id);
 				Assert.AreEqual(spans.Count, 1);
