Horde: Improve performance of querying large numbers of commits.

* Tags are now only computed on demand.
* Batching for queries that require describe commands reuses connections and stream metadata between iterations.

diff --git a/Engine/Source/Programs/Horde/Horde.Build/Perforce/IPerforceService.cs b/Engine/Source/Programs/Horde/Horde.Build/Perforce/IPerforceService.cs
--- a/Engine/Source/Programs/Horde/Horde.Build/Perforce/IPerforceService.cs
+++ b/Engine/Source/Programs/Horde/Horde.Build/Perforce/IPerforceService.cs
@@ -2,6 +2,7 @@
 
 using System;
 using System.Collections.Generic;
+using System.Linq;
 using System.Text.RegularExpressions;
 using System.Threading;
 using System.Threading.Tasks;
@@ -59,6 +60,48 @@ namespace Horde.Build.Perforce
 		/// <param name="cancellationToken"></param>
 		/// <returns></returns>
 		ValueTask<InfoRecord> GetInfoAsync(CancellationToken cancellationToken);
+
+		/// <summary>
+		/// Gets a cached stream view for the given stream
+		/// </summary>
+		/// <param name="stream">Stream to find the view for</param>
+		/// <param name="cancellationToken"></param>
+		/// <returns></returns>
+		ValueTask<PerforceViewMap> GetCachedStreamViewAsync(IStream stream, CancellationToken cancellationToken);
+	}
+
+	/// <summary>
+	/// Extension methods for PerforceService
+	/// </summary>
+	public static class PooledPerforceConnectionExtensions
+	{
+		/// <summary>
+		/// Map files from a <see cref="DescribeRecord"/> to relative paths within a stream
+		/// </summary>
+		/// <param name="perforce">Perforce connection</param>
+		/// <param name="stream">Stream to target</param>
+		/// <param name="describeRecord">Description of the change</param>
+		/// <param name="cancellationToken">Cancellation token for the operation</param>
+		/// <returns>List of files relative to the stream</returns>
+		public static Task<List<string>> GetStreamFilesAsync(this IPooledPerforceConnection perforce, IStream stream, DescribeRecord describeRecord, CancellationToken cancellationToken)
+		{
+			return GetStreamFilesAsync(perforce, stream, describeRecord.Files.ConvertAll(x => x.DepotFile), cancellationToken);
+		}
+
+		/// <summary>
+		/// Map files from a <see cref="DescribeRecord"/> to relative paths within a stream
+		/// </summary>
+		/// <param name="perforce">Perforce connection</param>
+		/// <param name="stream">Stream to target</param>
+		/// <param name="depotPaths">Paths within the depot</param>
+		/// <param name="cancellationToken">Cancellation token for the operation</param>
+		/// <returns>List of files relative to the stream</returns>
+		public static async Task<List<string>> GetStreamFilesAsync(this IPooledPerforceConnection perforce, IStream stream, List<string> depotPaths, CancellationToken cancellationToken)
+		{
+			PerforceViewMap viewMap = await perforce.GetCachedStreamViewAsync(stream, cancellationToken);
+			InfoRecord infoRecord = await perforce.GetInfoAsync(cancellationToken);
+			return viewMap.MapFiles(depotPaths, infoRecord.PathComparison).ToList();
+		}
 	}
 
 	/// <summary>
diff --git a/Engine/Source/Programs/Horde/Horde.Build/Perforce/PerforceService.cs b/Engine/Source/Programs/Horde/Horde.Build/Perforce/PerforceService.cs
--- a/Engine/Source/Programs/Horde/Horde.Build/Perforce/PerforceService.cs
+++ b/Engine/Source/Programs/Horde/Horde.Build/Perforce/PerforceService.cs
@@ -19,7 +19,6 @@ using Horde.Build.Server;
 using Horde.Build.Streams;
 using Horde.Build.Users;
 using Horde.Build.Utilities;
-using Microsoft.AspNetCore.Connections.Features;
 using Microsoft.Extensions.Caching.Memory;
 using Microsoft.Extensions.Logging;
 using Microsoft.Extensions.Options;
@@ -108,8 +107,27 @@ namespace Horde.Build.Perforce
 				return _inner.Perforce.CreateRecord(fields);
 			}
 
+			/// <inheritdoc/>
 			public ValueTask<InfoRecord> GetInfoAsync(CancellationToken cancellationToken) => _inner.GetInfoAsync(cancellationToken);
 
+			/// <inheritdoc/>
+			public async ValueTask<PerforceViewMap> GetCachedStreamViewAsync(IStream stream, CancellationToken cancellationToken)
+			{
+				PerforceViewMap? viewMap;
+				if (!_owner._streamCache.TryGetValue(stream.Id, out viewMap))
+				{
+					StreamRecord streamRecord = await _inner.Perforce.GetStreamAsync(stream.Name, true, cancellationToken);
+					viewMap = PerforceViewMap.Parse(streamRecord.View);
+
+					using (ICacheEntry entry = _owner._streamCache.CreateEntry(stream.Id))
+					{
+						entry.AbsoluteExpirationRelativeToNow = TimeSpan.FromMinutes(2.0);
+						entry.Value = viewMap;
+					}
+				}
+
+				return viewMap!;
+			}
 		}
 
 		protected class Credentials
@@ -181,10 +199,10 @@ namespace Horde.Build.Perforce
 			{
 				if (_files == null)
 				{
-					using (IPerforceConnection perforce = await _owner.ConnectAsync(_stream.Config.ClusterName, null, cancellationToken))
+					using (IPooledPerforceConnection perforce = await _owner.ConnectAsync(_stream.Config.ClusterName, null, cancellationToken))
 					{
 						DescribeRecord describeRecord = await perforce.DescribeAsync(Number, cancellationToken);
-						_files = GetStreamFiles(describeRecord, _stream.Config.Name);
+						_files = await perforce.GetStreamFilesAsync(_stream, describeRecord, cancellationToken);
 					}
 				}
 				return _files;
@@ -203,6 +221,7 @@ namespace Horde.Build.Perforce
 		readonly Dictionary<string, Dictionary<string, Credentials>> _userCredentialsByCluster = new Dictionary<string, Dictionary<string, Credentials>>(StringComparer.OrdinalIgnoreCase);
 		readonly IUserCollection _userCollection;
 		readonly MemoryCache _userCache = new MemoryCache(new MemoryCacheOptions { SizeLimit = 2000 });
+		readonly MemoryCache _streamCache = new MemoryCache(new MemoryCacheOptions());
 
 		readonly List<PooledConnection> _pooledConnections = new List<PooledConnection>();
 
@@ -234,6 +253,7 @@ namespace Horde.Build.Perforce
 			_pooledConnections.Clear();
 
 			_userCache.Dispose();
+			_streamCache.Dispose();
 		}
 
 		async Task<PooledConnectionHandle> CreatePooledConnectionAsync(string serverAndPort, Credentials credentials, ClientRecord? clientRecord, CancellationToken cancellationToken)
@@ -609,33 +629,6 @@ namespace Horde.Build.Perforce
 			return clientName;
 		}
 
-		/// <inheritdoc/>
-		public async Task<List<ICommit>> GetChangesAsync(IStream stream, int? minChange, int? maxChange, int? maxResults, CancellationToken cancellationToken)
-		{
-			using IScope scope = GlobalTracer.Instance.BuildSpan("PerforceService.GetChangesAsync").StartActive();
-			scope.Span.SetTag("ClusterName", stream.Config.ClusterName);
-			scope.Span.SetTag("MinChange", minChange ?? -1);
-			scope.Span.SetTag("MaxChange", maxChange ?? -1);
-			scope.Span.SetTag("MaxResults", maxResults ?? -1);
-
-			List<ICommit> results = new List<ICommit>();
-			using (PooledConnectionHandle perforce = await ConnectWithStreamClientAsync(stream, null, cancellationToken))
-			{
-				string filter = GetFilter($"//{perforce.Settings.ClientName}/...", minChange, maxChange);
-
-				InfoRecord info = await perforce.GetInfoAsync(cancellationToken);
-
-				List<ChangesRecord> changes = await perforce.GetChangesAsync(ChangesOptions.IncludeTimes | ChangesOptions.LongOutput, maxResults ?? -1, ChangeStatus.Submitted, filter, cancellationToken);
-
-				foreach (ChangesRecord change in changes)
-				{
-					ICommit commit = await CreateCommitAsync(stream, change, info, cancellationToken);
-					results.Add(commit);
-				}
-			}
-			return results;
-		}
-
 		async ValueTask<Commit> CreateCommitInternalAsync(IStream stream, int number, string author, string description, string? basePath, DateTime dateUtc, CancellationToken cancellationToken)
 		{
 			IUser authorUser = await FindOrAddUserAsync(stream.Config.ClusterName, author, cancellationToken);
@@ -652,12 +645,13 @@ namespace Horde.Build.Perforce
 			return new Commit(this, stream, number, originalChange, authorUser.Id, ownerUser.Id, description, basePath ?? String.Empty, dateUtc);
 		}
 
-		public async ValueTask<ICommit> CreateCommitAsync(IStream stream, DescribeRecord describeRecord, InfoRecord serverInfo, CancellationToken cancellationToken)
+		protected async ValueTask<ICommit> CreateCommitAsync(IPooledPerforceConnection perforce, IStream stream, DescribeRecord describeRecord, InfoRecord serverInfo, CancellationToken cancellationToken)
 		{
 			DateTime timeUtc = new DateTime(describeRecord.Time.Ticks - serverInfo.TimeZoneOffsetSecs * TimeSpan.TicksPerSecond, DateTimeKind.Utc);
 
 			Commit commit = await CreateCommitInternalAsync(stream, describeRecord.Number, describeRecord.User, describeRecord.Description, describeRecord.Path, timeUtc, cancellationToken);
-			commit.SetFiles(GetStreamFiles(describeRecord, stream.Name));
+			List<string> files = await perforce.GetStreamFilesAsync(stream, describeRecord, cancellationToken);
+			commit.SetFiles(files);
 
 			return commit;
 		}
@@ -727,7 +721,7 @@ namespace Horde.Build.Perforce
 
 				DateTime timeUtc = new DateTime(record.Time.Ticks - info.TimeZoneOffsetSecs * TimeSpan.TicksPerSecond, DateTimeKind.Utc);
 
-				ICommit commit = await CreateCommitAsync(stream, record, info, cancellationToken);
+				ICommit commit = await CreateCommitAsync(perforce, stream, record, info, cancellationToken);
 				return commit;
 			}
 		}
@@ -806,7 +800,7 @@ namespace Horde.Build.Perforce
 				List<DescribeRecord> records = await perforce.DescribeAsync(DescribeOptions.Shelved, -1, changeNumbers.ToArray(), cancellationToken);
 				foreach (DescribeRecord record in records)
 				{
-					ICommit commit = await CreateCommitAsync(stream, record, info, cancellationToken);
+					ICommit commit = await CreateCommitAsync(perforce, stream, record, info, cancellationToken);
 					results.Add(commit);
 				}
 			}
@@ -1079,20 +1073,6 @@ namespace Horde.Build.Perforce
 			return filter.ToString();
 		}
 
-		static List<string> GetStreamFiles(DescribeRecord describeRecord, string streamName)
-		{
-			List<string> files = new List<string>();
-			foreach (DescribeFileRecord describeFile in describeRecord.Files)
-			{
-				string? relativePath;
-				if (TryGetStreamRelativePath(describeFile.DepotFile, streamName, out relativePath))
-				{
-					files.Add(relativePath);
-				}
-			}
-			return files;
-		}
-
 		/// <summary>
 		/// Gets a stream-relative path from a depot path
 		/// </summary>
@@ -1148,43 +1128,75 @@ namespace Horde.Build.Perforce
 
 			public virtual async IAsyncEnumerable<ICommit> FindAsync(int? minChange, int? maxChange, int? maxResults, IReadOnlyList<CommitTag>? tags, [EnumeratorCancellation] CancellationToken cancellationToken = default)
 			{
-				int maxBatchResults = 20;
-				if (tags == null && maxResults != null)
-				{
-					maxBatchResults = Math.Min(maxBatchResults, maxResults.Value);
-				}
+				using IScope scope = GlobalTracer.Instance.BuildSpan("PerforceService.CommitSource.FindAsync").StartActive();
+				scope.Span.SetTag("ClusterName", _stream.Config.ClusterName);
+				scope.Span.SetTag("MinChange", minChange ?? -1);
+				scope.Span.SetTag("MaxChange", maxChange ?? -1);
+				scope.Span.SetTag("MaxResults", maxResults ?? -1);
 
-				int numResults = 0;
-				while (maxResults == null || numResults < maxResults.Value)
+				using (PooledConnectionHandle perforce = await _perforceService.ConnectWithStreamClientAsync(_stream, null, cancellationToken))
 				{
-					List<ICommit> commits = await _perforceService.GetChangesAsync(_stream, minChange, maxChange, maxBatchResults, cancellationToken);
-					if (commits.Count == 0)
+					InfoRecord info = await perforce.GetInfoAsync(cancellationToken);
+
+					if (tags == null || tags.Count == 0)
 					{
-						break;
-					}
+						string filter = GetFilter($"//{perforce.Settings.ClientName}/...", minChange, maxChange);
 
-					foreach (ICommit commit in commits)
+						List<ChangesRecord> changes = await perforce.GetChangesAsync(ChangesOptions.IncludeTimes | ChangesOptions.LongOutput, maxResults ?? -1, ChangeStatus.Submitted, filter, cancellationToken);
+						foreach (ChangesRecord change in changes)
+						{
+							ICommit commit = await _perforceService.CreateCommitAsync(_stream, change, info, cancellationToken);
+							yield return commit;
+						}
+					}
+					else
 					{
-						if (tags != null)
+						const int BatchSize = 20;
+
+						// Construct filters for all the tags we're searching for
+						List<FileFilter> filters = new List<FileFilter>();
+						foreach (CommitTag tag in tags)
 						{
-							IReadOnlyList<CommitTag> commitTags = await commit.GetTagsAsync(cancellationToken);
-							if (!tags.Any(x => commitTags.Contains(x)))
+							if (_stream.Config.TryGetCommitTag(tag, out CommitTagConfig? commitTagConfig))
 							{
-								continue;
+								filters.Add(commitTagConfig.CreateFileFilter());
 							}
 						}
 
-						yield return commit;
+						// Get the stream view
+						PerforceViewMap viewMap = await perforce.GetCachedStreamViewAsync(_stream, cancellationToken);
 
-						numResults++;
-
-						if (maxResults != null && numResults >= maxResults.Value)
+						// Do the query in batches
+						int numResults = maxResults ?? Int32.MaxValue;
+						while (numResults > 0)
 						{
-							break;
+							string filter = GetFilter($"//{perforce.Settings.ClientName}/...", minChange, maxChange);
+
+							List<ChangesRecord> changesRecords = await perforce.GetChangesAsync(ChangesOptions.None, null, BatchSize, ChangeStatus.Submitted, null, filter, cancellationToken);
+							if (changesRecords.Count == 0)
+							{
+								break;
+							}
+
+							List<DescribeRecord> describeRecords = await perforce.DescribeAsync(DescribeOptions.Shelved, -1, changesRecords.Select(x => x.Number).ToArray(), cancellationToken);
+							foreach (DescribeRecord describeRecord in describeRecords)
+							{
+								List<string> files = viewMap.MapFiles(describeRecord.Files.Select(x => x.DepotFile), info.PathComparison).ToList();
+								if (files.Any(x => filters.Any(y => y.Matches(x))))
+								{
+									ICommit commit = await _perforceService.CreateCommitAsync(perforce, _stream, describeRecord, info, cancellationToken);
+									yield return commit;
+
+									if (--numResults == 0)
+									{
+										yield break;
+									}
+								}
+							}
+
+							maxChange = changesRecords[^1].Number - 1;
 						}
 					}
-
-					maxChange = commits[^1].Number;
 				}
 			}
 
diff --git a/Engine/Source/Programs/Horde/Horde.Build/Perforce/PerforceServiceCache.cs b/Engine/Source/Programs/Horde/Horde.Build/Perforce/PerforceServiceCache.cs
--- a/Engine/Source/Programs/Horde/Horde.Build/Perforce/PerforceServiceCache.cs
+++ b/Engine/Source/Programs/Horde/Horde.Build/Perforce/PerforceServiceCache.cs
@@ -393,7 +393,7 @@ namespace Horde.Build.Perforce
 
 						if (files.Count > 0)
 						{
-							ICommit commit = await CreateCommitAsync(streamInfo.Stream, describeRecord, info, cancellationToken);
+							ICommit commit = await CreateCommitAsync(perforce, streamInfo.Stream, describeRecord, info, cancellationToken);
 
 							List<CommitTag> commitTags = new List<CommitTag>();
 							foreach (CommitTagInfo commitTagInfo in streamInfo.CommitTags)
diff --git a/Engine/Source/Programs/Horde/Horde.Build/Streams/StreamMessages.cs b/Engine/Source/Programs/Horde/Horde.Build/Streams/StreamMessages.cs
--- a/Engine/Source/Programs/Horde/Horde.Build/Streams/StreamMessages.cs
+++ b/Engine/Source/Programs/Horde/Horde.Build/Streams/StreamMessages.cs
@@ -454,7 +454,7 @@ namespace Horde.Build.Streams
 		/// <summary>
 		/// Tags for this commit
 		/// </summary>
-		public List<CommitTag> Tags { get; set; }
+		public List<CommitTag>? Tags { get; set; }
 
 		/// <summary>
 		/// List of files that were modified, relative to the stream base
@@ -468,13 +468,17 @@ namespace Horde.Build.Streams
 		/// <param name="author">Author of the change</param>
 		/// <param name="tags">Tags for the commit</param>
 		/// <param name="files">Files modified by the commit</param>
-		public GetCommitResponse(ICommit commit, IUser author, IReadOnlyList<CommitTag> tags, IReadOnlyList<string>? files)
+		public GetCommitResponse(ICommit commit, IUser author, IReadOnlyList<CommitTag>? tags, IReadOnlyList<string>? files)
 		{
 			Number = commit.Number;
 			Author = author.Name;
 			AuthorInfo = new GetThinUserInfoResponse(author);
 			Description = commit.Description;
-			Tags = new List<CommitTag>(tags);
+
+			if (tags != null)
+			{
+				Tags = new List<CommitTag>(tags);
+			}
 
 			if (files != null)
 			{
diff --git a/Engine/Source/Programs/Horde/Horde.Build/Streams/StreamsController.cs b/Engine/Source/Programs/Horde/Horde.Build/Streams/StreamsController.cs
--- a/Engine/Source/Programs/Horde/Horde.Build/Streams/StreamsController.cs
+++ b/Engine/Source/Programs/Horde/Horde.Build/Streams/StreamsController.cs
@@ -162,12 +162,13 @@ namespace Horde.Build.Streams
 		/// <param name="min">The starting changelist number</param>
 		/// <param name="max">The ending changelist number</param>
 		/// <param name="results">Number of results to return</param>
+		/// <param name="tags">Tags to filter the changes returned</param>
 		/// <param name="filter">The filter to apply to the results</param>
 		/// <returns>Http result code</returns>
 		[HttpGet]
 		[Route("/api/v1/streams/{streamId}/changes")]
 		[ProducesResponseType(typeof(List<GetCommitResponse>), 200)]
-		public async Task<ActionResult<List<object>>> GetChangesAsync(StreamId streamId, [FromQuery] int? min = null, [FromQuery] int? max = null, [FromQuery] int results = 50, PropertyFilter? filter = null)
+		public async Task<ActionResult<List<object>>> GetChangesAsync(StreamId streamId, [FromQuery] int? min = null, [FromQuery] int? max = null, [FromQuery] int results = 50, [FromQuery] string? tags = null, PropertyFilter? filter = null)
 		{
 			IStream? stream = await _streamService.GetStreamAsync(streamId);
 			if (stream == null)
@@ -179,14 +180,19 @@ namespace Horde.Build.Streams
 				return Forbid(AclAction.ViewChanges, streamId);
 			}
 
-			List<ICommit> commits = await _commitService.GetCollection(stream).FindAsync(min, max, results).ToListAsync();
+			List<CommitTag>? commitTags = null;
+			if (tags != null)
+			{
+				commitTags = tags.Split(';', StringSplitOptions.RemoveEmptyEntries).Select(x => new CommitTag(x)).ToList();
+			}
+
+			List<ICommit> commits = await _commitService.GetCollection(stream).FindAsync(min, max, results, commitTags).ToListAsync();
 
 			List<GetCommitResponse> responses = new List<GetCommitResponse>();
 			foreach (ICommit commit in commits)
 			{
 				IUser? author = await _userCollection.GetCachedUserAsync(commit.AuthorId);
-				IReadOnlyList<CommitTag> tags = await commit.GetTagsAsync(HttpContext.RequestAborted);
-				responses.Add(new GetCommitResponse(commit, author!, tags, null));
+				responses.Add(new GetCommitResponse(commit, author!, null, null));
 			}
 			return responses.ConvertAll(x => PropertyFilter.Apply(x, filter));
 		}
diff --git a/Engine/Source/Programs/Shared/EpicGames.Perforce/PerforceViewMap.cs b/Engine/Source/Programs/Shared/EpicGames.Perforce/PerforceViewMap.cs
--- a/Engine/Source/Programs/Shared/EpicGames.Perforce/PerforceViewMap.cs
+++ b/Engine/Source/Programs/Shared/EpicGames.Perforce/PerforceViewMap.cs
@@ -86,6 +86,23 @@ namespace EpicGames.Perforce
 			return included;
 		}
 
+		/// <summary>
+		/// Maps a set of files into the target files
+		/// </summary>
+		/// <param name="sourceFiles">List of source files</param>
+		/// <param name="comparison">Comparison to use for strings</param>
+		/// <returns>List of files in the target domain, excluding any not covered by the mapping</returns>
+		public IEnumerable<string> MapFiles(IEnumerable<string> sourceFiles, StringComparison comparison)
+		{
+			foreach (string sourceFile in sourceFiles)
+			{
+				if (TryMapFile(sourceFile, comparison, out string? targetFile))
+				{
+					yield return targetFile;
+				}
+			}
+		}
+
 		/// <summary>
 		/// Attempts to convert a source file to its target path
 		/// </summary>
