Horde: Throw an exception if querying details for a change through the wrong stream.

diff --git a/Engine/Source/Programs/Horde/Horde.Build/Perforce/PerforceService.cs b/Engine/Source/Programs/Horde/Horde.Build/Perforce/PerforceService.cs
--- a/Engine/Source/Programs/Horde/Horde.Build/Perforce/PerforceService.cs
+++ b/Engine/Source/Programs/Horde/Horde.Build/Perforce/PerforceService.cs
@@ -650,7 +650,12 @@ namespace Horde.Build.Perforce
 			DateTime timeUtc = new DateTime(describeRecord.Time.Ticks - serverInfo.TimeZoneOffsetSecs * TimeSpan.TicksPerSecond, DateTimeKind.Utc);
 
 			Commit commit = await CreateCommitInternalAsync(stream, describeRecord.Number, describeRecord.User, describeRecord.Description, describeRecord.Path, timeUtc, cancellationToken);
+
 			List<string> files = await perforce.GetStreamFilesAsync(stream, describeRecord, cancellationToken);
+			if (files.Count == 0)
+			{
+				throw new PerforceException($"Changelist {commit.Number} does not contain any files in {stream.Id}");
+			}
 			commit.SetFiles(files);
 
 			return commit;
