Horde: Reference global config through a cached configuration object, rather than copying onto the globals singleton.

diff --git a/Engine/Source/Programs/Horde/Horde.Build/Acls/AclService.cs b/Engine/Source/Programs/Horde/Horde.Build/Acls/AclService.cs
--- a/Engine/Source/Programs/Horde/Horde.Build/Acls/AclService.cs
+++ b/Engine/Source/Programs/Horde/Horde.Build/Acls/AclService.cs
@@ -4,6 +4,7 @@ using System;
 using System.Collections.Generic;
 using System.IdentityModel.Tokens.Jwt;
 using System.Linq;
+using System.Net;
 using System.Security.Claims;
 using System.Threading.Tasks;
 using Horde.Build.Agents;
@@ -67,16 +68,16 @@ namespace Horde.Build.Acls
 		public static AclClaim AgentRoleClaim { get; } = new AclClaim(HordeClaimTypes.Role, "agent");
 
 		private readonly Acl _defaultAcl = new();
-		private readonly MongoService _mongoService;
+		private readonly GlobalsService _globalsService;
 
 		/// <summary>
 		/// Constructor
 		/// </summary>
-		/// <param name="mongoService">The database service instance</param>
+		/// <param name="globalsService">The database service instance</param>
 		/// <param name="settings">The settings object</param>
-		public AclService(MongoService mongoService, IOptionsMonitor<ServerSettings> settings)
+		public AclService(GlobalsService globalsService, IOptionsMonitor<ServerSettings> settings)
 		{
-			_mongoService = mongoService;
+			_globalsService = globalsService;
 
             List<AclAction> adminActions = new();
 			foreach (AclAction? action in Enum.GetValues(typeof(AclAction)))
@@ -106,8 +107,8 @@ namespace Horde.Build.Acls
 		/// <returns>Scopes instance</returns>
 		public async Task<Acl> GetRootAcl()
 		{
-			Globals globals = await _mongoService.GetGlobalsAsync();
-			return globals.RootAcl ?? new Acl();
+			IGlobals globals = await _globalsService.GetAsync();
+			return Acl.Merge(new Acl(), globals.Config.Acl) ?? new Acl();
 		}
 
 		/// <summary>
@@ -141,9 +142,9 @@ namespace Horde.Build.Acls
 		/// <param name="claims">List of claims to include</param>
 		/// <param name="expiry">Time that the token expires</param>
 		/// <returns>JWT security token with a claim for creating new agents</returns>
-		public string IssueBearerToken(IEnumerable<AclClaim> claims, TimeSpan? expiry)
+		public async ValueTask<string> IssueBearerTokenAsync(IEnumerable<AclClaim> claims, TimeSpan? expiry)
 		{
-			return IssueBearerToken(claims.Select(x => new Claim(x.Type, x.Value)), expiry);
+			return await IssueBearerTokenAsync(claims.Select(x => new Claim(x.Type, x.Value)), expiry);
 		}
 
 		/// <summary>
@@ -152,11 +153,12 @@ namespace Horde.Build.Acls
 		/// <param name="claims">List of claims to include</param>
 		/// <param name="expiry">Time that the token expires</param>
 		/// <returns>JWT security token with a claim for creating new agents</returns>
-		public string IssueBearerToken(IEnumerable<Claim> claims, TimeSpan? expiry)
+		public async ValueTask<string> IssueBearerTokenAsync(IEnumerable<Claim> claims, TimeSpan? expiry)
 		{
-			SigningCredentials signingCredentials = new(_mongoService.JwtSigningKey, SecurityAlgorithms.HmacSha256);
+			IGlobals globals = await _globalsService.GetAsync();
+			SigningCredentials signingCredentials = new(globals.JwtSigningKey, SecurityAlgorithms.HmacSha256);
 
-			JwtSecurityToken token = new(_mongoService.JwtIssuer, null, claims, null, DateTime.UtcNow + expiry, signingCredentials);
+			JwtSecurityToken token = new(globals.JwtIssuer, null, claims, null, DateTime.UtcNow + expiry, signingCredentials);
 			return new JwtSecurityTokenHandler().WriteToken(token);
 		}
 
diff --git a/Engine/Source/Programs/Horde/Horde.Build/Agents/AgentService.cs b/Engine/Source/Programs/Horde/Horde.Build/Agents/AgentService.cs
--- a/Engine/Source/Programs/Horde/Horde.Build/Agents/AgentService.cs
+++ b/Engine/Source/Programs/Horde/Horde.Build/Agents/AgentService.cs
@@ -167,13 +167,13 @@ namespace Horde.Build.Agents
 		/// <param name="agentId">The agent id</param>
 		/// <param name="sessionId">The session id</param>
 		/// <returns>Bearer token for the agent</returns>
-		public string IssueSessionToken(AgentId agentId, SessionId sessionId)
+		public async ValueTask<string> IssueSessionTokenAsync(AgentId agentId, SessionId sessionId)
 		{
 			List<AclClaim> claims = new List<AclClaim>();
 			claims.Add(AclService.AgentRoleClaim);
 			claims.Add(AclService.GetAgentClaim(agentId));
 			claims.Add(AclService.GetSessionClaim(sessionId));
-			return _aclService.IssueBearerToken(claims, null);
+			return await _aclService.IssueBearerTokenAsync(claims, null);
 		}
 
 		/// <summary>
diff --git a/Engine/Source/Programs/Horde/Horde.Build/Agents/IAgent.cs b/Engine/Source/Programs/Horde/Horde.Build/Agents/IAgent.cs
--- a/Engine/Source/Programs/Horde/Horde.Build/Agents/IAgent.cs
+++ b/Engine/Source/Programs/Horde/Horde.Build/Agents/IAgent.cs
@@ -663,12 +663,12 @@ namespace Horde.Build.Agents
 		/// <param name="globals"></param>
 		/// <param name="workspaces"></param>
 		/// <returns></returns>
-		public static HashSet<AgentWorkspace> GetAutoSdkWorkspaces(this IAgent agent, Globals globals, List<AgentWorkspace> workspaces)
+		public static HashSet<AgentWorkspace> GetAutoSdkWorkspaces(this IAgent agent, IGlobals globals, List<AgentWorkspace> workspaces)
 		{
 			HashSet<AgentWorkspace> autoSdkWorkspaces = new HashSet<AgentWorkspace>();
 			foreach (string? clusterName in workspaces.Select(x => x.Cluster).Distinct())
 			{
-				PerforceCluster? cluster = globals.FindPerforceCluster(clusterName);
+				PerforceCluster? cluster = globals.Config.FindPerforceCluster(clusterName);
 				if (cluster != null)
 				{
 					AgentWorkspace? autoSdkWorkspace = GetAutoSdkWorkspace(agent, cluster);
diff --git a/Engine/Source/Programs/Horde/Horde.Build/Agents/Pools/PoolService.cs b/Engine/Source/Programs/Horde/Horde.Build/Agents/Pools/PoolService.cs
--- a/Engine/Source/Programs/Horde/Horde.Build/Agents/Pools/PoolService.cs
+++ b/Engine/Source/Programs/Horde/Horde.Build/Agents/Pools/PoolService.cs
@@ -21,9 +21,9 @@ namespace Horde.Build.Agents.Pools
 	public class PoolService
 	{
 		/// <summary>
-		/// The database service instance
+		/// The globals service instance
 		/// </summary>
-		readonly MongoService _mongoService;
+		readonly GlobalsService _globalsService;
 
 		/// <summary>
 		/// Collection of pool documents
@@ -43,12 +43,12 @@ namespace Horde.Build.Agents.Pools
 		/// <summary>
 		/// Constructor
 		/// </summary>
-		/// <param name="mongoService"></param>
+		/// <param name="globalsService"></param>
 		/// <param name="pools">Collection of pool documents</param>
 		/// <param name="clock"></param>
-		public PoolService(MongoService mongoService, IPoolCollection pools, IClock clock)
+		public PoolService(GlobalsService globalsService, IPoolCollection pools, IClock clock)
 		{
-			_mongoService = mongoService;
+			_globalsService = globalsService;
 			_pools = pools;
 			_clock = clock;
 		}
@@ -253,7 +253,7 @@ namespace Horde.Build.Agents.Pools
 
 			if (bAddAutoSdkWorkspace)
 			{
-				Globals globals = await _mongoService.GetGlobalsAsync();
+				IGlobals globals = await _globalsService.GetAsync();
 				workspaces.UnionWith(agent.GetAutoSdkWorkspaces(globals, workspaces.ToList()));
 			}
 
diff --git a/Engine/Source/Programs/Horde/Horde.Build/Authentication/HordeJwtBearerHandler.cs b/Engine/Source/Programs/Horde/Horde.Build/Authentication/HordeJwtBearerHandler.cs
--- a/Engine/Source/Programs/Horde/Horde.Build/Authentication/HordeJwtBearerHandler.cs
+++ b/Engine/Source/Programs/Horde/Horde.Build/Authentication/HordeJwtBearerHandler.cs
@@ -3,6 +3,7 @@
 using System;
 using System.IdentityModel.Tokens.Jwt;
 using System.Text.Encodings.Web;
+using System.Threading;
 using System.Threading.Tasks;
 using Horde.Build.Server;
 using Horde.Build.Utilities;
@@ -10,6 +11,7 @@ using Microsoft.AspNetCore.Authentication;
 using Microsoft.AspNetCore.Authentication.JwtBearer;
 using Microsoft.Extensions.Logging;
 using Microsoft.Extensions.Options;
+using Microsoft.IdentityModel.Protocols;
 
 namespace Horde.Build.Authentication
 {
@@ -22,53 +24,55 @@ namespace Horde.Build.Authentication
 		/// Default name of the authentication scheme
 		/// </summary>
 		public const string AuthenticationScheme = "ServerJwt";
-		readonly MongoService _mongoService;
 
-		public HordeJwtBearerHandler(ILoggerFactory logger, UrlEncoder encoder, ISystemClock clock, MongoService mongoService, IOptionsMonitorCache<JwtBearerOptions> optionsCache)
-			: base(GetOptionsMonitor(mongoService, optionsCache), logger, encoder, clock)
+		readonly AsyncCachedValue<IGlobals> _globals;
+
+		public HordeJwtBearerHandler(ILoggerFactory logger, UrlEncoder encoder, ISystemClock clock, GlobalsService globalsService, IOptionsMonitorCache<JwtBearerOptions> optionsCache)
+			: base(GetOptionsMonitor(optionsCache), logger, encoder, clock)
 		{
-			_mongoService = mongoService;
+			_globals = new AsyncCachedValue<IGlobals>(async () => await globalsService.GetAsync(), TimeSpan.FromSeconds(30.0));
 		}
 
-		private static IOptionsMonitor<JwtBearerOptions> GetOptionsMonitor(MongoService mongoService, IOptionsMonitorCache<JwtBearerOptions> optionsCache)
+		private static IOptionsMonitor<JwtBearerOptions> GetOptionsMonitor(IOptionsMonitorCache<JwtBearerOptions> optionsCache)
 		{
-			ConfigureNamedOptions<JwtBearerOptions> namedOptions = new ConfigureNamedOptions<JwtBearerOptions>(AuthenticationScheme, options => Configure(options, mongoService));
+			ConfigureNamedOptions<JwtBearerOptions> namedOptions = new ConfigureNamedOptions<JwtBearerOptions>(AuthenticationScheme, options => { });
 			OptionsFactory<JwtBearerOptions> optionsFactory = new OptionsFactory<JwtBearerOptions>(new[] { namedOptions }, Array.Empty<IPostConfigureOptions<JwtBearerOptions>>());
 			return new OptionsMonitor<JwtBearerOptions>(optionsFactory, Array.Empty<IOptionsChangeTokenSource<JwtBearerOptions>>(), optionsCache);
 		}
 
-		private static void Configure(JwtBearerOptions options, MongoService mongoService)
+		protected override async Task<AuthenticateResult> HandleAuthenticateAsync()
 		{
-			options.TokenValidationParameters.ValidateAudience = false;
+			// Get the current state
+			IGlobals globals = await _globals.GetAsync();
 
-			options.TokenValidationParameters.RequireExpirationTime = false;
-			options.TokenValidationParameters.ValidateLifetime = true;
+			Options.TokenValidationParameters.ValidateAudience = false;
 
-			options.TokenValidationParameters.ValidIssuer = mongoService.JwtIssuer;
-			options.TokenValidationParameters.ValidateIssuer = true;
+			Options.TokenValidationParameters.RequireExpirationTime = false;
+			Options.TokenValidationParameters.ValidateLifetime = true;
 
-			options.TokenValidationParameters.ValidateIssuerSigningKey = true;
-			options.TokenValidationParameters.IssuerSigningKey = mongoService.JwtSigningKey;
-		}
+			Options.TokenValidationParameters.ValidIssuer = globals.JwtIssuer;
+			Options.TokenValidationParameters.ValidateIssuer = true;
+
+			Options.TokenValidationParameters.ValidateIssuerSigningKey = true;
+			Options.TokenValidationParameters.IssuerSigningKey = globals.JwtSigningKey;
 
-		protected override Task<AuthenticateResult> HandleAuthenticateAsync()
-		{
 			// Silent fail if this JWT is not issued by the server
 			string? token;
 			if (!JwtUtils.TryGetBearerToken(Request, "Bearer ", out token))
 			{
-				return Task.FromResult(AuthenticateResult.NoResult());
+				return AuthenticateResult.NoResult();
 			}
 
-			// Validate that it's from the correct issuer
+			// Validate that it's from the correct issuer, and silent fail if not. Allows multiple handlers for bearer tokens.
 			JwtSecurityToken? jwtToken;
-			if (!JwtUtils.TryParseJwt(token, out jwtToken) || !String.Equals(jwtToken.Issuer, _mongoService.JwtIssuer, StringComparison.Ordinal))
+			if (!JwtUtils.TryParseJwt(token, out jwtToken) || !String.Equals(jwtToken.Issuer, globals.JwtIssuer, StringComparison.Ordinal))
 			{
-				return Task.FromResult(AuthenticateResult.NoResult());
+				return AuthenticateResult.NoResult();
 			}
 
 			// Pass it to the base class
-			return base.HandleAuthenticateAsync();
+			AuthenticateResult result = await base.HandleAuthenticateAsync();
+			return result;
 		}
 	}
 }
diff --git a/Engine/Source/Programs/Horde/Horde.Build/Compute/ComputeService.cs b/Engine/Source/Programs/Horde/Horde.Build/Compute/ComputeService.cs
--- a/Engine/Source/Programs/Horde/Horde.Build/Compute/ComputeService.cs
+++ b/Engine/Source/Programs/Horde/Horde.Build/Compute/ComputeService.cs
@@ -126,7 +126,7 @@ namespace Horde.Build.Compute
 		readonly RedisMessageQueue<ComputeTaskStatus> _messageQueue;
 		readonly ITicker _expireTasksTicker;
 		readonly IMemoryCache _requirementsCache;
-		readonly LazyCachedValue<Task<Globals>> _globals;
+		readonly LazyCachedValue<Task<IGlobals>> _globals;
 		readonly ILogger _logger;
 
 		static ComputeService()
@@ -137,14 +137,14 @@ namespace Horde.Build.Compute
 		/// <summary>
 		/// Constructor
 		/// </summary>
-		public ComputeService(MongoService mongoService, RedisService redisService, IStorageClient storageClient, IClock clock, ILogger<ComputeService> logger)
+		public ComputeService(GlobalsService globalsService, RedisService redisService, IStorageClient storageClient, IClock clock, ILogger<ComputeService> logger)
 		{
 			_storageClient = storageClient;
 			_taskScheduler = new RedisTaskScheduler<QueueKey, ComputeTaskInfo>(redisService.ConnectionPool, "compute/tasks/", logger);
 			_messageQueue = new RedisMessageQueue<ComputeTaskStatus>(redisService.GetDatabase(), "compute/messages/");
 			_expireTasksTicker = clock.AddTicker<ComputeService>(TimeSpan.FromMinutes(2.0), ExpireTasksAsync, logger);
 			_requirementsCache = new MemoryCache(new MemoryCacheOptions());
-			_globals = new LazyCachedValue<Task<Globals>>(() => mongoService.GetGlobalsAsync(), TimeSpan.FromSeconds(120.0));
+			_globals = new LazyCachedValue<Task<IGlobals>>(async () => await globalsService.GetAsync(), TimeSpan.FromSeconds(120.0));
 			_logger = logger;
 
 			OnLeaseStartedProperties.Add(x => x.TaskRefId);
@@ -231,8 +231,8 @@ namespace Horde.Build.Compute
 
 		async ValueTask<ComputeClusterConfig?> GetClusterAsync(ClusterId clusterId)
 		{
-			Globals globalsInstance = await _globals.GetCached();
-			return globalsInstance.ComputeClusters.FirstOrDefault(x => new ClusterId(x.Id) == clusterId);
+			IGlobals globals = await _globals.GetCached();
+			return globals.Config.Compute.FirstOrDefault(x => new ClusterId(x.Id) == clusterId);
 		}
 
 		/// <inheritdoc/>
diff --git a/Engine/Source/Programs/Horde/Horde.Build/Configuration/ConfigUpdateService.cs b/Engine/Source/Programs/Horde/Horde.Build/Configuration/ConfigUpdateService.cs
--- a/Engine/Source/Programs/Horde/Horde.Build/Configuration/ConfigUpdateService.cs
+++ b/Engine/Source/Programs/Horde/Horde.Build/Configuration/ConfigUpdateService.cs
@@ -42,7 +42,7 @@ namespace Horde.Build.Configuration
 		/// </summary>
 		const int Version = 12;
 
-		readonly MongoService _mongoService;
+		readonly GlobalsService _globalsService;
 		readonly ConfigCollection _configCollection;
 		readonly ToolCollection _toolCollection;
 		readonly ProjectService _projectService;
@@ -57,9 +57,9 @@ namespace Horde.Build.Configuration
 		/// <summary>
 		/// Constructor
 		/// </summary>
-		public ConfigUpdateService(MongoService mongoService, ConfigCollection configCollection, IPerforceService perforceService, ToolCollection toolCollection, ProjectService projectService, StreamService streamService, INotificationService notificationService, AgentService agentService, IClock clock, IOptionsMonitor<ServerSettings> settings, ILogger<ConfigUpdateService> logger)
+		public ConfigUpdateService(GlobalsService globalsService, MongoService mongoService, ConfigCollection configCollection, IPerforceService perforceService, ToolCollection toolCollection, ProjectService projectService, StreamService streamService, INotificationService notificationService, AgentService agentService, IClock clock, IOptionsMonitor<ServerSettings> settings, ILogger<ConfigUpdateService> logger)
 		{
-			_mongoService = mongoService;
+			_globalsService = globalsService;
 			_configCollection = configCollection;
 			_perforceService = perforceService;
 			_toolCollection = toolCollection;
@@ -124,23 +124,17 @@ namespace Horde.Build.Configuration
 				}
 				globalConfig = _cachedGlobalConfig;
 
-				Globals globals = await _mongoService.GetGlobalsAsync();
+				IGlobals globals = await _globalsService.GetAsync();
 				if (globals.ConfigRevision == revision)
 				{
 					break;
 				}
 
-				_logger.LogInformation("Updating configuration from {ConfigPath}", globals.ConfigRevision);
+				_logger.LogInformation("Updating configuration from {ConfigPath}", revision);
+				await _configCollection.AddConfigAsync(revision, globalConfig);
 
-				globals.ConfigRevision = revision;
-				globals.PerforceClusters = _cachedGlobalConfig.PerforceClusters;
-				globals.ScheduledDowntime = _cachedGlobalConfig.Downtime;
-				globals.MaxConformCount = _cachedGlobalConfig.MaxConformCount;
-				globals.ComputeClusters = _cachedGlobalConfig.Compute;
-				globals.Devices = _cachedGlobalConfig.Devices;
-				globals.RootAcl = Acl.Merge(null, _cachedGlobalConfig.Acl);
-
-				if (await _mongoService.TryUpdateSingletonAsync(globals))
+				IGlobals? newGlobals = await _globalsService.TryUpdateAsync(globals, revision);
+				if (newGlobals != null)
 				{
 					break;
 				}
diff --git a/Engine/Source/Programs/Horde/Horde.Build/Devices/DeviceService.cs b/Engine/Source/Programs/Horde/Horde.Build/Devices/DeviceService.cs
--- a/Engine/Source/Programs/Horde/Horde.Build/Devices/DeviceService.cs
+++ b/Engine/Source/Programs/Horde/Horde.Build/Devices/DeviceService.cs
@@ -88,49 +88,15 @@ namespace Horde.Build.Devices
 	/// </summary>
 	public sealed class DeviceService : IHostedService, IDisposable
 	{
-		/// <summary>
-		/// The ACL service instance
-		/// </summary>
+		readonly GlobalsService _globalsService;
 		readonly AclService _aclService;
-
-		/// <summary>
-		/// Instance of the notification service
-		/// </summary>
 		readonly INotificationService _notificationService;
-
-		/// <summary>
-		/// Singleton instance of the job service
-		/// </summary>
 		readonly JobService _jobService;
-
-		/// <summary>
-		/// Singleton instance of the stream service
-		/// </summary>
 		readonly StreamService _streamService;
-
-		/// <summary>
-		/// Singleton instance of the project service
-		/// </summary>
 		readonly ProjectService _projectService;
-
-		/// <summary>
-		/// Singleton instance of the mongo service
-		/// </summary>
 		readonly MongoService _mongoService;
-
-		/// <summary>
-		/// The user collection instance
-		/// </summary>
 		IUserCollection UserCollection { get; set; }
-
-		/// <summary>
-		/// Log output writer
-		/// </summary>
 		readonly ILogger<DeviceService> _logger;
-
-		/// <summary>
-		/// Device collection
-		/// </summary>
 		readonly IDeviceCollection _devices;
 		readonly ITicker _ticker;
 		readonly ITicker _telemetryTicker;
@@ -143,9 +109,10 @@ namespace Horde.Build.Devices
 		/// <summary>
 		/// Device service constructor
 		/// </summary>
-		public DeviceService(IDeviceCollection devices, ISingletonDocument<DevicePlatformMapV1> platformMapSingleton, IUserCollection userCollection, MongoService mongoService, JobService jobService, ProjectService projectService, StreamService streamService, AclService aclService, INotificationService notificationService, IClock clock, ILogger<DeviceService> logger)
+		public DeviceService(GlobalsService globalsService, IDeviceCollection devices, ISingletonDocument<DevicePlatformMapV1> platformMapSingleton, IUserCollection userCollection, MongoService mongoService, JobService jobService, ProjectService projectService, StreamService streamService, AclService aclService, INotificationService notificationService, IClock clock, ILogger<DeviceService> logger)
 		{
 			UserCollection = userCollection;
+			_globalsService = globalsService;
 			_devices = devices;
 			_jobService = jobService;
 			_projectService = projectService;
@@ -167,7 +134,6 @@ namespace Horde.Build.Devices
 			await _ticker.StartAsync();
 			await _telemetryTicker.StartAsync();
 		}
-		
 
 		/// <inheritdoc/>
 		public async Task StopAsync(CancellationToken cancellationToken)
@@ -190,16 +156,16 @@ namespace Horde.Build.Devices
 		{
 			try
 			{
-				Globals globals = await _mongoService.GetGlobalsAsync();
+				IGlobals globals = await _globalsService.GetAsync();
 
-				if (globals.Devices != null)
+				if (globals.Config.Devices != null)
 				{
 					await _platformMapSingleton.UpdateAsync(platformMap => {
 
 						platformMap.PlatformMap.Clear();
 						platformMap.PerfSpecHighMap.Clear();
 
-						foreach (DevicePlatformConfig platform in globals.Devices.Platforms)
+						foreach (DevicePlatformConfig platform in globals.Config.Devices.Platforms)
 						{
 							DevicePlatformId id = new DevicePlatformId(platform.Id);
 							foreach (string name in platform.Names)
diff --git a/Engine/Source/Programs/Horde/Horde.Build/Jobs/Artifacts/ArtifactsController.cs b/Engine/Source/Programs/Horde/Horde.Build/Jobs/Artifacts/ArtifactsController.cs
--- a/Engine/Source/Programs/Horde/Horde.Build/Jobs/Artifacts/ArtifactsController.cs
+++ b/Engine/Source/Programs/Horde/Horde.Build/Jobs/Artifacts/ArtifactsController.cs
@@ -31,9 +31,9 @@ namespace Horde.Build.Jobs.Artifacts
 	public class ArtifactsController : ControllerBase
 	{
 		/// <summary>
-		/// Instance of the database service
+		/// Instance of the globals service
 		/// </summary>
-		private readonly MongoService _mongoService;
+		private readonly GlobalsService _globalsService;
 
 		/// <summary>
 		/// Instance of the artifact collection
@@ -53,9 +53,9 @@ namespace Horde.Build.Jobs.Artifacts
 		/// <summary>
 		/// Constructor
 		/// </summary>
-		public ArtifactsController(MongoService mongoService, IArtifactCollection artifactCollection, AclService aclService, JobService jobService)
+		public ArtifactsController(GlobalsService globalsService, IArtifactCollection artifactCollection, AclService aclService, JobService jobService)
 		{
-			_mongoService = mongoService;
+			_globalsService = globalsService;
 			_artifactCollection = artifactCollection;
 			_aclService = aclService;
 			_jobService = jobService;
@@ -149,7 +149,7 @@ namespace Horde.Build.Jobs.Artifacts
 				return Forbid();
 			}
 
-			string? downloadCode = code ? (string?)GetDirectDownloadCodeForJob(jobId) : null;
+			string? downloadCode = code ? (string?)await GetDirectDownloadCodeForJobAsync(jobId) : null;
 
 			List<IArtifact> artifacts = await _artifactCollection.GetArtifactsAsync(jobId, stepId?.ToSubResourceId(), null);
 			return artifacts.ConvertAll(x => new GetArtifactResponse(x, downloadCode).ApplyFilter(filter));
@@ -170,10 +170,10 @@ namespace Horde.Build.Jobs.Artifacts
 		/// </summary>
 		/// <param name="jobId">The job id</param>
 		/// <returns>The download code</returns>
-		string GetDirectDownloadCodeForJob(JobId jobId)
+		async ValueTask<string> GetDirectDownloadCodeForJobAsync(JobId jobId)
 		{
 			Claim downloadClaim = GetDirectDownloadClaim(jobId);
-			return _aclService.IssueBearerToken(new[] { downloadClaim }, TimeSpan.FromHours(4.0));
+			return await _aclService.IssueBearerTokenAsync(new[] { downloadClaim }, TimeSpan.FromHours(4.0));
 		}
 
 		/// <summary>
@@ -199,7 +199,7 @@ namespace Horde.Build.Jobs.Artifacts
 				return Forbid();
 			}
 
-			string? downloadCode = code? (string?)GetDirectDownloadCodeForJob(artifact.JobId) : null;
+			string? downloadCode = code? (string?)await GetDirectDownloadCodeForJobAsync(artifact.JobId) : null;
 			return new GetArtifactResponse(artifact, downloadCode).ApplyFilter(filter);
 		}
 
@@ -300,14 +300,16 @@ namespace Horde.Build.Jobs.Artifacts
 		[Route("/api/v1/artifacts/{artifactId}/download")]
 		public async Task<ActionResult> DownloadArtifact(string artifactId, [FromQuery] string code)
 		{
+			IGlobals globals = await _globalsService.GetAsync();
+
 			TokenValidationParameters parameters = new TokenValidationParameters();
 			parameters.ValidateAudience = false;
 			parameters.RequireExpirationTime = true;
 			parameters.ValidateLifetime = true;
-			parameters.ValidIssuer = _mongoService.JwtIssuer;
+			parameters.ValidIssuer = globals.JwtIssuer;
 			parameters.ValidateIssuer = true;
 			parameters.ValidateIssuerSigningKey = true;
-			parameters.IssuerSigningKey = _mongoService.JwtSigningKey;
+			parameters.IssuerSigningKey = globals.JwtSigningKey;
 
 			JwtSecurityTokenHandler handler = new JwtSecurityTokenHandler();
 			ClaimsPrincipal principal = handler.ValidateToken(code, parameters, out _);
diff --git a/Engine/Source/Programs/Horde/Horde.Build/Jobs/JobTaskSource.cs b/Engine/Source/Programs/Horde/Horde.Build/Jobs/JobTaskSource.cs
--- a/Engine/Source/Programs/Horde/Horde.Build/Jobs/JobTaskSource.cs
+++ b/Engine/Source/Programs/Horde/Horde.Build/Jobs/JobTaskSource.cs
@@ -185,7 +185,7 @@ namespace Horde.Build.Jobs
 			}
 		}
 
-		readonly MongoService _mongoService;
+		readonly GlobalsService _globalsService;
 		readonly StreamService _streamService;
 		readonly ILogFileService _logFileService;
 		readonly IAgentCollection _agentsCollection;
@@ -236,9 +236,9 @@ namespace Horde.Build.Jobs
 		/// <summary>
 		/// Constructor
 		/// </summary>
-		public JobTaskSource(MongoService mongoService, IAgentCollection agents, IJobCollection jobs, IJobStepRefCollection jobStepRefs, IGraphCollection graphs, IPoolCollection pools, IUgsMetadataCollection ugsMetadataCollection, StreamService streamService, ILogFileService logFileService, PerforceLoadBalancer perforceLoadBalancer, IClock clock, IOptionsMonitor<ServerSettings> settings, ILogger<JobTaskSource> logger)
+		public JobTaskSource(GlobalsService globalsService, IAgentCollection agents, IJobCollection jobs, IJobStepRefCollection jobStepRefs, IGraphCollection graphs, IPoolCollection pools, IUgsMetadataCollection ugsMetadataCollection, StreamService streamService, ILogFileService logFileService, PerforceLoadBalancer perforceLoadBalancer, IClock clock, IOptionsMonitor<ServerSettings> settings, ILogger<JobTaskSource> logger)
 		{
-			_mongoService = mongoService;
+			_globalsService = globalsService;
 			_agentsCollection = agents;
 			_jobs = jobs;
 			_jobStepRefs = jobStepRefs;
@@ -782,7 +782,7 @@ namespace Horde.Build.Jobs
 				leaseName.Append(CultureInfo.InvariantCulture, $" - {job.Name}");
 
 				// Get the global settings
-				Globals globals = await _mongoService.GetGlobalsAsync();
+				IGlobals globals = await _globalsService.GetAsync();
 
 				// Encode the payload
 				ExecuteJobTask? task = await CreateExecuteJobTaskAsync(item._stream, job, batch, agent, item._workspace, item._useAutoSdk, logId);
@@ -850,7 +850,7 @@ namespace Horde.Build.Jobs
 			leaseName.Append(CultureInfo.InvariantCulture, $" - {job.Name}");
 
 			// Get the global settings
-			Globals globals = await _mongoService.GetGlobalsAsync();
+			IGlobals globals = await _globalsService.GetAsync();
 
 			// Encode the payload
 			ExecuteJobTask task = new ExecuteJobTask();
@@ -861,7 +861,7 @@ namespace Horde.Build.Jobs
 
 			List<HordeCommon.Rpc.Messages.AgentWorkspace> workspaces = new List<HordeCommon.Rpc.Messages.AgentWorkspace>();
 
-			PerforceCluster? cluster = globals.FindPerforceCluster(workspace.Cluster);
+			PerforceCluster? cluster = globals.Config.FindPerforceCluster(workspace.Cluster);
 			if (cluster == null)
 			{
 				return null;
diff --git a/Engine/Source/Programs/Horde/Horde.Build/Perforce/PerforceController.cs b/Engine/Source/Programs/Horde/Horde.Build/Perforce/PerforceController.cs
--- a/Engine/Source/Programs/Horde/Horde.Build/Perforce/PerforceController.cs
+++ b/Engine/Source/Programs/Horde/Horde.Build/Perforce/PerforceController.cs
@@ -20,9 +20,9 @@ namespace Horde.Build.Perforce
 	public class PerforceController : ControllerBase
 	{
 		/// <summary>
-		/// The database service instance
+		/// The globals service instance
 		/// </summary>
-		private readonly MongoService _mongoService;
+		private readonly GlobalsService _globalsService;
 
 		/// <summary>
 		/// The ACL service instance
@@ -37,9 +37,9 @@ namespace Horde.Build.Perforce
 		/// <summary>
 		/// Constructor
 		/// </summary>
-		public PerforceController(MongoService mongoService, AclService aclService, PerforceLoadBalancer perforceLoadBalancer)
+		public PerforceController(GlobalsService globalsService, AclService aclService, PerforceLoadBalancer perforceLoadBalancer)
 		{
-			_mongoService = mongoService;
+			_globalsService = globalsService;
 			_aclService = aclService;
 			_perforceLoadBalancer = perforceLoadBalancer;
 		}
@@ -75,8 +75,8 @@ namespace Horde.Build.Perforce
 				return Forbid();
 			}
 
-			Globals globals = await _mongoService.GetGlobalsAsync();
-			return globals.PerforceClusters;
+			IGlobals globals = await _globalsService.GetAsync();
+			return globals.Config.PerforceClusters;
 		}
 	}
 
diff --git a/Engine/Source/Programs/Horde/Horde.Build/Perforce/PerforceLoadBalancer.cs b/Engine/Source/Programs/Horde/Horde.Build/Perforce/PerforceLoadBalancer.cs
--- a/Engine/Source/Programs/Horde/Horde.Build/Perforce/PerforceLoadBalancer.cs
+++ b/Engine/Source/Programs/Horde/Horde.Build/Perforce/PerforceLoadBalancer.cs
@@ -137,6 +137,7 @@ namespace Horde.Build.Perforce
 		}
 
 		readonly MongoService _mongoService;
+		readonly GlobalsService _globalsService;
 		readonly ILeaseCollection _leaseCollection;
 		readonly SingletonDocument<PerforceServerList> _serverListSingleton;
 		readonly Random _random = new Random();
@@ -146,9 +147,10 @@ namespace Horde.Build.Perforce
 		/// <summary>
 		/// Constructor
 		/// </summary>
-		public PerforceLoadBalancer(MongoService mongoService, ILeaseCollection leaseCollection, IClock clock, ILogger<PerforceLoadBalancer> logger)
+		public PerforceLoadBalancer(MongoService mongoService, GlobalsService globalsService, ILeaseCollection leaseCollection, IClock clock, ILogger<PerforceLoadBalancer> logger)
 		{
 			_mongoService = mongoService;
+			_globalsService = globalsService;
 			_leaseCollection = leaseCollection;
 			_serverListSingleton = new SingletonDocument<PerforceServerList>(mongoService);
 			_logger = logger;
@@ -356,13 +358,13 @@ namespace Horde.Build.Perforce
 		/// <inheritdoc/>
 		async ValueTask TickInternalAsync(CancellationToken cancellationToken)
 		{
-			Globals globals = await _mongoService.GetGlobalsAsync();
+			IGlobals globals = await _globalsService.GetAsync();
 
 			// Set of new server entries
 			List<PerforceServerEntry> newServers = new List<PerforceServerEntry>();
 
 			// Update the state of all the valid servers
-			foreach (PerforceCluster cluster in globals.PerforceClusters)
+			foreach (PerforceCluster cluster in globals.Config.PerforceClusters)
 			{
 				foreach (PerforceServer server in cluster.Servers)
 				{
diff --git a/Engine/Source/Programs/Horde/Horde.Build/Perforce/PerforceService.cs b/Engine/Source/Programs/Horde/Horde.Build/Perforce/PerforceService.cs
--- a/Engine/Source/Programs/Horde/Horde.Build/Perforce/PerforceService.cs
+++ b/Engine/Source/Programs/Horde/Horde.Build/Perforce/PerforceService.cs
@@ -210,7 +210,7 @@ namespace Horde.Build.Perforce
 		}
 
 		readonly PerforceLoadBalancer _loadBalancer;
-		readonly LazyCachedValue<Task<Globals>> _cachedGlobals;
+		readonly LazyCachedValue<Task<IGlobals>> _cachedGlobals;
 		readonly ILogger _logger;
 
 		// Useful overrides for local debugging with read-only data
@@ -228,10 +228,10 @@ namespace Horde.Build.Perforce
 		/// <summary>
 		/// Constructor
 		/// </summary>
-		public PerforceService(PerforceLoadBalancer loadBalancer, MongoService mongoService, IUserCollection userCollection, IOptions<ServerSettings> settings, ILogger<PerforceService> logger)
+		public PerforceService(PerforceLoadBalancer loadBalancer, GlobalsService globalsService, IUserCollection userCollection, IOptions<ServerSettings> settings, ILogger<PerforceService> logger)
 		{
 			_loadBalancer = loadBalancer;
-			_cachedGlobals = new LazyCachedValue<Task<Globals>>(() => mongoService.GetGlobalsAsync(), TimeSpan.FromSeconds(30.0));
+			_cachedGlobals = new LazyCachedValue<Task<IGlobals>>(async () => await globalsService.GetAsync(), TimeSpan.FromSeconds(30.0));
 			_userCollection = userCollection;
 			_settings = settings.Value;
 			_logger = logger;
@@ -593,9 +593,9 @@ namespace Horde.Build.Perforce
 
 		async Task<PerforceCluster> GetClusterAsync(string? clusterName, string? serverAndPort = null)
 		{
-			Globals globals = await _cachedGlobals.GetCached();
+			IGlobals globals = await _cachedGlobals.GetCached();
 
-			PerforceCluster? cluster = globals.FindPerforceCluster(clusterName, serverAndPort);
+			PerforceCluster? cluster = globals.Config.FindPerforceCluster(clusterName, serverAndPort);
 			if (cluster == null)
 			{
 				throw new Exception($"Unknown Perforce cluster '{clusterName}'");
diff --git a/Engine/Source/Programs/Horde/Horde.Build/Perforce/PerforceServiceCache.cs b/Engine/Source/Programs/Horde/Horde.Build/Perforce/PerforceServiceCache.cs
--- a/Engine/Source/Programs/Horde/Horde.Build/Perforce/PerforceServiceCache.cs
+++ b/Engine/Source/Programs/Horde/Horde.Build/Perforce/PerforceServiceCache.cs
@@ -174,8 +174,8 @@ namespace Horde.Build.Perforce
 		/// <summary>
 		/// Constructor
 		/// </summary>
-		public PerforceServiceCache(PerforceLoadBalancer loadBalancer, MongoService mongoService, IUserCollection userCollection, IStreamCollection streamCollection, IClock clock, IOptions<ServerSettings> settings, IOptions<CommitServiceOptions> options, ILogger<PerforceServiceCache> logger)
-			: base(loadBalancer, mongoService, userCollection, settings, logger)
+		public PerforceServiceCache(PerforceLoadBalancer loadBalancer, MongoService mongoService, GlobalsService globalsService, IUserCollection userCollection, IStreamCollection streamCollection, IClock clock, IOptions<ServerSettings> settings, IOptions<CommitServiceOptions> options, ILogger<PerforceServiceCache> logger)
+			: base(loadBalancer, globalsService, userCollection, settings, logger)
 		{
 			Options = options.Value;
 
diff --git a/Engine/Source/Programs/Horde/Horde.Build/Server/AdminController.cs b/Engine/Source/Programs/Horde/Horde.Build/Server/AdminController.cs
--- a/Engine/Source/Programs/Horde/Horde.Build/Server/AdminController.cs
+++ b/Engine/Source/Programs/Horde/Horde.Build/Server/AdminController.cs
@@ -54,39 +54,19 @@ namespace Horde.Build.Server
 	public class AdminController : HordeControllerBase
 	{
 		readonly AclService _aclService;
-		readonly UpgradeService _upgradeService;
 		readonly IOptionsMonitor<ServerSettings> _settings;
 
 		/// <summary>
 		/// Constructor
 		/// </summary>
 		/// <param name="aclService">The ACL service singleton</param>
-		/// <param name="upgradeService">The upgrade service singelton</param>
 		/// <param name="settings">Server settings</param>
-		public AdminController(AclService aclService, UpgradeService upgradeService, IOptionsMonitor<ServerSettings> settings)
+		public AdminController(AclService aclService, IOptionsMonitor<ServerSettings> settings)
 		{
 			_aclService = aclService;
-			_upgradeService = upgradeService;
 			_settings = settings;
 		}
 
-		/// <summary>
-		/// Upgrade the database to the latest schema
-		/// </summary>
-		/// <param name="fromVersion">The schema version to upgrade from.</param>
-		[HttpPost]
-		[Route("/api/v1/admin/upgradeschema")]
-		public async Task<ActionResult> UpgradeSchemaAsync([FromQuery] int? fromVersion = null)
-		{
-			if (!await _aclService.AuthorizeAsync(AclAction.AdminWrite, User))
-			{
-				return Forbid(AclAction.AdminWrite);
-			}
-
-			await _upgradeService.UpgradeSchemaAsync(fromVersion);
-			return Ok();
-		}
-
 		/// <summary>
 		/// Issues a token for the given roles. Issues a token for the current user if not specified.
 		/// </summary>
@@ -100,7 +80,7 @@ namespace Horde.Build.Server
 				return Forbid(AclAction.IssueBearerToken);
 			}
 
-			return _aclService.IssueBearerToken(User.Claims, GetDefaultExpiryTime());
+			return await _aclService.IssueBearerTokenAsync(User.Claims, GetDefaultExpiryTime());
 		}
 
 		/// <summary>
@@ -120,7 +100,7 @@ namespace Horde.Build.Server
 			List<Claim> claims = new List<Claim>();
 			claims.AddRange(roles.Split('+').Select(x => new Claim(ClaimTypes.Role, x)));
 
-			return _aclService.IssueBearerToken(claims, GetDefaultExpiryTime());
+			return await _aclService.IssueBearerTokenAsync(claims, GetDefaultExpiryTime());
 		}
 
 		/// <summary>
@@ -140,7 +120,7 @@ namespace Horde.Build.Server
 			claims.Add(new AclClaim(ClaimTypes.Name, User.Identity?.Name ?? "Unknown"));
 			claims.Add(AclService.AgentRegistrationClaim);
 
-			return _aclService.IssueBearerToken(claims, null);
+			return await _aclService.IssueBearerTokenAsync(claims, null);
 		}
 
 		/// <summary>
@@ -160,7 +140,7 @@ namespace Horde.Build.Server
 			claims.Add(new AclClaim(ClaimTypes.Name, User.Identity?.Name ?? "Unknown"));
 			claims.Add(AclService.UploadSoftwareClaim);
 
-			return _aclService.IssueBearerToken(claims, null);
+			return await _aclService.IssueBearerTokenAsync(claims, null);
 		}
 
 		/// <summary>
@@ -180,7 +160,7 @@ namespace Horde.Build.Server
 			claims.Add(new AclClaim(ClaimTypes.Name, User.Identity?.Name ?? "Unknown"));
 			claims.Add(AclService.DownloadSoftwareClaim);
 
-			return _aclService.IssueBearerToken(claims, null);
+			return await _aclService.IssueBearerTokenAsync(claims, null);
 		}
 
 		/// <summary>
@@ -200,7 +180,7 @@ namespace Horde.Build.Server
 			claims.Add(new AclClaim(ClaimTypes.Name, User.Identity?.Name ?? "Unknown"));
 			claims.Add(AclService.ConfigureProjectsClaim);
 
-			return _aclService.IssueBearerToken(claims, null);
+			return await _aclService.IssueBearerTokenAsync(claims, null);
 		}
 
 		/// <summary>
@@ -220,7 +200,7 @@ namespace Horde.Build.Server
 			//Claims.Add(new AclClaim(ClaimTypes.Name, User.Identity.Name ?? "Unknown"));
 			claims.Add(AclService.StartChainedJobClaim);
 
-			return _aclService.IssueBearerToken(claims, null);
+			return await _aclService.IssueBearerTokenAsync(claims, null);
 		}
 
 		/// <summary>
diff --git a/Engine/Source/Programs/Horde/Horde.Build/Server/DowntimeService.cs b/Engine/Source/Programs/Horde/Horde.Build/Server/DowntimeService.cs
--- a/Engine/Source/Programs/Horde/Horde.Build/Server/DowntimeService.cs
+++ b/Engine/Source/Programs/Horde/Horde.Build/Server/DowntimeService.cs
@@ -39,7 +39,7 @@ namespace Horde.Build.Server
 			private set;
 		}
 
-		readonly MongoService _mongoService;
+		readonly GlobalsService _globalsService;
 		readonly ITicker _ticker;
 		readonly IOptionsMonitor<ServerSettings> _settings;
 		readonly ILogger _logger;
@@ -47,13 +47,13 @@ namespace Horde.Build.Server
 		/// <summary>
 		/// Constructor
 		/// </summary>
-		/// <param name="mongoService">The database service instance</param>
+		/// <param name="globalsService">The database service instance</param>
 		/// <param name="clock"></param>
 		/// <param name="settings">The server settings</param>
 		/// <param name="logger">Logger instance</param>
-		public DowntimeService(MongoService mongoService, IClock clock, IOptionsMonitor<ServerSettings> settings, ILogger<DowntimeService> logger)
+		public DowntimeService(GlobalsService globalsService, IClock clock, IOptionsMonitor<ServerSettings> settings, ILogger<DowntimeService> logger)
 		{
-			_mongoService = mongoService;
+			_globalsService = globalsService;
 			_settings = settings;
 			_logger = logger;
 
@@ -79,13 +79,13 @@ namespace Horde.Build.Server
 		/// <returns>Async task</returns>
 		async ValueTask TickAsync(CancellationToken stoppingToken)
 		{
-			Globals globals = await _mongoService.GetGlobalsAsync();
+			IGlobals globals = await _globalsService.GetAsync();
 
 			DateTimeOffset now = TimeZoneInfo.ConvertTime(DateTimeOffset.Now, _settings.CurrentValue.TimeZoneInfo);
-			bool bIsActive = globals.ScheduledDowntime.Any(x => x.IsActive(now));
+			bool bIsActive = globals.Config.Downtime.Any(x => x.IsActive(now));
 
 			DateTimeOffset? next = null;
-			foreach (ScheduledDowntime schedule in globals.ScheduledDowntime)
+			foreach (ScheduledDowntime schedule in globals.Config.Downtime)
 			{
 				DateTimeOffset start = schedule.GetNext(now).StartTime;
 				if (next == null || start < next)
diff --git a/Engine/Source/Programs/Horde/Horde.Build/Server/GlobalConfig.cs b/Engine/Source/Programs/Horde/Horde.Build/Server/GlobalConfig.cs
--- a/Engine/Source/Programs/Horde/Horde.Build/Server/GlobalConfig.cs
+++ b/Engine/Source/Programs/Horde/Horde.Build/Server/GlobalConfig.cs
@@ -3,8 +3,10 @@
 using System;
 using System.Collections.Generic;
 using System.ComponentModel.DataAnnotations;
+using System.Linq;
 using EpicGames.Core;
 using EpicGames.Horde.Common;
+using EpicGames.Perforce;
 using EpicGames.Serialization;
 using Horde.Build.Acls;
 using Horde.Build.Projects;
@@ -70,7 +72,52 @@ namespace Horde.Build.Server
 		/// <summary>
 		/// Access control list
 		/// </summary>
-		public AclConfig? Acl { get; set; }
+		public AclConfig Acl { get; set; } = new AclConfig();
+
+		/// <summary>
+		/// Finds a perforce cluster with the given name or that contains the provided server
+		/// </summary>
+		/// <param name="name">Name of the cluster</param>
+		/// <param name="serverAndPort">Find cluster which contains server</param>
+		/// <returns></returns>
+		public PerforceCluster? FindPerforceCluster(string? name, string? serverAndPort = null)
+		{
+			List<PerforceCluster> clusters = PerforceClusters;
+
+			if (serverAndPort != null)
+			{
+				return clusters.FirstOrDefault(x => x.Servers.FirstOrDefault(server => String.Equals(server.ServerAndPort, serverAndPort, StringComparison.OrdinalIgnoreCase)) != null);
+			}
+
+			if (clusters.Count == 0)
+			{
+				clusters = DefaultClusters;
+			}
+
+			if (name == null)
+			{
+				return clusters.FirstOrDefault();
+			}
+			else
+			{
+				return clusters.FirstOrDefault(x => String.Equals(x.Name, name, StringComparison.OrdinalIgnoreCase));
+			}
+		}
+
+		static List<PerforceCluster> DefaultClusters { get; } = GetDefaultClusters();
+
+		static List<PerforceCluster> GetDefaultClusters()
+		{
+			PerforceServer server = new PerforceServer();
+			server.ServerAndPort = PerforceSettings.Default.ServerAndPort;
+
+			PerforceCluster cluster = new PerforceCluster();
+			cluster.Name = "Default";
+			cluster.CanImpersonate = false;
+			cluster.Servers.Add(server);
+
+			return new List<PerforceCluster> { cluster };
+		}
 	}
 
 	/// <summary>
@@ -215,5 +262,238 @@ namespace Horde.Build.Server
 		[CbField("r")]
 		public double Rate { get; set; }
 	}
+
+	/// <summary>
+	/// How frequently the maintence window repeats
+	/// </summary>
+	public enum ScheduledDowntimeFrequency
+	{
+		/// <summary>
+		/// Once
+		/// </summary>
+		Once,
+
+		/// <summary>
+		/// Every day
+		/// </summary>
+		Daily,
+
+		/// <summary>
+		/// Every week
+		/// </summary>
+		Weekly,
+	}
+
+	/// <summary>
+	/// Settings for the maintenance window
+	/// </summary>
+	public class ScheduledDowntime
+	{
+		/// <summary>
+		/// Start time
+		/// </summary>
+		public DateTimeOffset StartTime { get; set; }
+
+		/// <summary>
+		/// Finish time
+		/// </summary>
+		public DateTimeOffset FinishTime { get; set; }
+
+		/// <summary>
+		/// Frequency that the window repeats
+		/// </summary>
+		public ScheduledDowntimeFrequency Frequency { get; set; } = ScheduledDowntimeFrequency.Once;
+
+		/// <summary>
+		/// Gets the next scheduled downtime
+		/// </summary>
+		/// <param name="now">The current time</param>
+		/// <returns>Start and finish time</returns>
+		public (DateTimeOffset StartTime, DateTimeOffset FinishTime) GetNext(DateTimeOffset now)
+		{
+			TimeSpan offset = TimeSpan.Zero;
+			if (Frequency == ScheduledDowntimeFrequency.Daily)
+			{
+				double days = (now - StartTime).TotalDays;
+				if (days >= 1.0)
+				{
+					days -= days % 1.0;
+				}
+				offset = TimeSpan.FromDays(days);
+			}
+			else if (Frequency == ScheduledDowntimeFrequency.Weekly)
+			{
+				double days = (now - StartTime).TotalDays;
+				if (days >= 7.0)
+				{
+					days -= days % 7.0;
+				}
+				offset = TimeSpan.FromDays(days);
+			}
+			return (StartTime + offset, FinishTime + offset);
+		}
+
+		/// <summary>
+		/// Determines if this schedule is active
+		/// </summary>
+		/// <param name="now">The current time</param>
+		/// <returns>True if downtime is active</returns>
+		public bool IsActive(DateTimeOffset now)
+		{
+			if (Frequency == ScheduledDowntimeFrequency.Once)
+			{
+				return now >= StartTime && now < FinishTime;
+			}
+			else if (Frequency == ScheduledDowntimeFrequency.Daily)
+			{
+				double days = (now - StartTime).TotalDays;
+				if (days < 0.0)
+				{
+					return false;
+				}
+				else
+				{
+					return (days % 1.0) < (FinishTime - StartTime).TotalDays;
+				}
+			}
+			else if (Frequency == ScheduledDowntimeFrequency.Weekly)
+			{
+				double days = (now - StartTime).TotalDays;
+				if (days < 0.0)
+				{
+					return false;
+				}
+				else
+				{
+					return (days % 7.0) < (FinishTime - StartTime).TotalDays;
+				}
+			}
+			else
+			{
+				return false;
+			}
+		}
+	}
+
+	/// <summary>
+	/// Path to a platform and stream to use for syncing AutoSDK
+	/// </summary>
+	public class AutoSdkWorkspace
+	{
+		/// <summary>
+		/// Name of this workspace
+		/// </summary>
+		public string? Name { get; set; }
+
+		/// <summary>
+		/// The agent properties to check (eg. "OSFamily=Windows")
+		/// </summary>
+		public List<string> Properties { get; set; } = new List<string>();
+
+		/// <summary>
+		/// Username for logging in to the server
+		/// </summary>
+		public string? UserName { get; set; }
+
+		/// <summary>
+		/// Stream to use
+		/// </summary>
+		[Required]
+		public string? Stream { get; set; }
+	}
+
+	/// <summary>
+	/// Information about an individual Perforce server
+	/// </summary>
+	public class PerforceServer
+	{
+		/// <summary>
+		/// The server and port. The server may be a DNS entry with multiple records, in which case it will be actively load balanced.
+		/// </summary>
+		public string ServerAndPort { get; set; } = "perforce:1666";
+
+		/// <summary>
+		/// Whether to query the healthcheck address under each server
+		/// </summary>
+		public bool HealthCheck { get; set; }
+
+		/// <summary>
+		/// Whether to resolve the DNS entries and load balance between different hosts
+		/// </summary>
+		public bool ResolveDns { get; set; }
+
+		/// <summary>
+		/// Maximum number of simultaneous conforms on this server
+		/// </summary>
+		public int MaxConformCount { get; set; }
+
+		/// <summary>
+		/// Optional condition for a machine to be eligable to use this server
+		/// </summary>
+		public Condition? Condition { get; set; }
+
+		/// <summary>
+		/// List of properties for an agent to be eligable to use this server
+		/// </summary>
+		public List<string>? Properties { get; set; }
+	}
+
+	/// <summary>
+	/// Credentials for a Perforce user
+	/// </summary>
+	public class PerforceCredentials
+	{
+		/// <summary>
+		/// The username
+		/// </summary>
+		public string UserName { get; set; } = String.Empty;
+
+		/// <summary>
+		/// Password for the user
+		/// </summary>
+		public string Password { get; set; } = String.Empty;
+	}
+
+	/// <summary>
+	/// Information about a cluster of Perforce servers. 
+	/// </summary>
+	public class PerforceCluster
+	{
+		/// <summary>
+		/// The default cluster name
+		/// </summary>
+		public const string DefaultName = "Default";
+
+		/// <summary>
+		/// Name of the cluster
+		/// </summary>
+		[Required]
+		public string Name { get; set; } = null!;
+
+		/// <summary>
+		/// Username for Horde to log in to this server. Will use the default user if not set.
+		/// </summary>
+		public string? ServiceAccount { get; set; }
+
+		/// <summary>
+		/// Whether the service account can impersonate other users
+		/// </summary>
+		public bool CanImpersonate { get; set; } = true;
+
+		/// <summary>
+		/// List of servers
+		/// </summary>
+		public List<PerforceServer> Servers { get; set; } = new List<PerforceServer>();
+
+		/// <summary>
+		/// List of server credentials
+		/// </summary>
+		public List<PerforceCredentials> Credentials { get; set; } = new List<PerforceCredentials>();
+
+		/// <summary>
+		/// List of autosdk streams
+		/// </summary>
+		public List<AutoSdkWorkspace> AutoSdk { get; set; } = new List<AutoSdkWorkspace>();
+	}
 }
 
diff --git a/Engine/Source/Programs/Horde/Horde.Build/Server/Globals.cs b/Engine/Source/Programs/Horde/Horde.Build/Server/Globals.cs
deleted file mode 100644
diff --git a/Engine/Source/Programs/Horde/Horde.Build/Server/MongoService.cs b/Engine/Source/Programs/Horde/Horde.Build/Server/MongoService.cs
--- a/Engine/Source/Programs/Horde/Horde.Build/Server/MongoService.cs
+++ b/Engine/Source/Programs/Horde/Horde.Build/Server/MongoService.cs
@@ -227,16 +227,6 @@ namespace Horde.Build.Server
 		/// </summary>
 		ServerSettings Settings { get; }
 
-		/// <summary>
-		/// Name of the JWT issuer
-		/// </summary>
-		public string JwtIssuer { get; }
-
-		/// <summary>
-		/// The signing key for this instance
-		/// </summary>
-		public SymmetricSecurityKey JwtSigningKey { get; }
-
 		/// <summary>
 		/// Logger for this instance
 		/// </summary>
@@ -348,26 +338,6 @@ namespace Horde.Build.Server
 				SingletonsV1 = GetCollection<BsonDocument>("Singletons");
 				SingletonsV2 = GetCollection<BsonDocument>("SingletonsV2");
 
-				Globals globals = GetGlobalsAsync().Result;
-				while (globals.JwtSigningKey == null)
-				{
-					globals.RotateSigningKey();
-					if (!TryUpdateSingletonAsync(globals).Result)
-					{
-						globals = GetGlobalsAsync().Result;
-					}
-				}
-
-				JwtIssuer = Settings.JwtIssuer ?? Dns.GetHostName();
-				if (String.IsNullOrEmpty(Settings.JwtSecret))
-				{
-					JwtSigningKey = new SymmetricSecurityKey(globals.JwtSigningKey);
-				}
-				else
-				{
-					JwtSigningKey = new SymmetricSecurityKey(Convert.FromBase64String(Settings.JwtSecret));
-				}
-
 				Credentials = GetCollection<Credential>("Credentials");
 			}
 			catch (Exception ex)
@@ -795,15 +765,6 @@ namespace Horde.Build.Server
 			}
 		}
 
-		/// <summary>
-		/// Gets a singleton document by id
-		/// </summary>
-		/// <returns>The globals document</returns>
-		public Task<Globals> GetGlobalsAsync()
-		{
-			return GetSingletonAsync<Globals>(() => new Globals() { SchemaVersion = UpgradeService.LatestSchemaVersion });
-		}
-
 		/// <summary>
 		/// Gets a singleton document by id
 		/// </summary>
diff --git a/Engine/Source/Programs/Horde/Horde.Build/Server/Notices/NoticesController.cs b/Engine/Source/Programs/Horde/Horde.Build/Server/Notices/NoticesController.cs
--- a/Engine/Source/Programs/Horde/Horde.Build/Server/Notices/NoticesController.cs
+++ b/Engine/Source/Programs/Horde/Horde.Build/Server/Notices/NoticesController.cs
@@ -39,7 +39,7 @@ namespace Horde.Build.Server.Notices
 		/// <summary>
 		/// cached globals
 		/// </summary>
-		readonly LazyCachedValue<Task<Globals>> _cachedGlobals;
+		readonly LazyCachedValue<Task<IGlobals>> _cachedGlobals;
 
 		/// <summary>
 		/// cached notices
@@ -55,17 +55,17 @@ namespace Horde.Build.Server.Notices
 		/// Constructor
 		/// </summary>
 		/// <param name="settings">The server settings</param>
-		/// <param name="mongoService">The mongo service singleton</param>
+		/// <param name="globalsService">The globals service singleton</param>
 		/// <param name="noticeService">The notice service singleton</param>
 		/// <param name="aclService">The acl service singleton</param>
 		/// <param name="userCollection">The user collection singleton</param>
-		public NoticesController(IOptionsMonitor<ServerSettings> settings, MongoService mongoService, NoticeService noticeService, AclService aclService, IUserCollection userCollection)
+		public NoticesController(IOptionsMonitor<ServerSettings> settings, GlobalsService globalsService, NoticeService noticeService, AclService aclService, IUserCollection userCollection)
 		{			
 			_settings = settings;
 			_aclService = aclService;
 			_userCollection = userCollection;
 			_noticeService = noticeService;
-			_cachedGlobals = new LazyCachedValue<Task<Globals>>(() => mongoService.GetGlobalsAsync(), TimeSpan.FromSeconds(30.0));
+			_cachedGlobals = new LazyCachedValue<Task<IGlobals>>(async () => await globalsService.GetAsync(), TimeSpan.FromSeconds(30.0));
 			_cachedNotices = new LazyCachedValue<Task<List<INotice>>>(() => noticeService.GetNoticesAsync(), TimeSpan.FromMinutes(1));
 		}
 
@@ -130,11 +130,11 @@ namespace Horde.Build.Server.Notices
 		{
 			List<GetNoticeResponse> messages = new List<GetNoticeResponse>();
 
-			Globals globals = await _cachedGlobals.GetCached();
+			IGlobals globals = await _cachedGlobals.GetCached();
 
 			DateTimeOffset now = TimeZoneInfo.ConvertTime(DateTimeOffset.Now, _settings.CurrentValue.TimeZoneInfo);
 						
-			foreach (ScheduledDowntime schedule in globals.ScheduledDowntime)
+			foreach (ScheduledDowntime schedule in globals.Config.Downtime)
 			{
 				DateTimeOffset start = TimeZoneInfo.ConvertTimeBySystemTimeZoneId(schedule.GetNext(now).StartTime, "UTC");
 				DateTimeOffset finish = TimeZoneInfo.ConvertTimeBySystemTimeZoneId(schedule.GetNext(now).FinishTime, "UTC");
diff --git a/Engine/Source/Programs/Horde/Horde.Build/Server/RpcService.cs b/Engine/Source/Programs/Horde/Horde.Build/Server/RpcService.cs
--- a/Engine/Source/Programs/Horde/Horde.Build/Server/RpcService.cs
+++ b/Engine/Source/Programs/Horde/Horde.Build/Server/RpcService.cs
@@ -323,7 +323,7 @@ namespace Horde.Build.Server
 			response.AgentId = agent.Id.ToString();
 			response.SessionId = agent.SessionId.ToString();
 			response.ExpiryTime = Timestamp.FromDateTime(agent.SessionExpiresAt!.Value);
-			response.Token = _agentService.IssueSessionToken(agent.Id, agent.SessionId!.Value);
+			response.Token = await _agentService.IssueSessionTokenAsync(agent.Id, agent.SessionId!.Value);
 			return response;
 		}
 
diff --git a/Engine/Source/Programs/Horde/Horde.Build/Server/UpgradeService.cs b/Engine/Source/Programs/Horde/Horde.Build/Server/UpgradeService.cs
deleted file mode 100644
diff --git a/Engine/Source/Programs/Horde/Horde.Build/Startup.cs b/Engine/Source/Programs/Horde/Horde.Build/Startup.cs
--- a/Engine/Source/Programs/Horde/Horde.Build/Startup.cs
+++ b/Engine/Source/Programs/Horde/Horde.Build/Startup.cs
@@ -409,6 +409,7 @@ namespace Horde.Build
 			services.AddSingleton<RequestTrackerService>();
 			services.AddSingleton<CredentialService>();
 			services.AddSingleton<MongoService>();
+			services.AddSingleton<GlobalsService>();
 			services.AddSingleton<IDogStatsd>(ctx =>
 			{
 				string? datadogAgentHost = Environment.GetEnvironmentVariable("DD_AGENT_HOST");
@@ -447,7 +448,6 @@ namespace Horde.Build
 			services.AddSingleton<IAvatarService, SlackNotificationSink>(sp => sp.GetRequiredService<SlackNotificationSink>());
 			services.AddSingleton<INotificationSink, SlackNotificationSink>(sp => sp.GetRequiredService<SlackNotificationSink>());
 			services.AddSingleton<StreamService>();
-			services.AddSingleton<UpgradeService>();
 			services.AddSingleton<DeviceService>();						
 			services.AddSingleton<NoticeService>();
 
diff --git a/Engine/Source/Programs/Horde/Horde.Build/Tasks/ConformTaskSource.cs b/Engine/Source/Programs/Horde/Horde.Build/Tasks/ConformTaskSource.cs
--- a/Engine/Source/Programs/Horde/Horde.Build/Tasks/ConformTaskSource.cs
+++ b/Engine/Source/Programs/Horde/Horde.Build/Tasks/ConformTaskSource.cs
@@ -38,6 +38,7 @@ namespace Horde.Build.Tasks
 		public override TaskSourceFlags Flags => TaskSourceFlags.AllowWhenDisabled;
 
 		readonly MongoService _mongoService;
+		readonly GlobalsService _globalsService;
 		readonly IAgentCollection _agentCollection;
 		readonly PoolService _poolService;
 		readonly SingletonDocument<ConformList> _conformList;
@@ -49,9 +50,10 @@ namespace Horde.Build.Tasks
 		/// <summary>
 		/// Constructor
 		/// </summary>
-		public ConformTaskSource(MongoService mongoService, IAgentCollection agentCollection, PoolService poolService, ILogFileService logService, PerforceLoadBalancer perforceLoadBalancer, IClock clock, ILogger<ConformTaskSource> logger)
+		public ConformTaskSource(MongoService mongoService, GlobalsService globalsService, IAgentCollection agentCollection, PoolService poolService, ILogFileService logService, PerforceLoadBalancer perforceLoadBalancer, IClock clock, ILogger<ConformTaskSource> logger)
 		{
 			_mongoService = mongoService;
+			_globalsService = globalsService;
 			_agentCollection = agentCollection;
 			_poolService = poolService;
 			_conformList = new SingletonDocument<ConformList>(mongoService);
@@ -198,12 +200,12 @@ namespace Horde.Build.Tasks
 		/// <inheritdoc/>
 		public async Task<bool> GetWorkspacesAsync(IAgent agent, IList<HordeCommon.Rpc.Messages.AgentWorkspace> workspaces)
 		{
-			Globals globals = await _mongoService.GetGlobalsAsync();
+			IGlobals globals = await _globalsService.GetAsync();
 
 			HashSet<AgentWorkspace> conformWorkspaces = await _poolService.GetWorkspacesAsync(agent, DateTime.UtcNow);
 			foreach (AgentWorkspace conformWorkspace in conformWorkspaces)
 			{
-				PerforceCluster? cluster = globals.FindPerforceCluster(conformWorkspace.Cluster);
+				PerforceCluster? cluster = globals.Config.FindPerforceCluster(conformWorkspace.Cluster);
 				if (cluster == null || !await agent.TryAddWorkspaceMessage(conformWorkspace, cluster, _perforceLoadBalancer, workspaces))
 				{
 					return false;
@@ -228,11 +230,11 @@ namespace Horde.Build.Tasks
 		/// <returns>True if the resource was allocated, false otherwise</returns>
 		private async Task<bool> AllocateConformLeaseAsync(AgentId agentId, IEnumerable<HordeCommon.Rpc.Messages.AgentWorkspace> workspaces, LeaseId leaseId)
 		{
-			Globals globals = await _mongoService.GetGlobalsAsync();
+			IGlobals globals = await _globalsService.GetAsync();
 			for (; ; )
 			{
 				ConformList currentValue = await _conformList.GetAsync();
-				if (globals.MaxConformCount != 0 && currentValue.Entries.Count + currentValue.Servers.Sum(x => x.Entries.Count) >= globals.MaxConformCount)
+				if (globals.Config.MaxConformCount != 0 && currentValue.Entries.Count + currentValue.Servers.Sum(x => x.Entries.Count) >= globals.Config.MaxConformCount)
 				{
 					return false;
 				}
@@ -242,7 +244,7 @@ namespace Horde.Build.Tasks
 				{
 					if (servers.Add(workspace.ServerAndPort))
 					{
-						PerforceCluster? cluster = globals.FindPerforceCluster(workspace.Cluster);
+						PerforceCluster? cluster = globals.Config.FindPerforceCluster(workspace.Cluster);
 						if (cluster == null)
 						{
 							_logger.LogWarning("Unable to find perforce cluster '{Cluster}' for conform", workspace.Cluster);
diff --git a/Engine/Source/Programs/Horde/Horde.Build/Utilities/AsyncCachedValue.cs b/Engine/Source/Programs/Horde/Horde.Build/Utilities/AsyncCachedValue.cs
--- a/Engine/Source/Programs/Horde/Horde.Build/Utilities/AsyncCachedValue.cs
+++ b/Engine/Source/Programs/Horde/Horde.Build/Utilities/AsyncCachedValue.cs
@@ -66,6 +66,14 @@ namespace Horde.Build.Utilities
 			_maxRefreshTime = maxRefreshTime;
 		}
 
+		/// <summary>
+		/// Invalidates the current value
+		/// </summary>
+		public void Invalidate()
+		{
+			_current = null;
+		}
+
 		/// <summary>
 		/// Tries to get the current value
 		/// </summary>
diff --git a/Engine/Source/Programs/Horde/Horde.Build/Utilities/SingletonDocument.cs b/Engine/Source/Programs/Horde/Horde.Build/Utilities/SingletonDocument.cs
--- a/Engine/Source/Programs/Horde/Horde.Build/Utilities/SingletonDocument.cs
+++ b/Engine/Source/Programs/Horde/Horde.Build/Utilities/SingletonDocument.cs
@@ -5,11 +5,36 @@ using System.Reflection;
 using System.Threading.Tasks;
 using Horde.Build.Server;
 using MongoDB.Bson;
+using MongoDB.Bson.Serialization.Attributes;
 
 namespace Horde.Build.Utilities
 {
 	using SingletonId = StringId<SingletonBase>;
 
+	/// <summary>
+	/// Base class for singleton documents
+	/// </summary>
+	public abstract class SingletonBase
+	{
+		/// <summary>
+		/// Unique id for this singleton
+		/// </summary>
+		[BsonId]
+		public SingletonId Id { get; set; }
+
+		/// <summary>
+		/// The revision index of this document
+		/// </summary>
+		public int Revision { get; set; }
+
+		/// <summary>
+		/// Callback to allow the singleton to fix up itself after being read
+		/// </summary>
+		public virtual void PostLoad()
+		{
+		}
+	}
+
 	/// <summary>
 	/// Attribute specifying the unique id for a singleton document
 	/// </summary>
