Horde: Missing files.

diff --git a/Engine/Source/Programs/Horde/Horde.Build/Server/GlobalsService.cs b/Engine/Source/Programs/Horde/Horde.Build/Server/GlobalsService.cs
new file mode 100644
--- /dev/null
+++ b/Engine/Source/Programs/Horde/Horde.Build/Server/GlobalsService.cs
@@ -0,0 +1,151 @@
+// Copyright Epic Games, Inc. All Rights Reserved.
+
+using System;
+using System.Collections.Generic;
+using System.ComponentModel.DataAnnotations;
+using System.Linq;
+using System.Net;
+using System.Security.Cryptography;
+using System.Threading;
+using System.Threading.Tasks;
+using EpicGames.Horde.Common;
+using EpicGames.Perforce;
+using Horde.Build.Acls;
+using Horde.Build.Configuration;
+using Horde.Build.Utilities;
+using Microsoft.Extensions.Options;
+using Microsoft.IdentityModel.Tokens;
+using MongoDB.Bson;
+using MongoDB.Bson.Serialization.Attributes;
+
+namespace Horde.Build.Server
+{
+	/// <summary>
+	/// Accessor for the global <see cref="Globals"/> singleton
+	/// </summary>
+	public class GlobalsService
+	{
+		/// <summary>
+		/// Global server settings
+		/// </summary>
+		[SingletonDocument("globals", "5e3981cb28b8ec59cd07184a")]
+		class Globals : SingletonBase, IGlobals
+		{
+			[BsonIgnore]
+			public GlobalsService _owner = null!;
+
+			public ObjectId InstanceId { get; set; }
+			public string ConfigRevision { get; set; } = "default";
+			public byte[]? JwtSigningKey { get; set; }
+			public int? SchemaVersion { get; set; }
+
+			[BsonIgnore]
+			public GlobalConfig Config { get; set; } = null!;
+
+			[BsonIgnore]
+			string IGlobals.JwtIssuer => _owner._jwtIssuer;
+
+			[BsonIgnore]
+			SymmetricSecurityKey IGlobals.JwtSigningKey => new SymmetricSecurityKey(_owner._fixedJwtSecret ?? JwtSigningKey);
+
+			public Globals()
+			{
+				InstanceId = ObjectId.GenerateNewId();
+			}
+
+			public Globals Clone()
+			{
+				return (Globals)MemberwiseClone();
+			}
+
+			public void RotateSigningKey()
+			{
+				JwtSigningKey = RandomNumberGenerator.GetBytes(128);
+			}
+		}
+
+		readonly MongoService _mongoService;
+		readonly ConfigCollection _configCollection;
+		readonly string _jwtIssuer;
+		readonly byte[]? _fixedJwtSecret;
+
+		/// <summary>
+		/// Constructor
+		/// </summary>
+		/// <param name="mongoService"></param>
+		/// <param name="configCollection"></param>
+		/// <param name="settings">Global settings instance</param>
+		public GlobalsService(MongoService mongoService, ConfigCollection configCollection, IOptions<ServerSettings> settings)
+		{
+			_mongoService = mongoService;
+			_configCollection = configCollection;
+
+			if (String.IsNullOrEmpty(settings.Value.JwtIssuer))
+			{
+				_jwtIssuer = Dns.GetHostName();
+			}
+			else
+			{
+				_jwtIssuer = settings.Value.JwtIssuer;
+			}
+
+
+			if (!String.IsNullOrEmpty(settings.Value.JwtSecret))
+			{
+				_fixedJwtSecret = Convert.FromBase64String(settings.Value.JwtSecret);
+			}
+		}
+
+		async Task PostLoadAsync(Globals globals)
+		{
+			globals._owner = this;
+
+			if (String.IsNullOrEmpty(globals.ConfigRevision))
+			{
+				globals.Config = new GlobalConfig();
+			}
+			else
+			{
+				globals.Config = await _configCollection.GetConfigAsync<GlobalConfig>(globals.ConfigRevision);
+			}
+		}
+
+		/// <summary>
+		/// Gets the current globals instance
+		/// </summary>
+		/// <returns>Globals instance</returns>
+		public async ValueTask<IGlobals> GetAsync()
+		{
+			Globals globals = await _mongoService.GetSingletonAsync<Globals>(() => CreateGlobals());
+			await PostLoadAsync(globals);
+			return globals;
+		}
+
+		static Globals CreateGlobals()
+		{
+			Globals globals = new Globals();
+			globals.RotateSigningKey();
+			return globals;
+		}
+
+		/// <summary>
+		/// Try to update the current globals object
+		/// </summary>
+		/// <param name="globals">The current options value</param>
+		/// <param name="configRevision"></param>
+		/// <returns></returns>
+		public async ValueTask<IGlobals?> TryUpdateAsync(IGlobals globals, string? configRevision)
+		{
+			Globals concreteGlobals = ((Globals)globals).Clone();
+			if (configRevision != null)
+			{
+				concreteGlobals.ConfigRevision = configRevision;
+			}
+			if (!await _mongoService.TryUpdateSingletonAsync(concreteGlobals))
+			{
+				return null;
+			}
+			return concreteGlobals;
+		}
+	}
+}
diff --git a/Engine/Source/Programs/Horde/Horde.Build/Server/IGlobals.cs b/Engine/Source/Programs/Horde/Horde.Build/Server/IGlobals.cs
new file mode 100644
--- /dev/null
+++ b/Engine/Source/Programs/Horde/Horde.Build/Server/IGlobals.cs
@@ -0,0 +1,50 @@
+// Copyright Epic Games, Inc. All Rights Reserved.
+
+using System;
+using System.Collections.Generic;
+using System.ComponentModel.DataAnnotations;
+using System.IdentityModel.Tokens.Jwt;
+using System.Linq;
+using System.Security.Claims;
+using System.Security.Cryptography;
+using EpicGames.Horde.Common;
+using EpicGames.Perforce;
+using Horde.Build.Acls;
+using Horde.Build.Utilities;
+using Microsoft.IdentityModel.Tokens;
+using MongoDB.Bson;
+using MongoDB.Bson.Serialization.Attributes;
+
+namespace Horde.Build.Server
+{
+	/// <summary>
+	/// Global server settings
+	/// </summary>
+	public interface IGlobals
+	{
+		/// <summary>
+		/// Unique instance id of this database
+		/// </summary>
+		ObjectId InstanceId { get; }
+
+		/// <summary>
+		/// Revision string for the current config
+		/// </summary>
+		string ConfigRevision { get; }
+
+		/// <summary>
+		/// The global config object
+		/// </summary>
+		GlobalConfig Config { get; }
+
+		/// <summary>
+		/// Issuer for JWT keys. Can be overridden by config file.
+		/// </summary>
+		string JwtIssuer { get; }
+
+		/// <summary>
+		/// The signing key for this server cluster
+		/// </summary>
+		public SymmetricSecurityKey JwtSigningKey { get; }
+	}
+}
