Horde: Fix tests.

diff --git a/Engine/Source/Programs/Horde/Horde.Build.Tests/Compute/ComputeServiceTest.cs b/Engine/Source/Programs/Horde/Horde.Build.Tests/Compute/ComputeServiceTest.cs
--- a/Engine/Source/Programs/Horde/Horde.Build.Tests/Compute/ComputeServiceTest.cs
+++ b/Engine/Source/Programs/Horde/Horde.Build.Tests/Compute/ComputeServiceTest.cs
@@ -8,6 +8,7 @@ using EpicGames.Horde.Compute;
 using EpicGames.Horde.Storage;
 using EpicGames.Serialization;
 using Horde.Build.Agents.Pools;
+using Horde.Build.Projects;
 using Horde.Build.Server;
 using Microsoft.Extensions.DependencyInjection;
 using Microsoft.VisualStudio.TestTools.UnitTesting;
@@ -36,9 +37,11 @@ namespace Horde.Build.Tests.Compute
 		[TestInitialize]
 		public async Task Setup()
 		{
-			Globals globals = await MongoService.GetGlobalsAsync();
-			globals.ComputeClusters.Add(_clusterConfig);
-			Assert.IsTrue(await MongoService.TryUpdateSingletonAsync(globals));
+			GlobalConfig globalConfig = new GlobalConfig();
+			globalConfig.Compute.Add(_clusterConfig);
+
+			await ConfigCollection.AddConfigAsync("globals", globalConfig);
+			Assert.IsNotNull(await GlobalsService.TryUpdateAsync(await GlobalsService.GetAsync(), "globals"));
 		}
 		
 		[TestMethod]
diff --git a/Engine/Source/Programs/Horde/Horde.Build.Tests/DeviceServiceTests.cs b/Engine/Source/Programs/Horde/Horde.Build.Tests/DeviceServiceTests.cs
--- a/Engine/Source/Programs/Horde/Horde.Build.Tests/DeviceServiceTests.cs
+++ b/Engine/Source/Programs/Horde/Horde.Build.Tests/DeviceServiceTests.cs
@@ -95,9 +95,11 @@ namespace Horde.Build.Tests
 				}
 			}
 
-			Globals globals = await MongoService.GetGlobalsAsync();
-			globals.Devices = Devices;
-			Assert.IsTrue(await MongoService.TryUpdateSingletonAsync(globals));
+			GlobalConfig globalConfig = new GlobalConfig();
+			globalConfig.Devices = Devices;
+
+			await ConfigCollection.AddConfigAsync("globals", globalConfig);
+			Assert.IsNotNull(await GlobalsService.TryUpdateAsync(await GlobalsService.GetAsync(), "globals"));
 
 			// add 4 devices to each platform, split between 2 pools
 			for (int i = 1; i < 4; i++)
diff --git a/Engine/Source/Programs/Horde/Horde.Build.Tests/TestSetup.cs b/Engine/Source/Programs/Horde/Horde.Build.Tests/TestSetup.cs
--- a/Engine/Source/Programs/Horde/Horde.Build.Tests/TestSetup.cs
+++ b/Engine/Source/Programs/Horde/Horde.Build.Tests/TestSetup.cs
@@ -99,6 +99,7 @@ namespace Horde.Build.Tests
 		public AgentService AgentService => ServiceProvider.GetRequiredService<AgentService>();
 		public ICommitService CommitService => ServiceProvider.GetRequiredService<ICommitService>();
 		public ComputeService ComputeService => ServiceProvider.GetRequiredService<ComputeService>();
+		public GlobalsService GlobalsService => ServiceProvider.GetRequiredService<GlobalsService>();
 		public MongoService MongoService => ServiceProvider.GetRequiredService<MongoService>();
 		public ITemplateCollection TemplateCollection => ServiceProvider.GetRequiredService<ITemplateCollection>();
 		internal PerforceServiceStub PerforceService => (PerforceServiceStub)ServiceProvider.GetRequiredService<IPerforceService>();
@@ -207,6 +208,7 @@ namespace Horde.Build.Tests
 			services.AddSingleton<ConfigCollection>();
 			services.AddSingleton<ComputeService>();
 			services.AddSingleton<RequestTrackerService>();
+			services.AddSingleton<GlobalsService>();
 			services.AddSingleton<CredentialService>();
 			services.AddSingleton<JobTaskSource>();
 			services.AddSingleton<IDowntimeService, DowntimeServiceStub>();
@@ -226,7 +228,6 @@ namespace Horde.Build.Tests
 			services.AddSingleton<RpcService>();
 			services.AddSingleton<ScheduleService>();
 			services.AddSingleton<StreamService>();
-			services.AddSingleton<UpgradeService>();
 			services.AddSingleton<DeviceService>();
 
 			services.AddSingleton<ConformTaskSource>();
diff --git a/Engine/Source/Programs/Horde/Horde.Build/Server/GlobalsService.cs b/Engine/Source/Programs/Horde/Horde.Build/Server/GlobalsService.cs
--- a/Engine/Source/Programs/Horde/Horde.Build/Server/GlobalsService.cs
+++ b/Engine/Source/Programs/Horde/Horde.Build/Server/GlobalsService.cs
@@ -35,7 +35,7 @@ namespace Horde.Build.Server
 			public GlobalsService _owner = null!;
 
 			public ObjectId InstanceId { get; set; }
-			public string ConfigRevision { get; set; } = "default";
+			public string ConfigRevision { get; set; } = String.Empty;
 			public byte[]? JwtSigningKey { get; set; }
 			public int? SchemaVersion { get; set; }
 
