Horde: Fix Perforce cache falling back to regular P4 query even if entire query is fulfilled by the DB.

diff --git a/Engine/Source/Programs/Horde/Horde.Build/Perforce/PerforceServiceCache.cs b/Engine/Source/Programs/Horde/Horde.Build/Perforce/PerforceServiceCache.cs
--- a/Engine/Source/Programs/Horde/Horde.Build/Perforce/PerforceServiceCache.cs
+++ b/Engine/Source/Programs/Horde/Horde.Build/Perforce/PerforceServiceCache.cs
@@ -447,6 +447,11 @@ namespace Horde.Build.Perforce
 
 			public override async IAsyncEnumerable<ICommit> FindAsync(int? minChange, int? maxChange, int? maxResults, IReadOnlyList<CommitTag>? tags, [EnumeratorCancellation] CancellationToken cancellationToken = default)
 			{
+				if (minChange != null && maxChange != null && minChange.Value > maxChange.Value)
+				{
+					yield break;
+				}
+
 				CacheState state = await _owner._mongoService.GetSingletonAsync<CacheState>();
 				if (state.Clusters.TryGetValue(_stream.Config.ClusterName, out ClusterState? clusterState))
 				{
@@ -503,6 +508,13 @@ namespace Horde.Build.Perforce
 						{
 							maxResults = maxResults.Value - numResults;
 						}
+
+						if (minChange != null && minChange.Value > minReplicatedChange)
+						{
+							yield break;
+						}
+
+						maxChange = minReplicatedChange - 1;
 					}
 				}
 
