Horde: Enable trace logging for PerforceService. Also add support for configuring console log output level.

diff --git a/Engine/Source/Programs/Horde/Horde.Build/Perforce/PerforceServiceCache.cs b/Engine/Source/Programs/Horde/Horde.Build/Perforce/PerforceServiceCache.cs
--- a/Engine/Source/Programs/Horde/Horde.Build/Perforce/PerforceServiceCache.cs
+++ b/Engine/Source/Programs/Horde/Horde.Build/Perforce/PerforceServiceCache.cs
@@ -174,7 +174,7 @@ namespace Horde.Build.Perforce
 		/// <summary>
 		/// Constructor
 		/// </summary>
-		public PerforceServiceCache(PerforceLoadBalancer loadBalancer, MongoService mongoService, GlobalsService globalsService, IUserCollection userCollection, IStreamCollection streamCollection, IClock clock, IOptions<ServerSettings> settings, IOptions<CommitServiceOptions> options, ILogger<PerforceServiceCache> logger)
+		public PerforceServiceCache(PerforceLoadBalancer loadBalancer, MongoService mongoService, GlobalsService globalsService, IUserCollection userCollection, IStreamCollection streamCollection, IClock clock, IOptions<ServerSettings> settings, IOptions<CommitServiceOptions> options, ILogger<PerforceService> logger)
 			: base(loadBalancer, globalsService, userCollection, settings, logger)
 		{
 			Options = options.Value;
diff --git a/Engine/Source/Programs/Horde/Horde.Build/Program.cs b/Engine/Source/Programs/Horde/Horde.Build/Program.cs
--- a/Engine/Source/Programs/Horde/Horde.Build/Program.cs
+++ b/Engine/Source/Programs/Horde/Horde.Build/Program.cs
@@ -43,7 +43,7 @@ namespace Horde.Build
 				{
 					theme = AnsiConsoleTheme.Code;
 				}
-				return sinkConfig.Console(outputTemplate: "[{Timestamp:HH:mm:ss} {Level:w3}] {Indent}{Message:l}{NewLine}{Exception}", theme: theme, restrictedToMinimumLevel: LogEventLevel.Debug);
+				return sinkConfig.Console(outputTemplate: "[{Timestamp:HH:mm:ss} {Level:w3}] {Indent}{Message:l}{NewLine}{Exception}", theme: theme, restrictedToMinimumLevel: settings.ConsoleLogLevel);
 			}
 		}
 
@@ -102,7 +102,7 @@ namespace Horde.Build
 
 	class Program
 	{
-		public static SemVer Version => _version;
+		public static SemVer Version => s_version;
 
 		public static DirectoryReference AppDir { get; } = GetAppDir();
 
@@ -112,7 +112,7 @@ namespace Horde.Build
 
 		public static Type[] ConfigSchemas = FindSchemaTypes();
 
-		static SemVer _version;
+		static SemVer s_version;
 
 		static Type[] FindSchemaTypes()
 		{
@@ -132,11 +132,11 @@ namespace Horde.Build
 			FileVersionInfo versionInfo = FileVersionInfo.GetVersionInfo(Assembly.GetExecutingAssembly().Location);
 			if (String.IsNullOrEmpty(versionInfo.ProductVersion))
 			{
-				_version = SemVer.Parse("0.0.0");
+				s_version = SemVer.Parse("0.0.0");
 			}
 			else
 			{
-				_version = SemVer.Parse(versionInfo.ProductVersion);
+				s_version = SemVer.Parse(versionInfo.ProductVersion);
 			}
 
 			CommandLineArguments arguments = new CommandLineArguments(args);
@@ -265,10 +265,7 @@ namespace Horde.Build
 				}
 
 				// note: this isn't great, though we need it early in server startup, and this is only hit on first server boot where the grpc cert isn't generated/set 
-				if (settings.ServerPrivateCert == null)
-				{
-					settings.ServerPrivateCert = privateCertFile.ToString();
-				}				
+				settings.ServerPrivateCert ??= privateCertFile.ToString();
 			}
 		}
 	}
diff --git a/Engine/Source/Programs/Horde/Horde.Build/ServerSettings.cs b/Engine/Source/Programs/Horde/Horde.Build/ServerSettings.cs
--- a/Engine/Source/Programs/Horde/Horde.Build/ServerSettings.cs
+++ b/Engine/Source/Programs/Horde/Horde.Build/ServerSettings.cs
@@ -8,6 +8,7 @@ using EpicGames.Horde.Storage.Nodes;
 using Horde.Build.Agents.Fleet;
 using Horde.Build.Storage.Backends;
 using Horde.Build.Utilities;
+using Serilog.Events;
 using TimeZoneConverter;
 
 namespace Horde.Build
@@ -207,7 +208,12 @@ namespace Horde.Build
 	{
 		/// <inheritdoc cref="RunMode" />
 		public RunMode[]? RunModes { get; set; } = null;
-		
+
+		/// <summary>
+		/// Output level for console
+		/// </summary>
+		public LogEventLevel ConsoleLogLevel { get; set; } = LogEventLevel.Debug;
+
 		/// <summary>
 		/// Main port for serving HTTP. Uses the default Kestrel port (5000) if not specified.
 		/// </summary>
diff --git a/Engine/Source/Programs/Horde/Horde.Build/appsettings.json b/Engine/Source/Programs/Horde/Horde.Build/appsettings.json
--- a/Engine/Source/Programs/Horde/Horde.Build/appsettings.json
+++ b/Engine/Source/Programs/Horde/Horde.Build/appsettings.json
@@ -1,5 +1,6 @@
 {
 	"Horde": {
+		"ConsoleLogLevel": "Debug",
 		"DashboardUrl": "https://localhost:3000",
 		"DisableAuth": true,
 		"CorsEnabled": true,
@@ -21,74 +22,8 @@
 				"Horde.Build.Authentication": "Warning",
 				"System.Net.Http.HttpClient": "Warning",
 				"Grpc": "Warning",
-				"Horde.Build.Services.Impl.IssueService": "Debug"
-			}
-		}
-	},
-	"AllowedHosts": "*",
-
-	"Mongo": {
-		// this assumes a mongo db were we have localhost exemption, e.g. we do not need to authenticate
-		// if using mongo via a docker container you will need to update this
-		"ConnectionString": "mongodb://localhost:27017",
-		"RequireTls12": true
-	},
-
-	// scylla is normally not enabled in development but we still configure it so we can easily switch
-	"Scylla": {
-		"LocalDatacenterName": "datacenter1",
-		"LocalKeyspaceSuffix": "local"
-	},
-
-	"Horde_Storage": {
-		// assume a single instance is running
-		"LeaderElectionImplementation": "Static",
-
-		// use mongo for any database needs
-		"ReferencesDbImplementation": "Mongo",
-		"ContentIdStoreImplementation": "Mongo",
-		"BlobIndexImplementation": "Mongo",
-
-		// store blobs in the filesystem
-		"StorageImplementations": [
-			"Filesystem"
-		],
-
-		// disable old api that we do not want to use
-		"TransactionLogWriterImplementation": "Memory",
-		"RefDbImplementation": "Memory"
-	},
-
-	"Filesystem": {
-		"RootDir": "$(ExecutableLocation)/BlobStorage"
-	},
-
-	"Replication": {
-		"StateRoot": "$(ExecutableLocation)/ganymede",
-		"ReplicationPollFrequencySeconds": 15,
-		"CurrentSite": "local-test",
-		"Enabled": false
-	},
-
-	"GC": {
-		"BlobCleanupServiceEnabled": false,
-		"CleanOldBlobs": false,
-		"CleanOldRefRecords": false
-	},
-
-	"Jupiter": {
-		"CurrentSite": "local"
-	},
-
-	// disable auth and allow any namespace
-	"Auth": {
-		"Method": "Disabled"
-	},
-
-	"Namespaces": {
-		"Policies": {
-			"*": {
-				"Claims": [ "*" ]
+				"Horde.Build.Issues.IssueService": "Debug",
+				"Horde.Build.Perforce.PerforceService": "Verbose"
 			}
 		}
 	}
diff --git a/Engine/Source/Programs/Shared/EpicGames.Perforce/NativePerforceConnection.cs b/Engine/Source/Programs/Shared/EpicGames.Perforce/NativePerforceConnection.cs
--- a/Engine/Source/Programs/Shared/EpicGames.Perforce/NativePerforceConnection.cs
+++ b/Engine/Source/Programs/Shared/EpicGames.Perforce/NativePerforceConnection.cs
@@ -256,7 +256,10 @@ namespace EpicGames.Perforce
 			}
 		}
 
+		static int s_nextUniqueId;
+
 		IntPtr _client;
+		readonly int _uniqueId;
 		readonly PinnedBuffer[] _buffers;
 		readonly OnBufferReadyFn _onBufferReadyInst;
 		readonly IntPtr _onBufferReadyFnPtr;
@@ -295,6 +298,8 @@ namespace EpicGames.Perforce
 			Settings = settings;
 			Logger = logger;
 
+			_uniqueId = Interlocked.Increment(ref s_nextUniqueId);
+
 			_buffers = new PinnedBuffer[bufferCount];
 			for (int idx = 0; idx < bufferCount; idx++)
 			{
@@ -527,6 +532,9 @@ namespace EpicGames.Perforce
 
 		private void ExecCommand(string command, List<string> args, byte[]? inputData, string? promptResponse, bool interceptIo)
 		{
+			Stopwatch timer = Stopwatch.StartNew();
+			Logger.LogTrace("Conn {ConnectionId}: {Command} {Args}", _uniqueId, command, String.Join(" ", args));
+
 			List<IntPtr> nativeArgs = new List<IntPtr>();
 			try
 			{
@@ -555,6 +563,8 @@ namespace EpicGames.Perforce
 					Marshal.FreeHGlobal(nativeArgs[idx]);
 				}
 			}
+
+			Logger.LogTrace("Conn {ConnectionId}: {Command} {Args} completed in {Time}ms", _uniqueId, command, String.Join(" ", args), timer.ElapsedMilliseconds);
 		}
 
 		/// <summary>
