Horde: Additional logging for Perforce cache.

diff --git a/Engine/Source/Programs/Horde/Horde.Build/Perforce/PerforceServiceCache.cs b/Engine/Source/Programs/Horde/Horde.Build/Perforce/PerforceServiceCache.cs
--- a/Engine/Source/Programs/Horde/Horde.Build/Perforce/PerforceServiceCache.cs
+++ b/Engine/Source/Programs/Horde/Horde.Build/Perforce/PerforceServiceCache.cs
@@ -452,6 +452,9 @@ namespace Horde.Build.Perforce
 					yield break;
 				}
 
+				int numResults = 0;
+				_owner._logger.LogDebug("Finding {StreamId} commits from {MinChange} to {MaxChange} (max: {MaxResults}, tags: {Tags})", _stream.Id, minChange ?? -1, maxChange ?? -1, maxResults ?? -1, (tags == null || tags.Count == 0) ? "none" : String.Join("/", tags.Select(x => x.ToString())));
+
 				CacheState state = await _owner._mongoService.GetSingletonAsync<CacheState>();
 				if (state.Clusters.TryGetValue(_stream.Config.ClusterName, out ClusterState? clusterState))
 				{
@@ -490,7 +493,6 @@ namespace Horde.Build.Perforce
 							filter &= Builders<CachedCommitDoc>.Filter.Gte(x => x.Number, minReplicatedChange);
 						}
 
-						int numResults = 0;
 						using (IAsyncCursor<CachedCommitDoc> cursor = await _owner._commits.Find(filter).SortByDescending(x => x.Number).Limit(maxResults).ToCursorAsync(cancellationToken))
 						{
 							while (await cursor.MoveNextAsync(cancellationToken))
@@ -504,6 +506,8 @@ namespace Horde.Build.Perforce
 							}
 						}
 
+						_owner._logger.LogDebug("Returned {NumResults} cached results (min: {MinReplicatedChange}, max: {MaxReplicatedChange})", numResults, minReplicatedChange, clusterState.MaxChange);
+
 						if (maxResults != null)
 						{
 							maxResults = maxResults.Value - numResults;
@@ -518,8 +522,10 @@ namespace Horde.Build.Perforce
 					}
 				}
 
-				if (maxResults == null || maxResults > 0)
+				if (maxResults == null || maxResults.Value > 0)
 				{
+					_owner._logger.LogDebug("Querying Perforce server for {StreamId} commits from {MinChange} to {MaxChange} (max: {MaxResults}, tags: {Tags})", _stream.Id, minChange ?? -1, maxChange ?? -1, maxResults ?? -1, (tags == null || tags.Count == 0) ? "none" : String.Join("/", tags.Select(x => x.ToString())));
+
 					await foreach (ICommit commit in base.FindAsync(minChange, maxChange, maxResults, tags, cancellationToken))
 					{
 						yield return commit;
