Horde: Few more tweaks to Perforce logging.

diff --git a/Engine/Source/Programs/Horde/Horde.Build/Perforce/PerforceService.cs b/Engine/Source/Programs/Horde/Horde.Build/Perforce/PerforceService.cs
--- a/Engine/Source/Programs/Horde/Horde.Build/Perforce/PerforceService.cs
+++ b/Engine/Source/Programs/Horde/Horde.Build/Perforce/PerforceService.cs
@@ -486,6 +486,7 @@ namespace Horde.Build.Perforce
 				client.Stream = streamName;
 				await perforce.CreateClientAsync(client, cancellationToken);
 
+				_logger.LogDebug("Created client {ClientName} (user: {User}, root: {Root}, stream: {Stream}, server: {ServerId})", client.Name, client.Owner, client.Root, client.Stream, info.ServerId ?? "unknown");
 				return await CreatePooledConnectionAsync(perforce.Settings.ServerAndPort, perforce.Credentials, client, cancellationToken);
 			}
 		}
diff --git a/Engine/Source/Programs/Horde/Horde.Build/Perforce/PerforceServiceCache.cs b/Engine/Source/Programs/Horde/Horde.Build/Perforce/PerforceServiceCache.cs
--- a/Engine/Source/Programs/Horde/Horde.Build/Perforce/PerforceServiceCache.cs
+++ b/Engine/Source/Programs/Horde/Horde.Build/Perforce/PerforceServiceCache.cs
@@ -453,7 +453,7 @@ namespace Horde.Build.Perforce
 				}
 
 				int numResults = 0;
-				_owner._logger.LogDebug("Finding {StreamId} commits from {MinChange} to {MaxChange} (max: {MaxResults}, tags: {Tags})", _stream.Id, minChange ?? -1, maxChange ?? -1, maxResults ?? -1, (tags == null || tags.Count == 0) ? "none" : String.Join("/", tags.Select(x => x.ToString())));
+				_owner._logger.LogDebug("Querying Perforce cache for {StreamId} commits from {MinChange} to {MaxChange} (max: {MaxResults}, tags: {Tags})", _stream.Id, minChange ?? -1, maxChange ?? -1, maxResults ?? -1, (tags == null || tags.Count == 0) ? "none" : String.Join("/", tags.Select(x => x.ToString())));
 
 				CacheState state = await _owner._mongoService.GetSingletonAsync<CacheState>();
 				if (state.Clusters.TryGetValue(_stream.Config.ClusterName, out ClusterState? clusterState))
@@ -506,7 +506,7 @@ namespace Horde.Build.Perforce
 							}
 						}
 
-						_owner._logger.LogDebug("Returned {NumResults} cached results (min: {MinReplicatedChange}, max: {MaxReplicatedChange})", numResults, minReplicatedChange, clusterState.MaxChange);
+						_owner._logger.LogDebug("Found {NumResults} cached results (min-replicated: {MinReplicatedChange}, max-replicated: {MaxReplicatedChange})", numResults, minReplicatedChange, clusterState.MaxChange);
 
 						if (maxResults != null)
 						{
diff --git a/Engine/Source/Programs/Shared/EpicGames.Perforce/NativePerforceConnection.cs b/Engine/Source/Programs/Shared/EpicGames.Perforce/NativePerforceConnection.cs
--- a/Engine/Source/Programs/Shared/EpicGames.Perforce/NativePerforceConnection.cs
+++ b/Engine/Source/Programs/Shared/EpicGames.Perforce/NativePerforceConnection.cs
@@ -564,7 +564,7 @@ namespace EpicGames.Perforce
 				}
 			}
 
-			Logger.LogTrace("Conn {ConnectionId}: {Command} {Args} completed in {Time}ms", _uniqueId, command, String.Join(" ", args), timer.ElapsedMilliseconds);
+			Logger.LogTrace("Conn {ConnectionId}: Request completed in {Time}ms", _uniqueId, timer.ElapsedMilliseconds);
 		}
 
 		/// <summary>
