Horde: Add check for batches being non-empty before trying to append to them.

diff --git a/Engine/Source/Programs/Horde/Horde.Build/Jobs/JobCollection.cs b/Engine/Source/Programs/Horde/Horde.Build/Jobs/JobCollection.cs
--- a/Engine/Source/Programs/Horde/Horde.Build/Jobs/JobCollection.cs
+++ b/Engine/Source/Programs/Horde/Horde.Build/Jobs/JobCollection.cs
@@ -1447,7 +1447,7 @@ namespace Horde.Build.Jobs
 			JobStepBatchDocument?[] appendToBatches = new JobStepBatchDocument?[graph.Groups.Count];
 			foreach (JobStepBatchDocument batch in job.Batches)
 			{
-				if (batch.CanBeAppendedTo())
+				if (batch.CanBeAppendedTo() && batch.Steps.Count > 0)
 				{
 					INodeGroup group = graph.Groups[batch.GroupIdx];
 					INode firstNode = group.Nodes[batch.Steps[0].NodeIdx];
