Horde: Additional logging for incorrect span key error.

diff --git a/Engine/Source/Programs/Horde/Horde.Build/Issues/IssuesController.cs b/Engine/Source/Programs/Horde/Horde.Build/Issues/IssuesController.cs
--- a/Engine/Source/Programs/Horde/Horde.Build/Issues/IssuesController.cs
+++ b/Engine/Source/Programs/Horde/Horde.Build/Issues/IssuesController.cs
@@ -358,9 +358,9 @@ namespace Horde.Build.Issues
 					IStream? stream = await _streamService.GetCachedStream(streamSpans.Key);
 					affectedStreams.Add(new GetIssueAffectedStreamResponse(details, stream, streamSpans));
 				}
-				catch 
+				catch (Exception ex)
 				{
-					_logger.LogError("Unable to get {StreamId} for span key", streamSpans.Key);
+					_logger.LogError(ex, "Unable to get {StreamId} for span key on issue {IssueId}", streamSpans.Key, details.Issue.Id);
 				}
 			}
 			return new GetIssueResponse(details, affectedStreams, showDesktopAlerts);
