Horde: Fix errors querying metadata for shelved changes.

diff --git a/Engine/Source/Programs/Horde/Horde.Build.Tests/Stubs/Services/PerforceServiceStub.cs b/Engine/Source/Programs/Horde/Horde.Build.Tests/Stubs/Services/PerforceServiceStub.cs
--- a/Engine/Source/Programs/Horde/Horde.Build.Tests/Stubs/Services/PerforceServiceStub.cs
+++ b/Engine/Source/Programs/Horde/Horde.Build.Tests/Stubs/Services/PerforceServiceStub.cs
@@ -205,7 +205,7 @@ namespace Horde.Build.Tests.Stubs.Services
 			return Task.CompletedTask;
 		}
 
-		public Task UpdateChangelistDescription(string clusterName, int change, string description, CancellationToken cancellationToken)
+		public Task UpdateChangelistDescription(string clusterName, int change, Func<string, string> description, CancellationToken cancellationToken)
 		{
 			return Task.CompletedTask;
 		}
diff --git a/Engine/Source/Programs/Horde/Horde.Build/Jobs/JobService.cs b/Engine/Source/Programs/Horde/Horde.Build/Jobs/JobService.cs
--- a/Engine/Source/Programs/Horde/Horde.Build/Jobs/JobService.cs
+++ b/Engine/Source/Programs/Horde/Horde.Build/Jobs/JobService.cs
@@ -8,6 +8,7 @@ using System.Security.Claims;
 using System.Text.RegularExpressions;
 using System.Threading.Tasks;
 using EpicGames.Core;
+using EpicGames.Perforce;
 using Horde.Build.Acls;
 using Horde.Build.Agents;
 using Horde.Build.Agents.Leases;
@@ -1095,8 +1096,7 @@ namespace Horde.Build.Jobs
 
 					_logger.LogInformation("Updating description for {ClonedPreflightChange}", clonedPreflightChange);
 
-					ICommit details = await _perforceService.GetChangeDetailsAsync(stream, clonedPreflightChange);
-					await _perforceService.UpdateChangelistDescription(stream.Config.ClusterName, clonedPreflightChange, details.Description.TrimEnd() + $"\n#preflight {job.Id}");
+					await _perforceService.UpdateChangelistDescription(stream.Config.ClusterName, clonedPreflightChange, x => x.TrimEnd() + $"\n#preflight {job.Id}");
 
 					_logger.LogInformation("Submitting change {Change} (through {ChangeCopy}) after successful completion of {JobId}", job.PreflightChange, clonedPreflightChange, job.Id);
 					(change, message) = await _perforceService.SubmitShelvedChangeAsync(stream.Config.ClusterName, clonedPreflightChange, job.PreflightChange);
diff --git a/Engine/Source/Programs/Horde/Horde.Build/Perforce/IPerforceService.cs b/Engine/Source/Programs/Horde/Horde.Build/Perforce/IPerforceService.cs
--- a/Engine/Source/Programs/Horde/Horde.Build/Perforce/IPerforceService.cs
+++ b/Engine/Source/Programs/Horde/Horde.Build/Perforce/IPerforceService.cs
@@ -137,15 +137,6 @@ namespace Horde.Build.Perforce
 		/// <returns>Whether the shelf is valid, plus the description for it</returns>
 		public Task<(CheckShelfResult, string?)> CheckShelfAsync(string clusterName, string streamName, int changeNumber, CancellationToken cancellationToken = default);
 
-		/// <summary>
-		/// Gets the latest change for a particular stream
-		/// </summary>
-		/// <param name="stream">Stream to query</param>
-		/// <param name="changeNumber">Change numbers to query</param>
-		/// <param name="cancellationToken">Cancellation token for the operation</param>
-		/// <returns>Commit details</returns>
-		public Task<ICommit> GetChangeDetailsAsync(IStream stream, int changeNumber, CancellationToken cancellationToken = default);
-
 		/// <summary>
 		/// Duplicates a shelved changelist
 		/// </summary>
@@ -179,10 +170,10 @@ namespace Horde.Build.Perforce
 		/// </summary>
 		/// <param name="clusterName">Name of the Perforce cluster</param>
 		/// <param name="change">The change to update</param>
-		/// <param name="description">The new description</param>
+		/// <param name="updateFunc">Function to update the description</param>
 		/// <param name="cancellationToken">Cancellation token for the operation</param>
 		/// <returns>Async task</returns>
-		public Task UpdateChangelistDescription(string clusterName, int change, string description, CancellationToken cancellationToken = default);
+		public Task UpdateChangelistDescription(string clusterName, int change, Func<string, string> updateFunc, CancellationToken cancellationToken = default);
 
 		/// <summary>
 		/// Creates a commit source for the given stream
diff --git a/Engine/Source/Programs/Horde/Horde.Build/Perforce/PerforceService.cs b/Engine/Source/Programs/Horde/Horde.Build/Perforce/PerforceService.cs
--- a/Engine/Source/Programs/Horde/Horde.Build/Perforce/PerforceService.cs
+++ b/Engine/Source/Programs/Horde/Horde.Build/Perforce/PerforceService.cs
@@ -836,7 +836,7 @@ namespace Horde.Build.Perforce
 		}
 
 		/// <inheritdoc/>
-		public async Task UpdateChangelistDescription(string clusterName, int change, string description, CancellationToken cancellationToken)
+		public async Task UpdateChangelistDescription(string clusterName, int change, Func<string, string> updateFunc, CancellationToken cancellationToken)
 		{
 			using IScope scope = GlobalTracer.Instance.BuildSpan("PerforceService.UpdateChangelistDescription").StartActive();
 			scope.Span.SetTag("ClusterName", clusterName);
@@ -848,13 +848,18 @@ namespace Horde.Build.Perforce
 				using (IPerforceConnection perforce = await ConnectAsChangeOwnerAsync(cluster, change, cancellationToken))
 				{
 					ChangeRecord record = await perforce.GetChangeAsync(GetChangeOptions.None, change, cancellationToken);
-					record.Description = description;
-					await perforce.UpdateChangeAsync(UpdateChangeOptions.None, record, cancellationToken);
+
+					string newDescription = updateFunc(record.Description ?? String.Empty);
+					if (!String.Equals(record.Description, newDescription, StringComparison.Ordinal))
+					{
+						record.Description = newDescription;
+						await perforce.UpdateChangeAsync(UpdateChangeOptions.None, record, cancellationToken);
+					}
 				}
 			}
 			catch (Exception ex)
 			{
-				_logger.LogError(ex, "Unable to update Changelist for CL {Change} to {Description}, {Message}", change, description, ex.Message);
+				_logger.LogError(ex, "Unable to update Changelist for CL {Change}: {Message}", change, ex.Message);
 			}
 		}
 
