Horde: Handle cases where deleted streams do not have a valid config file.

diff --git a/Engine/Source/Programs/Horde/Horde.Build/Streams/StreamCollection.cs b/Engine/Source/Programs/Horde/Horde.Build/Streams/StreamCollection.cs
--- a/Engine/Source/Programs/Horde/Horde.Build/Streams/StreamCollection.cs
+++ b/Engine/Source/Programs/Horde/Horde.Build/Streams/StreamCollection.cs
@@ -93,21 +93,22 @@ namespace Horde.Build.Streams
 			public async Task PostLoadAsync(ConfigCollection configCollection, ILogger logger)
 			{
 				try
+				{
+					Config = await configCollection.GetConfigAsync<StreamConfig>(ConfigRevision);
+				}
+				catch (Exception)
 				{
 					if (Deleted)
 					{
+						logger.LogWarning("Unable to get stream config for {StreamId} at {Revision}; using default.", Id, ConfigRevision);
 						Config = new StreamConfig();
 					}
 					else
 					{
-						Config = await configCollection.GetConfigAsync<StreamConfig>(ConfigRevision);
+						logger.LogError("Unable to get stream config for {StreamId} at {Revision}", Id, ConfigRevision);
+						throw;
 					}
 				}
-				catch (Exception)
-				{
-					logger.LogError("Unable to get stream config for {StreamId} at {Revision}", Id, ConfigRevision);
-					throw;
-				}
 
 				foreach (KeyValuePair<TemplateId, TemplateRefDoc> pair in Templates)
 				{
@@ -148,7 +149,11 @@ namespace Horde.Build.Streams
 			{
 				get
 				{
-					_config ??= _owner?.Config?.Templates.First(x => x.Id == Id);
+					if (_config == null && _owner != null && _owner.Config != null)
+					{
+						// This should generally always succeed, but adding a fallback to handle legacy data.
+						_config = _owner.Config.Templates.FirstOrDefault(x => x.Id == Id) ?? new TemplateRefConfig { Id = Id, Name = Id.ToString() };
+					}
 					return _config!;
 				}
 			}
