Horde: Mitigate issue with very large CL range delta requests
#rnx

diff --git a/Engine/Source/Programs/Horde/HordeDashboard/src/backend/CommitCache.ts b/Engine/Source/Programs/Horde/HordeDashboard/src/backend/CommitCache.ts
--- a/Engine/Source/Programs/Horde/HordeDashboard/src/backend/CommitCache.ts
+++ b/Engine/Source/Programs/Horde/HordeDashboard/src/backend/CommitCache.ts
@@ -45,6 +45,7 @@ export class CommitCache {
 
         const cminCL = 999999999999;
         const cmaxCL = -1;
+        const maxDelta = 131072;
 
         let minCL = cminCL;
         let maxCL = cmaxCL;
@@ -54,11 +55,11 @@ export class CommitCache {
 
             if (!change) {
 
-                if (c < minCL) {
+                if (c < minCL && (minCL === cminCL || Math.abs(minCL - c) < maxDelta)) {
                     minCL = c;
                 }
 
-                if (c > maxCL) {
+                if (c > maxCL && (maxCL === cmaxCL || Math.abs(maxCL - c) < maxDelta)) {
                     maxCL = c;
                 }
             }
