Horde: Fix missing stream client when submitting changes.

diff --git a/Engine/Source/Programs/Horde/Horde.Build.Tests/Stubs/Services/PerforceServiceStub.cs b/Engine/Source/Programs/Horde/Horde.Build.Tests/Stubs/Services/PerforceServiceStub.cs
--- a/Engine/Source/Programs/Horde/Horde.Build.Tests/Stubs/Services/PerforceServiceStub.cs
+++ b/Engine/Source/Programs/Horde/Horde.Build.Tests/Stubs/Services/PerforceServiceStub.cs
@@ -190,7 +190,7 @@ namespace Horde.Build.Tests.Stubs.Services
 			throw new NotImplementedException();
 		}
 
-		public Task<(int? Change, string Message)> SubmitShelvedChangeAsync(string clusterName, int shelvedChange, int originalChange, CancellationToken cancellationToken)
+		public Task<(int? Change, string Message)> SubmitShelvedChangeAsync(IStream stream, int shelvedChange, int originalChange, CancellationToken cancellationToken)
 		{
 			throw new NotImplementedException();
 		}
diff --git a/Engine/Source/Programs/Horde/Horde.Build/Jobs/JobService.cs b/Engine/Source/Programs/Horde/Horde.Build/Jobs/JobService.cs
--- a/Engine/Source/Programs/Horde/Horde.Build/Jobs/JobService.cs
+++ b/Engine/Source/Programs/Horde/Horde.Build/Jobs/JobService.cs
@@ -1099,7 +1099,7 @@ namespace Horde.Build.Jobs
 					await _perforceService.UpdateChangelistDescription(stream.Config.ClusterName, clonedPreflightChange, x => x.TrimEnd() + $"\n#preflight {job.Id}");
 
 					_logger.LogInformation("Submitting change {Change} (through {ChangeCopy}) after successful completion of {JobId}", job.PreflightChange, clonedPreflightChange, job.Id);
-					(change, message) = await _perforceService.SubmitShelvedChangeAsync(stream.Config.ClusterName, clonedPreflightChange, job.PreflightChange);
+					(change, message) = await _perforceService.SubmitShelvedChangeAsync(stream, clonedPreflightChange, job.PreflightChange);
 
 					_logger.LogInformation("Attempt to submit {Change} (through {ChangeCopy}): {Message}", job.PreflightChange, clonedPreflightChange, message);
 
diff --git a/Engine/Source/Programs/Horde/Horde.Build/Perforce/PerforceService.cs b/Engine/Source/Programs/Horde/Horde.Build/Perforce/PerforceService.cs
--- a/Engine/Source/Programs/Horde/Horde.Build/Perforce/PerforceService.cs
+++ b/Engine/Source/Programs/Horde/Horde.Build/Perforce/PerforceService.cs
@@ -1005,15 +1005,14 @@ namespace Horde.Build.Perforce
 		}
 
 		/// <inheritdoc/>
-		public async Task<(int? Change, string Message)> SubmitShelvedChangeAsync(string clusterName, int change, int originalChange, CancellationToken cancellationToken)
+		public async Task<(int? Change, string Message)> SubmitShelvedChangeAsync(IStream stream, int change, int originalChange, CancellationToken cancellationToken)
 		{
 			using IScope scope = GlobalTracer.Instance.BuildSpan("PerforceService.SubmitShelvedChangeAsync").StartActive();
-			scope.Span.SetTag("ClusterName", clusterName);
+			scope.Span.SetTag("Stream", stream.Id);
 			scope.Span.SetTag("Change", change);
 			scope.Span.SetTag("OriginalChange", originalChange);
 
-			PerforceCluster cluster = await GetClusterAsync(clusterName);
-			IPerforceServer server = await GetServerAsync(cluster);
+			PerforceCluster cluster = await GetClusterAsync(stream.Config.ClusterName);
 
 			List<DescribeRecord> records;
 			using (IPerforceConnection perforce = await ConnectAsync(cluster, null, cancellationToken))
@@ -1045,7 +1044,7 @@ namespace Horde.Build.Perforce
 				}
 			}
 
-			using (IPerforceConnection perforce = await ConnectAsync(cluster, records[0].User, cancellationToken))
+			using (IPerforceConnection perforce = await ConnectWithStreamClientAsync(stream, records[0].User, cancellationToken))
 			{
 				try
 				{
