Horde: Missing file.

diff --git a/Engine/Source/Programs/Horde/Horde.Build/Perforce/IPerforceService.cs b/Engine/Source/Programs/Horde/Horde.Build/Perforce/IPerforceService.cs
--- a/Engine/Source/Programs/Horde/Horde.Build/Perforce/IPerforceService.cs
+++ b/Engine/Source/Programs/Horde/Horde.Build/Perforce/IPerforceService.cs
@@ -149,12 +149,12 @@ namespace Horde.Build.Perforce
 		/// <summary>
 		/// Submit a shelved changelist
 		/// </summary>
-		/// <param name="clusterName">Name of the Perforce cluster</param>
+		/// <param name="stream">Stream to submit to</param>
 		/// <param name="shelvedChange">The shelved changelist number, created by <see cref="DuplicateShelvedChangeAsync(String,Int32,CancellationToken)"/></param>
 		/// <param name="originalChange">The original changelist number</param>
 		/// <param name="cancellationToken">Cancellation token for the operation</param>
 		/// <returns>Tuple consisting of the submitted changelist number and message</returns>
-		public Task<(int? Change, string Message)> SubmitShelvedChangeAsync(string clusterName, int shelvedChange, int originalChange, CancellationToken cancellationToken = default);
+		public Task<(int? Change, string Message)> SubmitShelvedChangeAsync(IStream stream, int shelvedChange, int originalChange, CancellationToken cancellationToken = default);
 
 		/// <summary>
 		/// Deletes a changelist containing shelved files.
