Horde: Prevent gated schedules throwing an exception when commit isn't in the current stream.

diff --git a/Engine/Source/Programs/Horde/Horde.Build/Jobs/Schedules/ScheduleService.cs b/Engine/Source/Programs/Horde/Horde.Build/Jobs/Schedules/ScheduleService.cs
--- a/Engine/Source/Programs/Horde/Horde.Build/Jobs/Schedules/ScheduleService.cs
+++ b/Engine/Source/Programs/Horde/Horde.Build/Jobs/Schedules/ScheduleService.cs
@@ -364,41 +364,55 @@ namespace Horde.Build.Jobs.Schedules
 				cancellationToken.ThrowIfCancellationRequested();
 
 				// Get the next valid change
-				ICommit? changeDetails = null;
+				int change;
+				ICommit? commit;
+
 				if (schedule.Config.Gate != null)
 				{
-					changeDetails = await GetNextChangeForGateAsync(stream, templateId, schedule.Config.Gate, minChangeNumber, maxChangeNumber, cancellationToken);
+					commit = null;
+					change = await GetNextChangeForGateAsync(stream, templateId, schedule.Config.Gate, minChangeNumber, maxChangeNumber, cancellationToken);
 				}
 				else if (await commitEnumerator.MoveNextAsync(cancellationToken))
 				{
-					changeDetails = commitEnumerator.Current;
+					commit = commitEnumerator.Current;
+					change = commit.Number;
+				}
+				else
+				{
+					commit = null;
+					change = 0;
 				}
 
 				// Quit if we didn't find anything
-				if (changeDetails == null)
+				if (change <= 0)
 				{
 					break;
 				}
-				if (changeDetails.Number < minChangeNumber)
+				if (change < minChangeNumber)
 				{
 					break;
 				}
-				if (changeDetails.Number == minChangeNumber && (schedule.Config.RequireSubmittedChange || triggerChanges.Count > 0))
+				if (change == minChangeNumber && (schedule.Config.RequireSubmittedChange || triggerChanges.Count > 0))
 				{
 					break;
 				}
 
 				// Adjust the changelist for the desired filter
-				int change = changeDetails.Number;
-				if (await ShouldBuildChangeAsync(changeDetails, schedule.Config.Commits, fileFilter))
+				if (commit == null || await ShouldBuildChangeAsync(commit, schedule.Config.Commits, fileFilter))
 				{
-					ICommit? codeChange = await commits.GetLastCodeChange(changeDetails.Number, cancellationToken);
-					if (codeChange == null)
+					ICommit? lastCodeCommit = await commits.GetLastCodeChange(change, cancellationToken);
+
+					int codeChange = change;
+					if (lastCodeCommit != null)
+					{
+						codeChange = lastCodeCommit.Number;
+					}
+					else
 					{
 						_logger.LogWarning("Unable to find code change for CL {Change}", change);
-						codeChange = changeDetails;
 					}
-					triggerChanges.Add((change, codeChange.Number));
+
+					triggerChanges.Add((change, codeChange));
 				}
 
 				// Check we haven't exceeded the time limit
@@ -492,7 +506,7 @@ namespace Horde.Build.Jobs.Schedules
 		/// Gets the next change to build for a schedule on a gate
 		/// </summary>
 		/// <returns></returns>
-		private async Task<ICommit?> GetNextChangeForGateAsync(IStream stream, TemplateId templateRefId, ScheduleGateConfig gate, int? minChange, int? maxChange, CancellationToken cancellationToken)
+		private async Task<int> GetNextChangeForGateAsync(IStream stream, TemplateId templateRefId, ScheduleGateConfig gate, int? minChange, int? maxChange, CancellationToken cancellationToken)
 		{
 			for (; ; )
 			{
@@ -501,7 +515,7 @@ namespace Horde.Build.Jobs.Schedules
 				List<IJob> jobs = await _jobCollection.FindAsync(streamId: stream.Id, templates: new[] { gate.TemplateId }, minChange: minChange, maxChange: maxChange, count: 1);
 				if (jobs.Count == 0)
 				{
-					return null;
+					return 0;
 				}
 
 				IJob job = jobs[0];
@@ -515,7 +529,7 @@ namespace Horde.Build.Jobs.Schedules
 						JobStepOutcome outcome = state.Value.Item2;
 						if (outcome == JobStepOutcome.Success || outcome == JobStepOutcome.Warnings)
 						{
-							return await _commitService.GetCollection(stream).GetAsync(job.Change, cancellationToken);
+							return job.Change;
 						}
 						_logger.LogInformation("Skipping trigger of {StreamName} template {TemplateId} - last {OtherTemplateRefId} job ({JobId}) ended with errors", stream.Id, templateRefId, gate.TemplateId, job.Id);
 					}
