Horde: Fix gated schedule test.

diff --git a/Engine/Source/Programs/Horde/Horde.Build/Jobs/Schedules/ScheduleService.cs b/Engine/Source/Programs/Horde/Horde.Build/Jobs/Schedules/ScheduleService.cs
--- a/Engine/Source/Programs/Horde/Horde.Build/Jobs/Schedules/ScheduleService.cs
+++ b/Engine/Source/Programs/Horde/Horde.Build/Jobs/Schedules/ScheduleService.cs
@@ -369,8 +369,8 @@ namespace Horde.Build.Jobs.Schedules
 
 				if (schedule.Config.Gate != null)
 				{
-					commit = null;
 					change = await GetNextChangeForGateAsync(stream, templateId, schedule.Config.Gate, minChangeNumber, maxChangeNumber, cancellationToken);
+					commit = await commits.FindAsync(change, change, 1, null, cancellationToken).FirstOrDefaultAsync(cancellationToken); // May be a change in a different stream
 				}
 				else if (await commitEnumerator.MoveNextAsync(cancellationToken))
 				{
@@ -400,9 +400,9 @@ namespace Horde.Build.Jobs.Schedules
 				// Adjust the changelist for the desired filter
 				if (commit == null || await ShouldBuildChangeAsync(commit, schedule.Config.Commits, fileFilter))
 				{
-					ICommit? lastCodeCommit = await commits.GetLastCodeChange(change, cancellationToken);
-
 					int codeChange = change;
+
+					ICommit? lastCodeCommit = await commits.GetLastCodeChange(change, cancellationToken);
 					if (lastCodeCommit != null)
 					{
 						codeChange = lastCodeCommit.Number;
