Horde: Support multiple sizing strategies per pool guarded by conditions

Also changes strategies to be created on a per-need basis rather than be singletons.
Potentially has a drawback on performance, but there are ways to solve that later as well.

diff --git a/Engine/Source/Programs/Horde/Horde.Build.Tests/Fleet/AutoscaleServiceV2Test.cs b/Engine/Source/Programs/Horde/Horde.Build.Tests/Fleet/AutoscaleServiceV2Test.cs
--- a/Engine/Source/Programs/Horde/Horde.Build.Tests/Fleet/AutoscaleServiceV2Test.cs
+++ b/Engine/Source/Programs/Horde/Horde.Build.Tests/Fleet/AutoscaleServiceV2Test.cs
@@ -67,37 +67,6 @@ namespace Horde.Build.Tests.Fleet
 		readonly FleetManagerSpy _fleetManagerSpy = new();
 		readonly IDogStatsd _dogStatsD = new NoOpDogStatsd();
 
-		[TestMethod]
-		public async Task PerPoolStrategy()
-		{
-			using AutoscaleServiceV2 service = GetAutoscaleService(_fleetManagerSpy);
-			
-			IPool pool1 = await PoolService.CreatePoolAsync("bogusPoolLease1", null, true, 0, 0, sizeStrategy: PoolSizeStrategy.LeaseUtilization);
-			IPool pool2 = await PoolService.CreatePoolAsync("bogusPoolLease2", null, true, 0, 0, sizeStrategy: null);
-			IPool pool3 = await PoolService.CreatePoolAsync("bogusPoolJobQueue1", null, true, 0, 0, sizeStrategy: PoolSizeStrategy.JobQueue);
-			IPool pool4 = await PoolService.CreatePoolAsync("bogusPoolJobQueue2", null, true, 0, 0, sizeStrategy: PoolSizeStrategy.JobQueue);
-			IPool pool5 = await PoolService.CreatePoolAsync("bogusPoolNoOp", null, true, 0, 0, sizeStrategy: PoolSizeStrategy.NoOp);
-
-			PoolSizeStrategySpy leaseUtilizationSpy = new();
-			PoolSizeStrategySpy jobQueueSpy = new();
-			PoolSizeStrategySpy noOpSpy = new();
-			PoolSizeStrategySpy computeQueueAwsMetricSpy = new();
-			service.OverridePoolSizeStrategiesDuringTesting(leaseUtilizationSpy, jobQueueSpy, noOpSpy, computeQueueAwsMetricSpy);
-
-			await service.TickLeaderAsync(CancellationToken.None);
-			
-			Assert.AreEqual(1, leaseUtilizationSpy.CallCount);
-			Assert.IsTrue(leaseUtilizationSpy.PoolIdsSeen.Contains(pool1.Id));
-			Assert.IsTrue(leaseUtilizationSpy.PoolIdsSeen.Contains(pool2.Id));
-			
-			Assert.AreEqual(1, jobQueueSpy.CallCount);
-			Assert.IsTrue(jobQueueSpy.PoolIdsSeen.Contains(pool3.Id));
-			Assert.IsTrue(jobQueueSpy.PoolIdsSeen.Contains(pool4.Id));
-			
-			Assert.AreEqual(1, noOpSpy.CallCount);
-			Assert.IsTrue(noOpSpy.PoolIdsSeen.Contains(pool5.Id));
-		}
-		
 		[TestMethod]
 		public async Task OnlyEnabledAgentsAreAutoScaled()
 		{
@@ -106,11 +75,11 @@ namespace Horde.Build.Tests.Fleet
 			await CreateAgentAsync(pool, true);
 			await CreateAgentAsync(pool, true);
 			await CreateAgentAsync(pool, false);
-			
-			Dictionary<PoolSizeStrategy, List<PoolSizeData>> poolSizeDatas = await service.GetPoolSizeDataByStrategyType();
-			List<IAgent> agents = poolSizeDatas.Values.SelectMany(x => x.SelectMany(y => y.Agents)).ToList();
 
-			Assert.AreEqual(2, agents.Count);
+			List<PoolSizeData> poolSizeDatas = await service.GetPoolSizeDataAsync();
+			poolSizeDatas = await service.CalculateDesiredPoolSizesAsync(poolSizeDatas);
+			Assert.AreEqual(1, poolSizeDatas.Count);
+			Assert.AreEqual(2, poolSizeDatas[0].Agents.Count);
 		}
 
 		[TestMethod]
@@ -120,16 +89,16 @@ namespace Horde.Build.Tests.Fleet
 			IPool pool = await PoolService.CreatePoolAsync("testPool", null, true, 0, 0, sizeStrategy: PoolSizeStrategy.NoOp);
 
 			// First scale-out will succeed
-			await service.ResizePools(new() { new PoolSizeData(pool, new List<IAgent>(), 1, "Testing") });
+			await service.ResizePoolsAsync(new() { new PoolSizeData(pool, new List<IAgent>(), 1, "Testing") });
 			Assert.AreEqual(1, _fleetManagerSpy.ExpandPoolAsyncCallCount);
 			
 			// Cannot scale-out due to cool-down
-			await service.ResizePools(new() { new PoolSizeData(pool, new List<IAgent>(), 2, "Testing") });
+			await service.ResizePoolsAsync(new() { new PoolSizeData(pool, new List<IAgent>(), 2, "Testing") });
 			Assert.AreEqual(1, _fleetManagerSpy.ExpandPoolAsyncCallCount);
 
 			// Wait some time and then try again
 			await Clock.AdvanceAsync(TimeSpan.FromHours(2));
-			await service.ResizePools(new() { new PoolSizeData(pool, new List<IAgent>(), 2, "Testing") });
+			await service.ResizePoolsAsync(new() { new PoolSizeData(pool, new List<IAgent>(), 2, "Testing") });
 			Assert.AreEqual(2, _fleetManagerSpy.ExpandPoolAsyncCallCount);
 		}
 		
@@ -142,16 +111,16 @@ namespace Horde.Build.Tests.Fleet
 			IAgent agent2 = await CreateAgentAsync(pool);
 
 			// First scale-out will succeed
-			await service.ResizePools(new() { new PoolSizeData(pool, new () { agent1, agent2 }, 1, "Testing") });
+			await service.ResizePoolsAsync(new() { new PoolSizeData(pool, new () { agent1, agent2 }, 1, "Testing") });
 			Assert.AreEqual(1, _fleetManagerSpy.ShrinkPoolAsyncCallCount);
 			
 			// Cannot scale-out due to cool-down
-			await service.ResizePools(new() { new PoolSizeData(pool, new () { agent1 }, 0, "Testing") });
+			await service.ResizePoolsAsync(new() { new PoolSizeData(pool, new () { agent1 }, 0, "Testing") });
 			Assert.AreEqual(1, _fleetManagerSpy.ShrinkPoolAsyncCallCount);
 
 			// Wait some time and then try again
 			await Clock.AdvanceAsync(TimeSpan.FromHours(2));
-			await service.ResizePools(new() { new PoolSizeData(pool, new () { agent1 }, 0, "Testing") });
+			await service.ResizePoolsAsync(new() { new PoolSizeData(pool, new () { agent1 }, 0, "Testing") });
 			Assert.AreEqual(2, _fleetManagerSpy.ShrinkPoolAsyncCallCount);
 		}
 
@@ -160,11 +129,104 @@ namespace Horde.Build.Tests.Fleet
 			ILogger<AutoscaleServiceV2> logger = ServiceProvider.GetRequiredService<ILogger<AutoscaleServiceV2>>();
 			IOptions<ServerSettings> serverSettingsOpt = ServiceProvider.GetRequiredService<IOptions<ServerSettings>>();
 
-			LeaseUtilizationStrategy leaseUtilizationStrategy = new(AgentCollection, PoolCollection, LeaseCollection, Clock);
-			JobQueueStrategy jobQueueStrategy = new(JobCollection, GraphCollection, StreamService, Clock);
-			ComputeQueueAwsMetricStrategy computeQueueAwsMetricStrategy = new(null!, null!, null!);
-			AutoscaleServiceV2 service = new AutoscaleServiceV2(leaseUtilizationStrategy,jobQueueStrategy, new NoOpPoolSizeStrategy(), computeQueueAwsMetricStrategy, AgentCollection, PoolCollection, fleetManager, _dogStatsD, Clock, serverSettingsOpt, logger);
+			AutoscaleServiceV2 service = new (AgentCollection, GraphCollection, JobCollection, LeaseCollection, PoolCollection,
+				StreamService, fleetManager, _dogStatsD, Clock, serverSettingsOpt, logger);
 			return service;
 		}
 	}
+
+	[TestClass]
+	public class PoolSizeStrategyFactoryTest : TestSetup
+	{
+		private static int s_poolCount;
+
+		[TestMethod]
+		public async Task CreateJobQueueFromLegacySettings()
+		{
+			IPool pool1 = await PoolService.CreatePoolAsync("test1", sizeStrategy: PoolSizeStrategy.JobQueue);
+			Assert.AreEqual(typeof(JobQueueStrategy), AutoscaleServiceV2.CreatePoolSizeStrategy(pool1).GetType());
+			
+			IPool pool2 = await PoolService.CreatePoolAsync("test2", sizeStrategy: PoolSizeStrategy.JobQueue, jobQueueSettings: new JobQueueSettings(22, 33));
+			IPoolSizeStrategy s = AutoscaleServiceV2.CreatePoolSizeStrategy(pool2);
+			Assert.AreEqual(typeof(JobQueueStrategy), s.GetType());
+			Assert.AreEqual(22.0, ((JobQueueStrategy)s).Settings.ScaleOutFactor);
+			Assert.AreEqual(33.0, ((JobQueueStrategy)s).Settings.ScaleInFactor);
+		}
+		
+		[TestMethod]
+		public async Task CreateLeaseUtilizationFromLegacySettings()
+		{
+			IPool pool = await PoolService.CreatePoolAsync("test1", sizeStrategy: PoolSizeStrategy.LeaseUtilization);
+			Assert.AreEqual(typeof(LeaseUtilizationStrategy), AutoscaleServiceV2.CreatePoolSizeStrategy(pool).GetType());
+		}
+
+		[TestMethod]
+		public async Task CreateJobQueueStrategy()
+		{
+			IPoolSizeStrategy s = await CreateStrategy(new PoolSizeStrategyInfo(PoolSizeStrategy.JobQueue, null, "{\"ScaleOutFactor\": 100, \"ScaleInFactor\": 200}"));
+			Assert.AreEqual(typeof(JobQueueStrategy), s.GetType());
+			Assert.AreEqual(100.0, ((JobQueueStrategy)s).Settings.ScaleOutFactor);
+			Assert.AreEqual(200.0, ((JobQueueStrategy)s).Settings.ScaleInFactor);
+		}
+		
+		[TestMethod]
+		public async Task CreateLeaseUtilizationStrategy()
+		{
+			IPoolSizeStrategy s = await CreateStrategy(new PoolSizeStrategyInfo(PoolSizeStrategy.LeaseUtilization, null, "{}"));
+			Assert.AreEqual(typeof(LeaseUtilizationStrategy), s.GetType());
+		}
+		
+		[TestMethod]
+		public async Task CreateNoOpStrategy()
+		{
+			IPoolSizeStrategy s = await CreateStrategy(new PoolSizeStrategyInfo(PoolSizeStrategy.NoOp, null, "{}"));
+			Assert.AreEqual(typeof(NoOpPoolSizeStrategy), s.GetType());
+		}
+		
+		[TestMethod]
+		public async Task UnknownConfigFieldsAreIgnored()
+		{
+			IPoolSizeStrategy s = await CreateStrategy(new PoolSizeStrategyInfo(PoolSizeStrategy.JobQueue, null, "{\"BAD-PROPERTY\": 1337}"));
+			Assert.AreEqual(typeof(JobQueueStrategy), s.GetType());
+			Assert.AreEqual(0.25, ((JobQueueStrategy)s).Settings.ScaleOutFactor);
+			Assert.AreEqual(0.9, ((JobQueueStrategy)s).Settings.ScaleInFactor);
+		}
+		
+		[TestMethod]
+		public async Task CreateFromEmptyStrategyList()
+		{
+			IPoolSizeStrategy s = await CreateStrategy();
+			Assert.AreEqual(typeof(NoOpPoolSizeStrategy), s.GetType());
+		}
+		
+		[TestMethod]
+		public async Task ConditionSimple()
+		{
+			IPoolSizeStrategy s = await CreateStrategy(
+				new PoolSizeStrategyInfo(PoolSizeStrategy.LeaseUtilization, "false", "{}"),
+				new PoolSizeStrategyInfo(PoolSizeStrategy.JobQueue, "true", "{\"ScaleOutFactor\": 100, \"ScaleInFactor\": 200}")
+			);
+			Assert.AreEqual(typeof(JobQueueStrategy), s.GetType());
+		}
+		
+		[TestMethod]
+		public async Task ConditionDayOfWeek()
+		{
+			Clock.UtcNow = new DateTime(2022, 9, 5, 15, 0, 0, DateTimeKind.Utc); // A Monday
+			IPoolSizeStrategy s1 = await CreateStrategy(new PoolSizeStrategyInfo(PoolSizeStrategy.LeaseUtilization, "dayOfWeek == 'monday'", "{}"));
+			
+			Clock.UtcNow = new DateTime(2022, 9, 6, 15, 0, 0, DateTimeKind.Utc); // A Tuesday
+			IPoolSizeStrategy s2 = await CreateStrategy(new PoolSizeStrategyInfo(PoolSizeStrategy.LeaseUtilization, "dayOfWeek == 'monday'", "{}"));
+			
+			Assert.AreEqual(typeof(LeaseUtilizationStrategy), s1.GetType());
+			Assert.AreEqual(typeof(NoOpPoolSizeStrategy), s2.GetType());
+		}
+		
+		private async Task<IPoolSizeStrategy> CreateStrategy(params PoolSizeStrategyInfo[] infos)
+		{
+			IPool pool = await PoolService.CreatePoolAsync("testPool-" + s_poolCount++, null, true, 0, 0, sizeStrategy: PoolSizeStrategy.NoOp);
+			await PoolService.UpdatePoolAsync(pool, newSizeStrategies: infos.ToList());
+			return AutoscaleServiceV2.CreatePoolSizeStrategy(pool);
+		}
+	}
 }
\ No newline at end of file
diff --git a/Engine/Source/Programs/Horde/Horde.Build.Tests/Fleet/JobQueueStrategyTest.cs b/Engine/Source/Programs/Horde/Horde.Build.Tests/Fleet/JobQueueStrategyTest.cs
--- a/Engine/Source/Programs/Horde/Horde.Build.Tests/Fleet/JobQueueStrategyTest.cs
+++ b/Engine/Source/Programs/Horde/Horde.Build.Tests/Fleet/JobQueueStrategyTest.cs
@@ -142,8 +142,9 @@ namespace Horde.Build.Tests.Fleet
 				IJob job = await AddPlaceholderJob(graph, stream.Id, nodeForAgentType1);
 				await JobCollection.TryUpdateBatchAsync(job, graph, job.Batches[0].Id, null, JobStepBatchState.Ready, null);
 			}
-			
-			return (new (JobCollection, GraphCollection, StreamService, Clock), poolSize);
+
+			JobQueueSettings settings = new ();
+			return (new (JobCollection, GraphCollection, StreamService, Clock, settings), poolSize);
 		}
 		
 		private async Task<IJob> AddPlaceholderJob(IGraph graph, StreamId streamId, string nodeNameToExecute)
diff --git a/Engine/Source/Programs/Horde/Horde.Build.Tests/TestSetup.cs b/Engine/Source/Programs/Horde/Horde.Build.Tests/TestSetup.cs
--- a/Engine/Source/Programs/Horde/Horde.Build.Tests/TestSetup.cs
+++ b/Engine/Source/Programs/Horde/Horde.Build.Tests/TestSetup.cs
@@ -95,6 +95,7 @@ namespace Horde.Build.Tests
 		public IDeviceCollection DeviceCollection => ServiceProvider.GetRequiredService<IDeviceCollection>();
 
 		public AclService AclService => ServiceProvider.GetRequiredService<AclService>();
+		public AutoscaleServiceV2 AutoscaleServiceV2 => ServiceProvider.GetRequiredService<AutoscaleServiceV2>();
 		public AgentSoftwareService AgentSoftwareService => ServiceProvider.GetRequiredService<AgentSoftwareService>();
 		public AgentService AgentService => ServiceProvider.GetRequiredService<AgentService>();
 		public ICommitService CommitService => ServiceProvider.GetRequiredService<ICommitService>();
@@ -203,6 +204,7 @@ namespace Horde.Build.Tests
 			services.AddSingleton<AgentService>();
 			services.AddSingleton<AgentSoftwareService>();
 			services.AddSingleton<AutoscaleService>();
+			services.AddSingleton<AutoscaleServiceV2>();
 			services.AddSingleton<AwsReuseFleetManager, AwsReuseFleetManager>();
 			services.AddSingleton<ConsistencyService>();
 			services.AddSingleton<ConfigCollection>();
diff --git a/Engine/Source/Programs/Horde/Horde.Build/Agents/Fleet/AutoscaleServiceV2.cs b/Engine/Source/Programs/Horde/Horde.Build/Agents/Fleet/AutoscaleServiceV2.cs
--- a/Engine/Source/Programs/Horde/Horde.Build/Agents/Fleet/AutoscaleServiceV2.cs
+++ b/Engine/Source/Programs/Horde/Horde.Build/Agents/Fleet/AutoscaleServiceV2.cs
@@ -4,9 +4,14 @@ using System;
 using System.Collections.Generic;
 using System.Diagnostics;
 using System.Linq;
+using System.Text.Json;
 using System.Threading;
 using System.Threading.Tasks;
+using Horde.Build.Agents.Leases;
 using Horde.Build.Agents.Pools;
+using Horde.Build.Jobs;
+using Horde.Build.Jobs.Graphs;
+using Horde.Build.Streams;
 using Horde.Build.Utilities;
 using HordeCommon;
 using Microsoft.Extensions.Hosting;
@@ -25,12 +30,12 @@ namespace Horde.Build.Agents.Fleet
 	/// </summary>
 	public sealed class AutoscaleServiceV2 : IHostedService, IDisposable
 	{
-		private IPoolSizeStrategy _leaseUtilizationStrategy;
-		private IPoolSizeStrategy _jobQueueStrategy;
-		private IPoolSizeStrategy _noOpStrategy;
-		private IPoolSizeStrategy _computeQueueAwsMetricStrategy;
 		private readonly IAgentCollection _agentCollection;
+		private readonly IGraphCollection _graphCollection;
+		private readonly IJobCollection _jobCollection;
+		private readonly ILeaseCollection _leaseCollection;
 		private readonly IPoolCollection _poolCollection;
+		private readonly StreamService _streamService;
 		private readonly IFleetManager _fleetManager;
 		private readonly IDogStatsd _dogStatsd;
 		private readonly IClock _clock;
@@ -44,24 +49,24 @@ namespace Horde.Build.Agents.Fleet
 		/// Constructor
 		/// </summary>
 		public AutoscaleServiceV2(
-			LeaseUtilizationStrategy leaseUtilizationStrategy,
-			JobQueueStrategy jobQueueStrategy,
-			NoOpPoolSizeStrategy noOpStrategy,
-			ComputeQueueAwsMetricStrategy computeQueueAwsMetricStrategy,
 			IAgentCollection agentCollection,
+			IGraphCollection graphCollection,
+			IJobCollection jobCollection,
+			ILeaseCollection leaseCollection,
 			IPoolCollection poolCollection,
+			StreamService streamService,
 			IFleetManager fleetManager,
 			IDogStatsd dogStatsd,
 			IClock clock,
 			IOptions<ServerSettings> settings,
 			ILogger<AutoscaleServiceV2> logger)
 		{
-			_leaseUtilizationStrategy = leaseUtilizationStrategy;
-			_jobQueueStrategy = jobQueueStrategy;
-			_noOpStrategy = noOpStrategy;
-			_computeQueueAwsMetricStrategy = computeQueueAwsMetricStrategy;
 			_agentCollection = agentCollection;
+			_graphCollection = graphCollection;
+			_jobCollection = jobCollection;
+			_leaseCollection = leaseCollection;
 			_poolCollection = poolCollection;
+			_streamService = streamService;
 			_fleetManager = fleetManager;
 			_dogStatsd = dogStatsd;
 			_clock = clock;
@@ -97,7 +102,9 @@ namespace Horde.Build.Agents.Fleet
 			Stopwatch stopwatch = Stopwatch.StartNew();
 			using IScope _ = GlobalTracer.Instance.BuildSpan("AutoscaleService.TickAsync").StartActive();
 			
-			await BatchResizePools(false);
+			List<PoolSizeData> input = await GetPoolSizeDataAsync();
+			List<PoolSizeData> output = await CalculateDesiredPoolSizesAsync(input);
+			await ResizePoolsAsync(output);
 
 			stopwatch.Stop();
 			_logger.LogInformation("Autoscaling pools took {ElapsedTime} ms", stopwatch.ElapsedMilliseconds);
@@ -109,50 +116,34 @@ namespace Horde.Build.Agents.Fleet
 			Stopwatch stopwatch = Stopwatch.StartNew();
 			using IScope _ = GlobalTracer.Instance.BuildSpan("AutoscaleService.TickHighFrequency").StartActive();
 			
-			await BatchResizePools(true);
-			
+			// TODO: Re-enable high frequency scaling (only used for experimental scaling of remote execution agents at the moment)
+			await Task.Delay(0, stoppingToken);
+
 			stopwatch.Stop();
 			_logger.LogInformation("Autoscaling pools (high frequency) took {ElapsedTime} ms", stopwatch.ElapsedMilliseconds);
 		}
-		
-		private async Task BatchResizePools(bool onlyHighFrequency)
-		{
-			// Group pools by strategy type to ensure they can be called once (more optimal)
-			Dictionary<PoolSizeStrategy, List<PoolSizeData>> poolSizeDataByStrategy = await GetPoolSizeDataByStrategyType();
 
-			if (onlyHighFrequency)
-			{
-				poolSizeDataByStrategy = poolSizeDataByStrategy
-					.Where(x => x.Key.SupportsHighFrequency())
-					.ToDictionary(i => i.Key, i => i.Value);
-			}
-			
-			foreach ((PoolSizeStrategy strategyType, List<PoolSizeData> currentData) in poolSizeDataByStrategy)
+		internal async Task<List<PoolSizeData>> CalculateDesiredPoolSizesAsync(List<PoolSizeData> poolSizeDatas)
+		{
+			List<PoolSizeData> result = new();
+			foreach (PoolSizeData psd in poolSizeDatas)
 			{
-				List<PoolSizeData> newData = await GetPoolSizeStrategy(strategyType).CalcDesiredPoolSizesAsync(currentData);
-				await ResizePools(newData);
+				result.Add((await CreatePoolSizeStrategy(psd.Pool).CalcDesiredPoolSizesAsync(new List<PoolSizeData> { psd }))[0]);
 			}
+
+			return result;
 		}
 
-		internal async Task<Dictionary<PoolSizeStrategy, List<PoolSizeData>>> GetPoolSizeDataByStrategyType()
+		internal async Task<List<PoolSizeData>> GetPoolSizeDataAsync()
 		{
 			List<IAgent> agents = await _agentCollection.FindAsync(status: AgentStatus.Ok, enabled: true);
 			List<IAgent> GetAgentsInPool(PoolId poolId) => agents.FindAll(a => a.GetPools().Any(p => p == poolId));
 			List<IPool> pools = await _poolCollection.GetAsync();
 
-			Dictionary<PoolSizeStrategy, List<PoolSizeData>> result = new();
-			foreach (PoolSizeStrategy strategyType in Enum.GetValues<PoolSizeStrategy>())
-			{
-				result[strategyType] = pools
-					.Where(x => x.SizeStrategy == strategyType)
-					.Select(x => new PoolSizeData(x, GetAgentsInPool(x.Id), null))
-					.ToList();
-			}
-
-			return result;
+			return pools.Select(pool => new PoolSizeData(pool, GetAgentsInPool(pool.Id), null)).ToList();
 		}
 
-		internal async Task ResizePools(List<PoolSizeData> poolSizeDatas)
+		internal async Task ResizePoolsAsync(List<PoolSizeData> poolSizeDatas)
 		{
 			foreach (PoolSizeData poolSizeData in poolSizeDatas.OrderByDescending(x => x.Agents.Count))
 			{
@@ -223,35 +214,63 @@ namespace Horde.Build.Agents.Fleet
 		}
 
 		/// <summary>
-		/// Backdoor for tests to override strategies with test doubles
-		/// These cannot supplied in the constructor since the DI requires concrete implementations.
+		/// Instantiate a sizing strategy for the given pool
+		/// Can resolve either from the list of strategies or the old legacy way
 		/// </summary>
-		/// <param name="leaseUtilizationStrategy"></param>
-		/// <param name="jobQueueStrategy"></param>
-		/// <param name="noOpStrategy"></param>
-		/// <param name="computeQueueAwsMetricStrategy"></param>
-		internal void OverridePoolSizeStrategiesDuringTesting(
-			IPoolSizeStrategy leaseUtilizationStrategy,
-			IPoolSizeStrategy jobQueueStrategy,
-			IPoolSizeStrategy noOpStrategy,
-			IPoolSizeStrategy computeQueueAwsMetricStrategy)
-		{
-			_leaseUtilizationStrategy = leaseUtilizationStrategy;
-			_jobQueueStrategy = jobQueueStrategy;
-			_noOpStrategy = noOpStrategy;
-			_computeQueueAwsMetricStrategy = computeQueueAwsMetricStrategy;
-		}
-		
-		private IPoolSizeStrategy GetPoolSizeStrategy(PoolSizeStrategy type)
+		/// <param name="pool">Pool to use</param>
+		/// <returns>A pool sizing strategy with parameters set</returns>
+		/// <exception cref="ArgumentException"></exception>
+		public IPoolSizeStrategy CreatePoolSizeStrategy(IPool pool)
 		{
-			return type switch
+	
+			IEnumerable<string> GetPropValues(string name)
 			{
-				PoolSizeStrategy.LeaseUtilization => _leaseUtilizationStrategy,
-				PoolSizeStrategy.JobQueue => _jobQueueStrategy,
-				PoolSizeStrategy.NoOp => _noOpStrategy,
-				PoolSizeStrategy.ComputeQueueAwsMetric => _computeQueueAwsMetricStrategy,
-				_ => throw new Exception($"Unknown strategy: {type}")
-			};
+				return name switch
+				{
+					"dayOfWeek" => new List<string> { _clock.UtcNow.DayOfWeek.ToString().ToLower() },
+					_ => Array.Empty<string>()
+				};
+			}
+
+			if (pool.SizeStrategies.Count > 0)
+			{
+				foreach (PoolSizeStrategyInfo info in pool.SizeStrategies)
+				{
+					if (info.Condition == null || info.Condition.Evaluate(GetPropValues))
+					{
+						switch (info.Type)
+						{
+							case PoolSizeStrategy.JobQueue:
+								JobQueueSettings? settings = JsonSerializer.Deserialize<JobQueueSettings>(info.Config);
+								if (settings == null) throw new ArgumentException("Unable to deserialize pool sizing config");
+								return new JobQueueStrategy(_jobCollection, _graphCollection, _streamService, _clock, settings);
+							
+							case PoolSizeStrategy.LeaseUtilization:
+								// No settings for lease utilization strategy
+								return new LeaseUtilizationStrategy(_agentCollection, _poolCollection, _leaseCollection, _clock);
+							
+							case PoolSizeStrategy.NoOp:
+								return new NoOpPoolSizeStrategy();
+							
+							default:
+								throw new ArgumentException("Invalid pool size strategy type " + info.Type);
+						}
+					}
+				}
+			}
+
+			// These is the legacy way of creating and configuring strategies (list-based approach above is preferred)
+			switch (pool.SizeStrategy)
+			{
+				case PoolSizeStrategy.JobQueue:
+					return new JobQueueStrategy(_jobCollection, _graphCollection, _streamService, _clock, pool.JobQueueSettings);
+				case PoolSizeStrategy.LeaseUtilization:
+					return new LeaseUtilizationStrategy(_agentCollection, _poolCollection, _leaseCollection, _clock);
+				case PoolSizeStrategy.NoOp:
+					return new NoOpPoolSizeStrategy();
+				default:
+					throw new ArgumentException("Unknown pool size strategy " + pool.SizeStrategy);
+			}
 		}
 	}
 }
diff --git a/Engine/Source/Programs/Horde/Horde.Build/Agents/Fleet/JobQueueStrategy.cs b/Engine/Source/Programs/Horde/Horde.Build/Agents/Fleet/JobQueueStrategy.cs
--- a/Engine/Source/Programs/Horde/Horde.Build/Agents/Fleet/JobQueueStrategy.cs
+++ b/Engine/Source/Programs/Horde/Horde.Build/Agents/Fleet/JobQueueStrategy.cs
@@ -3,6 +3,7 @@
 using System;
 using System.Collections.Generic;
 using System.Linq;
+using System.Text.Json.Serialization;
 using System.Threading.Tasks;
 using Horde.Build.Agents.Pools;
 using Horde.Build.Jobs;
@@ -35,6 +36,14 @@ namespace Horde.Build.Agents.Fleet
 		/// </summary>
 		public double ScaleInFactor { get; set; } = 0.9;
 
+		/// <summary>
+		/// Constructor used for JSON serialization
+		/// </summary>
+		[JsonConstructor]
+		public JobQueueSettings()
+		{
+		}
+		
 		/// <summary>
 		/// Constructor
 		/// </summary>
@@ -55,6 +64,8 @@ namespace Horde.Build.Agents.Fleet
 	/// </summary>
 	public class JobQueueStrategy : IPoolSizeStrategy
 	{
+		internal JobQueueSettings Settings { get; }
+		
 		private readonly IJobCollection _jobs;
 		private readonly IGraphCollection _graphs;
 		private readonly StreamService _streamService;
@@ -82,14 +93,16 @@ namespace Horde.Build.Agents.Fleet
 		/// <param name="graphs"></param>
 		/// <param name="streamService"></param>
 		/// <param name="clock"></param>
+		/// <param name="settings"></param>
 		/// <param name="samplePeriod">Time period for each sample</param>
-		public JobQueueStrategy(IJobCollection jobs, IGraphCollection graphs, StreamService streamService, IClock clock, TimeSpan? samplePeriod = null)
+		public JobQueueStrategy(IJobCollection jobs, IGraphCollection graphs, StreamService streamService, IClock clock, JobQueueSettings? settings = null, TimeSpan? samplePeriod = null)
 		{
 			_jobs = jobs;
 			_graphs = graphs;
 			_streamService = streamService;
 			_clock = clock;
 			_samplePeriod = samplePeriod ?? _samplePeriod;
+			Settings = settings ?? _defaultPoolSettings;
 		}
 
 		/// <inheritdoc/>
diff --git a/Engine/Source/Programs/Horde/Horde.Build/Agents/Pools/IPool.cs b/Engine/Source/Programs/Horde/Horde.Build/Agents/Pools/IPool.cs
--- a/Engine/Source/Programs/Horde/Horde.Build/Agents/Pools/IPool.cs
+++ b/Engine/Source/Programs/Horde/Horde.Build/Agents/Pools/IPool.cs
@@ -11,6 +11,44 @@ namespace Horde.Build.Agents.Pools
 {
 	using PoolId = StringId<IPool>;
 
+	/// <summary>
+	/// Metadata for configuring and picking a pool sizing strategy
+	/// </summary>
+	public class PoolSizeStrategyInfo
+	{
+		/// <summary>
+		/// Strategy implementation to use
+		/// </summary>
+		public PoolSizeStrategy Type { get; set; }
+		
+		/// <summary>
+		/// Condition if this strategy should be enabled (right now, using date/time as a distinguishing factor)
+		/// </summary>
+		public Condition? Condition { get; set; }
+
+		/// <summary>
+		/// Configuration for the strategy, serialized as JSON
+		/// </summary>
+		public string Config { get; set; } = "";
+
+		/// <summary>
+		/// Empty constructor for JSON serialization
+		/// </summary>
+		public PoolSizeStrategyInfo()
+		{
+		}
+		
+		/// <summary>
+		/// Constructor
+		/// </summary>
+		public PoolSizeStrategyInfo(PoolSizeStrategy type, Condition? condition, string config)
+		{
+			Type = type;
+			Condition = condition;
+			Config = config;
+		}
+	}
+	
 	/// <summary>
 	/// A pool of machines
 	/// </summary>
@@ -90,6 +128,11 @@ namespace Horde.Build.Agents.Pools
 		/// Pool sizing strategy to be used for this pool
 		/// </summary>
 		public PoolSizeStrategy SizeStrategy { get; }
+
+		/// <summary>
+		/// List of pool sizing strategies for this pool. The first strategy with a matching condition will be picked.
+		/// </summary>
+		public IReadOnlyList<PoolSizeStrategyInfo> SizeStrategies { get; }
 		
 		/// <summary>
 		/// Settings for lease utilization pool sizing strategy (if used)
diff --git a/Engine/Source/Programs/Horde/Horde.Build/Agents/Pools/IPoolCollection.cs b/Engine/Source/Programs/Horde/Horde.Build/Agents/Pools/IPoolCollection.cs
--- a/Engine/Source/Programs/Horde/Horde.Build/Agents/Pools/IPoolCollection.cs
+++ b/Engine/Source/Programs/Horde/Horde.Build/Agents/Pools/IPoolCollection.cs
@@ -94,6 +94,7 @@ namespace Horde.Build.Agents.Pools
 		/// <param name="scaleOutCooldown">Cooldown time between scale-out events</param>
 		/// <param name="scaleInCooldown">Cooldown time between scale-in events</param>
 		/// <param name="sizeStrategy">Pool sizing strategy</param>
+		/// <param name="newSizeStrategies">List of pool sizing strategies</param>
 		/// <param name="leaseUtilizationSettings">Settings for lease utilization strategy</param>
 		/// <param name="jobQueueSettings">Settings for job queue strategy</param>
 		/// <param name="computeQueueAwsMetricSettings">Settings for compute queue AWS metric strategy</param>
@@ -115,6 +116,7 @@ namespace Horde.Build.Agents.Pools
 			TimeSpan? scaleOutCooldown = null,
 			TimeSpan? scaleInCooldown = null,
 			PoolSizeStrategy? sizeStrategy = null,
+			List<PoolSizeStrategyInfo>? newSizeStrategies = null,
 			LeaseUtilizationSettings? leaseUtilizationSettings = null,
 			JobQueueSettings? jobQueueSettings = null,
 			ComputeQueueAwsMetricSettings? computeQueueAwsMetricSettings = null,
diff --git a/Engine/Source/Programs/Horde/Horde.Build/Agents/Pools/PoolCollection.cs b/Engine/Source/Programs/Horde/Horde.Build/Agents/Pools/PoolCollection.cs
--- a/Engine/Source/Programs/Horde/Horde.Build/Agents/Pools/PoolCollection.cs
+++ b/Engine/Source/Programs/Horde/Horde.Build/Agents/Pools/PoolCollection.cs
@@ -44,6 +44,7 @@ namespace Horde.Build.Agents.Pools
 			public TimeSpan? ScaleInCooldown { get; set; }
 			public PoolSizeStrategy? SizeStrategy { get; set; }
 			PoolSizeStrategy IPool.SizeStrategy => SizeStrategy ?? s_defaultStrategy!.Value;
+			public List<PoolSizeStrategyInfo> SizeStrategies { get; set; } = new ();
 			public LeaseUtilizationSettings? LeaseUtilizationSettings { get; set; }
 			public JobQueueSettings? JobQueueSettings { get; set; }
 			public ComputeQueueAwsMetricSettings? ComputeQueueAwsMetricSettings { get; set; }
@@ -51,7 +52,8 @@ namespace Horde.Build.Agents.Pools
 
 			// Read-only wrappers
 			IReadOnlyList<AgentWorkspace> IPool.Workspaces => Workspaces;
-			IReadOnlyDictionary<string, string> IPool.Properties => Properties;			
+			IReadOnlyDictionary<string, string> IPool.Properties => Properties;
+			IReadOnlyList<PoolSizeStrategyInfo> IPool.SizeStrategies => SizeStrategies;
 
 			public PoolDocument()
 			{
@@ -74,6 +76,7 @@ namespace Horde.Build.Agents.Pools
 				ScaleOutCooldown = other.ScaleOutCooldown;
 				ScaleInCooldown = other.ScaleInCooldown;
 				SizeStrategy = other.SizeStrategy;
+				SizeStrategies.AddRange(other.SizeStrategies);
 				LeaseUtilizationSettings = other.LeaseUtilizationSettings;
 				JobQueueSettings = other.JobQueueSettings;
 				ComputeQueueAwsMetricSettings = other.ComputeQueueAwsMetricSettings;
@@ -228,6 +231,7 @@ namespace Horde.Build.Agents.Pools
 			TimeSpan? scaleOutCooldown,
 			TimeSpan? scaleInCooldown, 
 			PoolSizeStrategy? sizeStrategy,
+			List<PoolSizeStrategyInfo>? newSizeStrategies,
 			LeaseUtilizationSettings? leaseUtilizationSettings,
 			JobQueueSettings? jobQueueSettings, 
 			ComputeQueueAwsMetricSettings? computeQueueAwsMetricSettings, 
@@ -329,6 +333,10 @@ namespace Horde.Build.Agents.Pools
 			{
 				transaction.Set(x => x.SizeStrategy, null);
 			}
+			if (newSizeStrategies != null)
+			{
+				transaction.Set(x => x.SizeStrategies, newSizeStrategies);				
+			}
 
 			if (leaseUtilizationSettings != null)
 			{
diff --git a/Engine/Source/Programs/Horde/Horde.Build/Agents/Pools/PoolMessages.cs b/Engine/Source/Programs/Horde/Horde.Build/Agents/Pools/PoolMessages.cs
--- a/Engine/Source/Programs/Horde/Horde.Build/Agents/Pools/PoolMessages.cs
+++ b/Engine/Source/Programs/Horde/Horde.Build/Agents/Pools/PoolMessages.cs
@@ -246,6 +246,11 @@ namespace Horde.Build.Agents.Pools
 		/// Pool sizing strategy
 		/// </summary>
 		public PoolSizeStrategy? SizeStrategy { get; set; }
+		
+		/// <summary>
+		/// Pool sizing strategy
+		/// </summary>
+		public List<PoolSizeStrategy>? SizeStrategies { get; set; }
 
 		/// <summary>
 		/// Set pool to use default strategy
diff --git a/Engine/Source/Programs/Horde/Horde.Build/Agents/Pools/PoolService.cs b/Engine/Source/Programs/Horde/Horde.Build/Agents/Pools/PoolService.cs
--- a/Engine/Source/Programs/Horde/Horde.Build/Agents/Pools/PoolService.cs
+++ b/Engine/Source/Programs/Horde/Horde.Build/Agents/Pools/PoolService.cs
@@ -126,6 +126,7 @@ namespace Horde.Build.Agents.Pools
 		/// <param name="scaleOutCooldown">Cooldown time between scale-out events</param>
 		/// <param name="scaleInCooldown">Cooldown time between scale-in events</param>
 		/// <param name="sizeStrategy">New pool sizing strategy for the pool</param>
+		/// <param name="newSizeStrategies">New pool sizing strategies for the pool</param>
 		/// <param name="leaseUtilizationSettings">Settings for lease utilization-based strategy</param>
 		/// <param name="jobQueueSettings">Settings for job queue-based strategy</param>
 		/// <param name="computeQueueAwsMetricSettings">Settings for compute queue AWS metric strategy</param>
@@ -143,6 +144,7 @@ namespace Horde.Build.Agents.Pools
 			TimeSpan? scaleOutCooldown = null,
 			TimeSpan? scaleInCooldown = null,
 			PoolSizeStrategy? sizeStrategy = null,
+			List<PoolSizeStrategyInfo>? newSizeStrategies = null,
 			LeaseUtilizationSettings? leaseUtilizationSettings = null,
 			JobQueueSettings? jobQueueSettings = null,
 			ComputeQueueAwsMetricSettings? computeQueueAwsMetricSettings = null,
@@ -162,6 +164,7 @@ namespace Horde.Build.Agents.Pools
 					scaleOutCooldown: scaleOutCooldown,
 					scaleInCooldown: scaleInCooldown,
 					sizeStrategy: sizeStrategy,
+					newSizeStrategies: newSizeStrategies,
 					leaseUtilizationSettings: leaseUtilizationSettings,
 					jobQueueSettings: jobQueueSettings,
 					computeQueueAwsMetricSettings: computeQueueAwsMetricSettings,
diff --git a/Engine/Source/Programs/Horde/Horde.Build/Agents/Pools/PoolsController.cs b/Engine/Source/Programs/Horde/Horde.Build/Agents/Pools/PoolsController.cs
--- a/Engine/Source/Programs/Horde/Horde.Build/Agents/Pools/PoolsController.cs
+++ b/Engine/Source/Programs/Horde/Horde.Build/Agents/Pools/PoolsController.cs
@@ -150,7 +150,7 @@ namespace Horde.Build.Agents.Pools
 			TimeSpan? scaleInCooldown = update.ScaleInCooldown == null ? null : TimeSpan.FromSeconds(update.ScaleInCooldown.Value);
 
 			await _poolService.UpdatePoolAsync(pool, update.Name, update.Condition, update.EnableAutoscaling,
-				update.MinAgents, update.NumReserveAgents, update.Properties, conformInterval, scaleOutCooldown, scaleInCooldown, update.SizeStrategy,
+				update.MinAgents, update.NumReserveAgents, update.Properties, conformInterval, scaleOutCooldown, scaleInCooldown, update.SizeStrategy, null,
 				update.LeaseUtilizationSettings?.Convert(), update.JobQueueSettings?.Convert(), update.ComputeQueueAwsMetricSettings?.Convert(), update.UseDefaultStrategy);
 			return new OkResult();
 		}
diff --git a/Engine/Source/Programs/Horde/Horde.Build/Utilities/ConditionSerializer.cs b/Engine/Source/Programs/Horde/Horde.Build/Utilities/ConditionSerializer.cs
--- a/Engine/Source/Programs/Horde/Horde.Build/Utilities/ConditionSerializer.cs
+++ b/Engine/Source/Programs/Horde/Horde.Build/Utilities/ConditionSerializer.cs
@@ -33,9 +33,16 @@ namespace Horde.Build.Utilities
 		}
 
 		/// <inheritdoc/>
-		public void Serialize(BsonSerializationContext context, BsonSerializationArgs args, Condition value)
+		public void Serialize(BsonSerializationContext context, BsonSerializationArgs args, Condition? value)
 		{
-			context.Writer.WriteString(value.ToString());
+			if (value == null)
+			{
+				context.Writer.WriteNull();				
+			}
+			else
+			{
+				context.Writer.WriteString(value.ToString());				
+			}
 		}
 	}
 }
