Horde: Fix incorrect P4 connection for submitting changes through edge servers.

diff --git a/Engine/Source/Programs/Horde/Horde.Build/Perforce/PerforceService.cs b/Engine/Source/Programs/Horde/Horde.Build/Perforce/PerforceService.cs
--- a/Engine/Source/Programs/Horde/Horde.Build/Perforce/PerforceService.cs
+++ b/Engine/Source/Programs/Horde/Horde.Build/Perforce/PerforceService.cs
@@ -1044,7 +1044,7 @@ namespace Horde.Build.Perforce
 				}
 			}
 
-			using (IPerforceConnection perforce = await ConnectWithStreamClientAsync(stream, records[0].User, cancellationToken))
+			using (IPerforceConnection perforce = await ConnectAsChangeOwnerAsync(cluster, change, cancellationToken))
 			{
 				try
 				{
