Horde: Fix errors failing to update Slack messages. Posting messages to channels will automatically resolve user ids to specific DM channels, which must be used to update the message later.

diff --git a/Engine/Source/Programs/Horde/Horde.Build/Notifications/Sinks/SlackNotificationSink.cs b/Engine/Source/Programs/Horde/Horde.Build/Notifications/Sinks/SlackNotificationSink.cs
--- a/Engine/Source/Programs/Horde/Horde.Build/Notifications/Sinks/SlackNotificationSink.cs
+++ b/Engine/Source/Programs/Horde/Horde.Build/Notifications/Sinks/SlackNotificationSink.cs
@@ -153,6 +153,9 @@ namespace Horde.Build.Notifications.Sinks
 
 			[BsonElement("lnk"), BsonIgnoreIfNull]
 			public string? Permalink { get; set; }
+
+			[BsonIgnore]
+			public SlackMessageId MessageId => new SlackMessageId(Channel, null, Ts);
 		}
 
 		class SlackUserDocument : IAvatar
@@ -277,26 +280,26 @@ namespace Horde.Build.Notifications.Sinks
 
 		#region Message state 
 
-		async Task<(MessageStateDocument, bool)> AddOrUpdateMessageStateAsync(string recipient, string eventId, UserId? userId, string digest, string? ts = null)
+		async Task<(MessageStateDocument, bool)> AddOrUpdateMessageStateAsync(string recipient, string eventId, UserId? userId, string digest, SlackMessageId? messageId)
 		{
 			ObjectId newId = ObjectId.GenerateNewId();
 
 			FilterDefinition<MessageStateDocument> filter = Builders<MessageStateDocument>.Filter.Eq(x => x.Recipient, recipient) & Builders<MessageStateDocument>.Filter.Eq(x => x.EventId, eventId);
 			UpdateDefinition<MessageStateDocument> update = Builders<MessageStateDocument>.Update.SetOnInsert(x => x.Id, newId).Set(x => x.UserId, userId).Set(x => x.Digest, digest);
 
-			if (ts != null)
+			if (messageId != null)
 			{
-				update = update.Set(x => x.Ts, ts);
+				update = update.Set(x => x.Channel, messageId.Channel).Set(x => x.Ts, messageId.Ts);
 			}
 
 			MessageStateDocument state = await _messageStates.FindOneAndUpdateAsync(filter, update, new FindOneAndUpdateOptions<MessageStateDocument> { IsUpsert = true, ReturnDocument = ReturnDocument.After });
 			if (state.Id == newId)
 			{
-				_logger.LogInformation("Posted message {MessageId} (ts: {Ts}, recipient: {Recipient}, user: {UserId}, event: {EventId}, digest: {Digest})", state.Id, state.Ts, recipient, userId ?? UserId.Empty, eventId, digest);
+				_logger.LogInformation("Posted message {StateId} (recipient: {Recipient}, user: {UserId}, event: {EventId}, messageId: {MessageId}, digest: {Digest})", state.Id, recipient, userId ?? UserId.Empty, eventId, state.MessageId, digest);
 			}
 			else
 			{
-				_logger.LogInformation("Updated message {MessageId} (ts: {Ts}, recipient: {Recipient}, user: {UserId}, event: {EventId}, digest: {Digest})", state.Id, state.Ts, recipient, userId ?? UserId.Empty, eventId, digest);
+				_logger.LogInformation("Updated message {StateId} (recipient: {Recipient}, user: {UserId}, event: {EventId}, messageId: {MessageId}, digest: {Digest})", state.Id, recipient, userId ?? UserId.Empty, eventId, state.MessageId, digest);
 			}
 
 			return (state, state.Id == newId);
@@ -308,10 +311,11 @@ namespace Horde.Build.Notifications.Sinks
 			return await _messageStates.Find(filter).FirstOrDefaultAsync();
 		}
 
-		async Task SetMessageTimestampAsync(ObjectId messageId, string channel, string ts, string? permalink = null)
+		async Task UpdateMessageStateAsync(ObjectId stateId, SlackMessageId id, string? permalink = null)
 		{
-			FilterDefinition<MessageStateDocument> filter = Builders<MessageStateDocument>.Filter.Eq(x => x.Id, messageId);
-			UpdateDefinition<MessageStateDocument> update = Builders<MessageStateDocument>.Update.Set(x => x.Channel, channel).Set(x => x.Ts, ts);
+			FilterDefinition<MessageStateDocument> filter = Builders<MessageStateDocument>.Filter.Eq(x => x.Id, stateId);
+			UpdateDefinition<MessageStateDocument> update = Builders<MessageStateDocument>.Update.Set(x => x.Channel, id.Channel).Set(x => x.Ts, id.Ts);
+
 			if (permalink == null)
 			{
 				update = update.Unset(x => x.Permalink);
@@ -320,8 +324,9 @@ namespace Horde.Build.Notifications.Sinks
 			{
 				update = update.Set(x => x.Permalink, permalink);
 			}
+
 			await _messageStates.FindOneAndUpdateAsync(filter, update);
-			_logger.LogInformation("Updated message {MessageId} (ts: {Ts}, recipient: {Recipient}, permalink: {Permalink})", messageId, ts, channel, permalink ?? "(n/a)");
+			_logger.LogInformation("Updated message {StateId} (messageId: {MessageId}, permalink: {Permalink})", stateId, id, permalink ?? "(n/a)");
 		}
 
 		#endregion
@@ -810,6 +815,8 @@ namespace Horde.Build.Notifications.Sinks
 				List<IIssueSuspect> suspects = await _issueService.Collection.FindSuspectsAsync(issue);
 
 				(MessageStateDocument state, bool isNew) = await SendOrUpdateMessageAsync(triageChannel, eventId, null, text);
+				SlackMessageId threadId = state.MessageId;
+
 				if (isNew)
 				{
 					// Create the summary text
@@ -849,16 +856,11 @@ namespace Horde.Build.Notifications.Sinks
 						message.Blocks.Add(new SectionBlock(new TextObject("```...```")));
 					}
 
-					string? summaryTs = await _slackClient.PostMessageAsync(triageChannel, state.Ts, message);
+					SlackMessageId summaryId = await _slackClient.PostMessageToThreadAsync(threadId, message);
 
 					// Permalink to the summary text so we link inside the thread rather than just to the original message
-					string? permalink = null;
-					if (summaryTs != null)
-					{
-						permalink = await _slackClient.GetPermalinkAsync(state.Channel, summaryTs);
-					}
-
-					await SetMessageTimestampAsync(state.Id, state.Channel, state.Ts, permalink);
+					string? permalink = await _slackClient.GetPermalinkAsync(summaryId);
+					await UpdateMessageStateAsync(state.Id, state.MessageId, permalink);
 				}
 
 				// Post a message containing the controls and status
@@ -917,7 +919,7 @@ namespace Horde.Build.Notifications.Sinks
 						}
 					}
 
-					await SendOrUpdateMessageAsync(triageChannel, state.Ts, $"{eventId}_buttons", null, message);
+					await SendOrUpdateMessageToThreadAsync(triageChannel, $"{eventId}_buttons", null, threadId, message);
 				}
 
 				bool notifyTriageAlias = false;
@@ -935,7 +937,7 @@ namespace Horde.Build.Notifications.Sinks
 							mention += $" ({changes})";
 						}
 
-						await _slackClient.PostMessageAsync(triageChannel, state.Ts, $"Assigned to {mention}");
+						await _slackClient.PostMessageToThreadAsync(threadId, $"Assigned to {mention}");
 						inviteUserIds.Add(issue.OwnerId.Value);
 					}
 					else
@@ -953,7 +955,7 @@ namespace Horde.Build.Notifications.Sinks
 							}
 
 							string suspectMessage = $"Possibly {StringUtils.FormatList(suspectList, "or")}.";
-							await _slackClient.PostMessageAsync(triageChannel, state.Ts, suspectMessage);
+							await _slackClient.PostMessageToThreadAsync(threadId, suspectMessage);
 						}
 						else
 						{
@@ -988,26 +990,26 @@ namespace Horde.Build.Notifications.Sinks
 					if (triageAlias != null)
 					{
 						string triageMessage = $"(cc {FormatUserOrGroupMention(triageAlias)} for triage).";
-						await SendOrUpdateMessageAsync(triageChannel, state.Ts, eventId + "_triage", null, triageMessage);
+						await SendOrUpdateMessageToThreadAsync(triageChannel, eventId + "_triage", null, threadId, triageMessage);
 					}
 				}
 
 				if (issue.AcknowledgedAt != null)
 				{
-					await _slackClient.AddReactionAsync(state.Channel, state.Ts, "eyes");
+					await _slackClient.AddReactionAsync(threadId, "eyes");
 				}
 				else
 				{
-					await _slackClient.RemoveReactionAsync(state.Channel, state.Ts, "eyes");
+					await _slackClient.RemoveReactionAsync(threadId, "eyes");
 				}
 
 				if (issue.QuarantinedByUserId != null)
 				{
-					await _slackClient.AddReactionAsync(state.Channel, state.Ts, "mask");
+					await _slackClient.AddReactionAsync(threadId, "mask");
 				}
 				else
 				{
-					await _slackClient.RemoveReactionAsync(state.Channel, state.Ts, "mask");
+					await _slackClient.RemoveReactionAsync(threadId, "mask");
 				}
 
 				IIssueSpan? fixFailedSpan = null;
@@ -1019,7 +1021,7 @@ namespace Horde.Build.Notifications.Sinks
 					{
 						string fixedEventId = $"issue_{issue.Id}_fixed_{issue.FixChange}";
 						string fixedMessage = $"Marked as fixed in {FormatChange(issue.FixChange.Value)}";
-						await PostSingleMessageToThreadAsync(triageChannel, state.Ts, fixedEventId, fixedMessage);
+						await PostSingleMessageToThreadAsync(triageChannel, fixedEventId, threadId, fixedMessage);
 					}
 					else
 					{
@@ -1030,7 +1032,7 @@ namespace Horde.Build.Notifications.Sinks
 							string mention = await FormatMentionAsync(issue.OwnerId.Value, workflow.AllowMentions);
 							fixFailedMessage += $" ({mention})";
 						}
-						await PostSingleMessageToThreadAsync(triageChannel, state.Ts, fixFailedEventId, fixFailedMessage);
+						await PostSingleMessageToThreadAsync(triageChannel, fixFailedEventId, threadId, fixFailedMessage);
 					}
 
 					if (fixFailedSpan == null)
@@ -1042,7 +1044,7 @@ namespace Horde.Build.Notifications.Sinks
 								string streamName = spans.FirstOrDefault(x => x.StreamId == stream.StreamId)?.StreamName ?? stream.StreamId.ToString();
 								string missingEventId = $"issue_{issue.Id}_fixmissing_{issue.FixChange}_{stream.StreamId}";
 								string missingMessage = $"Note: Fix may need manually merging to {streamName}";
-								await PostSingleMessageToThreadAsync(triageChannel, state.Ts, missingEventId, missingMessage);
+								await PostSingleMessageToThreadAsync(triageChannel, missingEventId, threadId, missingMessage);
 							}
 						}
 					}
@@ -1050,41 +1052,38 @@ namespace Horde.Build.Notifications.Sinks
 
 				if (fixFailedSpan != null)
 				{
-					await _slackClient.AddReactionAsync(state.Channel, state.Ts, "x");
+					await _slackClient.AddReactionAsync(threadId, "x");
 				}
 				else
 				{
-					await _slackClient.RemoveReactionAsync(state.Channel, state.Ts, "x");
+					await _slackClient.RemoveReactionAsync(threadId, "x");
 				}
 
 				if (issue.ResolvedAt != null && fixFailedSpan == null)
 				{
-					await _slackClient.AddReactionAsync(state.Channel, state.Ts, "tick");
+					await _slackClient.AddReactionAsync(threadId, "tick");
 				}
 				else
 				{
-					await _slackClient.RemoveReactionAsync(state.Channel, state.Ts, "tick");
+					await _slackClient.RemoveReactionAsync(threadId, "tick");
 				}
 
 				if (issue.ExternalIssueKey != null)
 				{
 					string extIssueEventId = $"issue_{issue.Id}_ext_{issue.ExternalIssueKey}";
 					string extIssueMessage = $"Linked to issue {FormatExternalIssue(issue.ExternalIssueKey)}";
-					await PostSingleMessageToThreadAsync(triageChannel, state.Ts, extIssueEventId, extIssueMessage);
+					await PostSingleMessageToThreadAsync(triageChannel, extIssueEventId, threadId, extIssueMessage);
 				}
 			}
 		}
 
-		async Task PostSingleMessageToThreadAsync(string channel, string threadTs, string eventId, string message)
+		async Task PostSingleMessageToThreadAsync(string recipient, string eventId, SlackMessageId threadId, string message)
 		{
-			(MessageStateDocument state, bool isNew) = await AddOrUpdateMessageStateAsync(channel, eventId, null, "");
+			(MessageStateDocument state, bool isNew) = await AddOrUpdateMessageStateAsync(recipient, eventId, null, "", null);
 			if (isNew)
 			{
-				string? ts = await _slackClient.PostMessageAsync(channel, threadTs, message);
-				if (ts != null)
-				{
-					await SetMessageTimestampAsync(state.Id, channel, ts);
-				}
+				SlackMessageId messageId = await _slackClient.PostMessageToThreadAsync(threadId, message);
+				await UpdateMessageStateAsync(state.Id, messageId);
 			}
 		}
 
@@ -1477,12 +1476,12 @@ namespace Horde.Build.Notifications.Sinks
 			SlackMessage headerMessage = new SlackMessage();
 			headerMessage.AddHeader($"Summary for {report.Stream.Name}");
 
-			string? ts = await SendMessageAsync(channel, headerMessage);
-			if (ts != null)
+			SlackMessageId? messageId = await SendMessageAsync(channel, headerMessage);
+			if (messageId != null)
 			{
 				string reportEventId = GetReportEventId(report.Stream.Id, report.WorkflowId);
 				string json = JsonSerializer.Serialize(state, _jsonSerializerOptions);
-				await AddOrUpdateMessageStateAsync(channel, reportEventId, null, json, ts);
+				await AddOrUpdateMessageStateAsync(channel, reportEventId, null, json, messageId);
 
 				if (state.Blocks.Count == 0)
 				{
@@ -1492,7 +1491,7 @@ namespace Horde.Build.Notifications.Sinks
 
 				for (int idx = 0; idx < state.Blocks.Count; idx++)
 				{
-					string blockEventId = GetReportBlockEventId(ts, idx);
+					string blockEventId = GetReportBlockEventId(messageId.Ts, idx);
 					await UpdateReportBlockAsync(channel, blockEventId, time, report.Stream, state.Blocks[idx].TemplateId, issuesByBlock[idx], state.Blocks[idx].TemplateHeader);
 				}
 
@@ -1984,7 +1983,7 @@ namespace Horde.Build.Notifications.Sinks
 			return userDocument;
 		}
 
-		private async Task<string?> SendMessageAsync(string recipient, SlackMessage message)
+		private async Task<SlackMessageId?> SendMessageAsync(string recipient, SlackMessage message)
 		{
 			if (!IsRecipientAllowed(recipient, message.Text))
 			{
@@ -2008,10 +2007,10 @@ namespace Horde.Build.Notifications.Sinks
 
 		private async Task<(MessageStateDocument, bool)> SendOrUpdateMessageAsync(string recipient, string eventId, UserId? userId, SlackMessage message)
 		{
-			return await SendOrUpdateMessageAsync(recipient, null, eventId, userId, message);
+			return await SendOrUpdateMessageToThreadAsync(recipient, eventId, userId, null, message);
 		}
 
-		private async Task<(MessageStateDocument, bool)> SendOrUpdateMessageAsync(string recipient, string? threadTs, string eventId, UserId? userId, SlackMessage message)
+		private async Task<(MessageStateDocument, bool)> SendOrUpdateMessageToThreadAsync(string recipient, string eventId, UserId? userId, SlackMessageId? threadId, SlackMessage message)
 		{
 			string requestDigest = ContentHash.MD5(JsonSerializer.Serialize(message, new JsonSerializerOptions { DefaultIgnoreCondition = JsonIgnoreCondition.WhenWritingNull })).ToString();
 
@@ -2021,34 +2020,34 @@ namespace Horde.Build.Notifications.Sinks
 				return (prevState, false);
 			}
 
-			(MessageStateDocument state, bool isNew) = await AddOrUpdateMessageStateAsync(recipient, eventId, userId, requestDigest);
+			(MessageStateDocument state, bool isNew) = await AddOrUpdateMessageStateAsync(recipient, eventId, userId, requestDigest, null);
 			if (isNew)
 			{
-				_logger.LogInformation("Sending new slack message to {SlackUser} (msg: {MessageId}, threadTs: {ThreadTs})", recipient, state.Id, threadTs ?? "n/a");
+				_logger.LogInformation("Sending new slack message to {SlackUser} (state: {StateId}, threadMessageId: {ThreadTs})", recipient, state.Id, threadId?.ToString() ?? "n/a");
 
-				state.Channel = recipient;
-				if (threadTs == null)
+				SlackMessageId id;
+				if (threadId == null)
 				{
-					state.Ts = await _slackClient.PostMessageAsync(recipient, message);
+					id = await _slackClient.PostMessageAsync(recipient, message);
 				}
 				else
 				{
-					state.Ts = await _slackClient.PostMessageAsync(recipient, threadTs, message);
+					id = await _slackClient.PostMessageToThreadAsync(threadId, message);
 				}
 
-				await SetMessageTimestampAsync(state.Id, state.Channel, state.Ts);
+				await UpdateMessageStateAsync(state.Id, id);
 			}
 			else if (!String.IsNullOrEmpty(state.Ts))
 			{
-				_logger.LogInformation("Updating existing slack message {MessageId} for user {SlackUser} (channel: {Channel}, ts: {Ts}, threadTs: {ThreadTs})", state.Id, recipient, state.Channel, state.Ts, threadTs ?? "n/a");
+				_logger.LogInformation("Updating existing slack message {StateId} for user {SlackUser} (messageId: {MessageId}, threadId: {ThreadTs})", state.Id, recipient, state.MessageId, threadId?.ToString() ?? "n/a");
 
-				if (threadTs == null)
+				if (threadId == null)
 				{
-					await _slackClient.UpdateMessageAsync(state.Channel, state.Ts, message);
+					await _slackClient.UpdateMessageAsync(state.MessageId, message);
 				}
 				else
 				{
-					await _slackClient.UpdateMessageAsync(state.Channel, state.Ts, threadTs, message);
+					await _slackClient.UpdateMessageAsync(new SlackMessageId(threadId.Channel, threadId.Ts, state.Ts), message);
 				}
 			}
 			return (state, isNew);
@@ -2148,7 +2147,7 @@ namespace Horde.Build.Notifications.Sinks
 				}
 
 				Uri issueUrl = GetIssueUrl(issue, span.FirstFailure);
-				await _slackClient.PostMessageAsync(workflow.TriageChannel, state.Ts, $"{FormatUserOrGroupMention(workflow.EscalateAlias)} - Issue <{issueUrl}|{issue.Id}> has not been resolved after {openTimeStr}.");
+				await _slackClient.PostMessageToThreadAsync(state.MessageId, $"{FormatUserOrGroupMention(workflow.EscalateAlias)} - Issue <{issueUrl}|{issue.Id}> has not been resolved after {openTimeStr}.");
 			}
 
 			DateTime nextEscalationTime = span.FirstFailure.StepTime;
diff --git a/Engine/Source/Programs/Shared/EpicGames.Slack/SlackClient.cs b/Engine/Source/Programs/Shared/EpicGames.Slack/SlackClient.cs
--- a/Engine/Source/Programs/Shared/EpicGames.Slack/SlackClient.cs
+++ b/Engine/Source/Programs/Shared/EpicGames.Slack/SlackClient.cs
@@ -32,6 +32,40 @@ namespace EpicGames.Slack
 		}
 	}
 
+	/// <summary>
+	/// Response from posting a Slack message
+	/// </summary>
+	public class SlackMessageId
+	{
+		/// <summary>
+		/// The channel id. Note that posting messages is more lenient than updating messages, and will accept a user ID as a recipient channel. The actual DM channel this resolves to will be returned in this parameter.
+		/// </summary>
+		public string Channel { get; set; }
+
+		/// <summary>
+		/// Thread containing the message
+		/// </summary>
+		public string? ThreadTs { get; set; }
+
+		/// <summary>
+		/// Timestamp of the message
+		/// </summary>
+		public string Ts { get; set; }
+
+		/// <summary>
+		/// Constructor
+		/// </summary>
+		public SlackMessageId(string channel, string? threadTs, string ts)
+		{
+			Channel = channel;
+			ThreadTs = threadTs;
+			Ts = ts;
+		}
+
+		/// <inheritdoc/>
+		public override string ToString() => (ThreadTs == null) ? $"{Channel}:{Ts}" : $"{Channel}:{ThreadTs}:{Ts}";
+	}
+
 	/// <summary>
 	/// Wrapper around Slack client functionality
 	/// </summary>
@@ -164,7 +198,8 @@ namespace EpicGames.Slack
 		/// </summary>
 		/// <param name="recipient">Recipient of the message. May be a channel or Slack user id.</param>
 		/// <param name="message">New message to post</param>
-		public async Task<string> PostMessageAsync(string recipient, SlackMessage message)
+		/// <returns>Identifier for the message. Note that the returned channel value will respond to a concrete channel identifier if a user is specified as recipient, and must be used to update the message.</returns>
+		public async Task<SlackMessageId> PostMessageAsync(string recipient, SlackMessage message)
 		{
 			return await PostOrUpdateMessageAsync(recipient, null, null, message, false);
 		}
@@ -172,40 +207,27 @@ namespace EpicGames.Slack
 		/// <summary>
 		/// Posts a message to a recipient
 		/// </summary>
-		/// <param name="recipient">Recipient of the message. May be a channel or Slack user id.</param>
-		/// <param name="threadTs">Timestamp of the thread to post the message to</param>
+		/// <param name="threadId">Parent message to post to</param>
 		/// <param name="message">New message to post</param>
 		/// <param name="replyBroadcast">Whether to broadcast the message to the channel</param>
-		public async Task<string> PostMessageAsync(string recipient, string threadTs, SlackMessage message, bool replyBroadcast = false)
-		{
-			return await PostOrUpdateMessageAsync(recipient, null, threadTs, message, replyBroadcast);
-		}
-
-		/// <summary>
-		/// Updates an existing message
-		/// </summary>
-		/// <param name="recipient">Recipient of the message. May be a channel or Slack user id.</param>
-		/// <param name="ts">The message timestamp</param>
-		/// <param name="message">New message to post</param>
-		public async Task UpdateMessageAsync(string recipient, string ts, SlackMessage message)
+		/// <returns>Identifier for the message. Note that the returned channel value will respond to a concrete channel identifier if a user is specified as recipient, and must be used to update the message.</returns>
+		public async Task<SlackMessageId> PostMessageToThreadAsync(SlackMessageId threadId, SlackMessage message, bool replyBroadcast = false)
 		{
-			await PostOrUpdateMessageAsync(recipient, ts, null, message, false);
+			return await PostOrUpdateMessageAsync(threadId.Channel, null, threadId.Ts, message, replyBroadcast);
 		}
 
 		/// <summary>
 		/// Updates an existing message
 		/// </summary>
-		/// <param name="recipient">Recipient of the message. May be a channel or Slack user id.</param>
-		/// <param name="ts">The message timestamp</param>
-		/// <param name="threadTs">Timestamp of the thread to post the message to</param>
+		/// <param name="id">The message id</param>
 		/// <param name="message">New message to post</param>
-		/// <param name="replyBroadcast">Whether to broadcast the message to the channel</param>
-		public async Task UpdateMessageAsync(string recipient, string ts, string threadTs, SlackMessage message, bool replyBroadcast = false)
+		/// <param name="replyBroadcast">For threaded messages, whether to broadcast the message to the channel</param>
+		public async Task UpdateMessageAsync(SlackMessageId id, SlackMessage message, bool replyBroadcast = false)
 		{
-			await PostOrUpdateMessageAsync(recipient, ts, threadTs, message, replyBroadcast);
+			await PostOrUpdateMessageAsync(id.Channel, id.Ts, id.ThreadTs, message, replyBroadcast);
 		}
 
-		async Task<string> PostOrUpdateMessageAsync(string recipient, string? ts, string? threadTs, SlackMessage message, bool replyBroadcast)
+		async Task<SlackMessageId> PostOrUpdateMessageAsync(string recipient, string? ts, string? threadTs, SlackMessage message, bool replyBroadcast)
 		{
 			PostMessageRequest request = new PostMessageRequest();
 			request.Channel = recipient;
@@ -238,7 +260,7 @@ namespace EpicGames.Slack
 				throw new SlackException(response.Error ?? "unknown");
 			}
 
-			return response.Ts;
+			return new SlackMessageId(response.Channel ?? recipient, threadTs, response.Ts);
 		}
 
 		class GetPermalinkRequest
@@ -262,12 +284,11 @@ namespace EpicGames.Slack
 		/// <summary>
 		/// Gets a permalink for a message
 		/// </summary>
-		/// <param name="channel">Channel containing the message</param>
-		/// <param name="ts">Message timestamp</param>
+		/// <param name="id">Message identifier</param>
 		/// <returns>Link to the message</returns>
-		public async Task<string> GetPermalinkAsync(string channel, string ts)
+		public async Task<string> GetPermalinkAsync(SlackMessageId id)
 		{
-			string requestUrl = $"{GetPermalinkUrl}?channel={channel}&message_ts={ts}";
+			string requestUrl = $"{GetPermalinkUrl}?channel={id.Channel}&message_ts={id.Ts}";
 			using (HttpRequestMessage sendMessageRequest = new HttpRequestMessage(HttpMethod.Get, requestUrl))
 			{
 				GetPermalinkResponse response = await SendRequestAsync<GetPermalinkResponse>(sendMessageRequest, "");
@@ -301,14 +322,13 @@ namespace EpicGames.Slack
 		/// <summary>
 		/// Adds a reaction to a posted message
 		/// </summary>
-		/// <param name="channel">Channel containing the message</param>
-		/// <param name="ts">Message timestamp</param>
+		/// <param name="messageId">Message to react to</param>
 		/// <param name="name">Name of the reaction to post</param>
-		public async Task AddReactionAsync(string channel, string ts, string name)
+		public async Task AddReactionAsync(SlackMessageId messageId, string name)
 		{
 			ReactionMessage message = new ReactionMessage();
-			message.Channel = channel;
-			message.Ts = ts;
+			message.Channel = messageId.Channel;
+			message.Ts = messageId.Ts;
 			message.Name = name;
 
 			static bool ShouldLogError(SlackResponse response) => !response.Ok && !String.Equals(response.Error, "already_reacted", StringComparison.Ordinal);
@@ -319,14 +339,13 @@ namespace EpicGames.Slack
 		/// <summary>
 		/// Removes a reaction from a posted message
 		/// </summary>
-		/// <param name="channel">Channel containing the message</param>
-		/// <param name="ts">Message timestamp</param>
+		/// <param name="messageId">Message to react to</param>
 		/// <param name="name">Name of the reaction to post</param>
-		public async Task RemoveReactionAsync(string channel, string ts, string name)
+		public async Task RemoveReactionAsync(SlackMessageId messageId, string name)
 		{
 			ReactionMessage message = new ReactionMessage();
-			message.Channel = channel;
-			message.Ts = ts;
+			message.Channel = messageId.Channel;
+			message.Ts = messageId.Ts;
 			message.Name = name;
 
 			static bool ShouldLogError(SlackResponse response) => !response.Ok && !String.Equals(response.Error, "no_reaction", StringComparison.Ordinal);
