Horde: Disable code to delete shelved changes after running preflights. Possible bug with submits not reporting errors correctly.

diff --git a/Engine/Source/Programs/Horde/Horde.Build/Jobs/JobService.cs b/Engine/Source/Programs/Horde/Horde.Build/Jobs/JobService.cs
--- a/Engine/Source/Programs/Horde/Horde.Build/Jobs/JobService.cs
+++ b/Engine/Source/Programs/Horde/Horde.Build/Jobs/JobService.cs
@@ -1180,8 +1180,12 @@ namespace Horde.Build.Jobs
 		/// <param name="clusterName"></param>
 		/// <param name="change">The changelist to delete</param>
 		/// <returns>True if the change was deleted successfully, false otherwise</returns>
-		private async Task<bool> DeleteShelvedChangeAsync(string clusterName, int change)
+		private Task<bool> DeleteShelvedChangeAsync(string clusterName, int change)
 		{
+			_ = clusterName;
+			_logger.LogInformation("Leaving shelved change {Change}", change);
+			return Task.FromResult<bool>(true);
+/*
 			_logger.LogInformation("Removing shelf {Change}", change);
 			try
 			{
@@ -1192,7 +1196,7 @@ namespace Horde.Build.Jobs
 			{
 				_logger.LogError(ex, "Unable to delete shelved change {Change}", change);
 				return false;
-			}
+			}*/
 		}
 
 		/// <summary>
