Horde: Handle -1 being passed as the changelist number for new jobs.

diff --git a/Engine/Source/Programs/Horde/Horde.Build/Jobs/JobsController.cs b/Engine/Source/Programs/Horde/Horde.Build/Jobs/JobsController.cs
--- a/Engine/Source/Programs/Horde/Horde.Build/Jobs/JobsController.cs
+++ b/Engine/Source/Programs/Horde/Horde.Build/Jobs/JobsController.cs
@@ -139,7 +139,7 @@ namespace Horde.Build.Jobs
 
 			// Get the change to build
 			int change;
-			if (create.Change.HasValue)
+			if (create.Change.HasValue && create.Change.Value != -1)
 			{
 				change = create.Change.Value;
 			}
