Horde: Use the #ROBOMERGE-AUTHOR tag to correctly attribute merge commits before Robomerge has a change to replace the owner.

diff --git a/Engine/Source/Programs/Horde/Horde.Build/Perforce/PerforceService.cs b/Engine/Source/Programs/Horde/Horde.Build/Perforce/PerforceService.cs
--- a/Engine/Source/Programs/Horde/Horde.Build/Perforce/PerforceService.cs
+++ b/Engine/Source/Programs/Horde/Horde.Build/Perforce/PerforceService.cs
@@ -633,6 +633,18 @@ namespace Horde.Build.Perforce
 
 		async ValueTask<Commit> CreateCommitInternalAsync(IStream stream, int number, string author, string description, string? basePath, DateTime dateUtc, CancellationToken cancellationToken)
 		{
+			// Robomerge submits changes as its own user, then modifies the commit with the real author afterwards. Parse it out of the changelist description to handle these cases.
+			if (String.Equals(author, "robomerge", StringComparison.OrdinalIgnoreCase))
+			{
+				Match match = Regex.Match(description, @"^#ROBOMERGE-AUTHOR:\s*([^\s]+)", RegexOptions.Multiline);
+				if (match.Success)
+				{
+					string newAuthor = match.Groups[1].Value;
+					_logger.LogDebug("Changing author of CL {Change} from {Author} to {NewAuthor}", number, author, newAuthor);
+					author = newAuthor;
+				}
+			}
+
 			IUser authorUser = await FindOrAddUserAsync(stream.Config.ClusterName, author, cancellationToken);
 			IUser ownerUser = authorUser;
 
