Horde: Remove some device service error messages that don't need to be surfaced in server logs
#rnx

diff --git a/Engine/Source/Programs/Horde/Horde.Build/Devices/DevicesController.cs b/Engine/Source/Programs/Horde/Horde.Build/Devices/DevicesController.cs
--- a/Engine/Source/Programs/Horde/Horde.Build/Devices/DevicesController.cs
+++ b/Engine/Source/Programs/Horde/Horde.Build/Devices/DevicesController.cs
@@ -833,7 +833,6 @@ namespace Horde.Build.Devices
 
 			if (String.IsNullOrEmpty(poolId))
 			{
-				_logger.LogError("No pool specified, defaulting to UE4 {Details} JobId: {JobId}, StepId: {StepId}", details, request.JobId, request.StepId);
 				string message = $"No pool specified, defaulting to UE4" + details;
 				await _deviceService.NotifyDeviceServiceAsync(message, null, request.JobId, request.StepId);
 				poolId = "ue4";
@@ -886,7 +885,7 @@ namespace Horde.Build.Devices
 				if (platform == null)
 				{
 					string message = $"Unknown platform {platformId}" + details;
-					_logger.LogError("Unknown platform {PlatformId} {Details}", platformId, details);
+					_logger.LogError("Unknown platform {PlatformName} {PlatformId} {Details}", platformName, platformId, details);
 					await _deviceService.NotifyDeviceServiceAsync(message, null, request.JobId, request.StepId);
 					return BadRequest(message);
 				}
@@ -1111,7 +1110,6 @@ namespace Horde.Build.Devices
 
 			if (device == null)
 			{
-				_logger.LogError("Device error reported for unknown device {DeviceName}", deviceName);
 				return BadRequest($"Unknown device {deviceName}");
 			}
 
@@ -1139,9 +1137,7 @@ namespace Horde.Build.Devices
 						message += $" - Host {reservation.Hostname}";
 					}
 				}
-			}
-
-			_logger.LogError("{DeviceError}", message);
+			}			
 
 			await _deviceService.NotifyDeviceServiceAsync(message, device.Id, jobId, stepId);
 
