Horde: Use the HeadChange instead of ChangeNumber when getting revisions for config updates
#rnx
#skipci

diff --git a/Engine/Source/Programs/Horde/Horde.Build/Configuration/ConfigUpdateService.cs b/Engine/Source/Programs/Horde/Horde.Build/Configuration/ConfigUpdateService.cs
--- a/Engine/Source/Programs/Horde/Horde.Build/Configuration/ConfigUpdateService.cs
+++ b/Engine/Source/Programs/Horde/Horde.Build/Configuration/ConfigUpdateService.cs
@@ -379,7 +379,7 @@ namespace Horde.Build.Configuration
 						}
 
 						Uri fileUri = new Uri($"{PerforceScheme}://{perforcePath.Key}{file.DepotFile}");
-						revisions[fileUri] = $"ver={Version},chg={file.ChangeNumber},path={fileUri}";
+						revisions[fileUri] = $"ver={Version},chg={file.HeadChange},path={fileUri}";
 					}
 				}
 			}
