Horde: Additional debug logging for marketplace build error.

diff --git a/Engine/Source/Programs/Horde/Horde.Build/Perforce/PerforceService.cs b/Engine/Source/Programs/Horde/Horde.Build/Perforce/PerforceService.cs
--- a/Engine/Source/Programs/Horde/Horde.Build/Perforce/PerforceService.cs
+++ b/Engine/Source/Programs/Horde/Horde.Build/Perforce/PerforceService.cs
@@ -984,6 +984,7 @@ namespace Horde.Build.Perforce
 						change = await perforce.CreateChangeAsync(change, cancellationToken);
 
 						SubmitRecord submit = await perforce.SubmitAsync(change.Number, SubmitOptions.SubmitUnchanged, cancellationToken);
+						_logger.LogInformation("Submitted new changelist with {DepotPath}: CL {Change}", depotPath, submit.SubmittedChangeNumber);
 						return submit.SubmittedChangeNumber;
 					}
 					catch (Exception ex)
@@ -1153,8 +1154,8 @@ namespace Horde.Build.Perforce
 			{
 				using IScope scope = GlobalTracer.Instance.BuildSpan("PerforceService.CommitSource.FindAsync").StartActive();
 				scope.Span.SetTag("ClusterName", _stream.Config.ClusterName);
-				scope.Span.SetTag("MinChange", minChange ?? -1);
-				scope.Span.SetTag("MaxChange", maxChange ?? -1);
+				scope.Span.SetTag("MinChange", minChange ?? -2);
+				scope.Span.SetTag("MaxChange", maxChange ?? -2);
 				scope.Span.SetTag("MaxResults", maxResults ?? -1);
 
 				using (PooledConnectionHandle perforce = await _perforceService.ConnectWithStreamClientAsync(_stream, null, cancellationToken))
diff --git a/Engine/Source/Programs/Horde/Horde.Build/Perforce/PerforceServiceCache.cs b/Engine/Source/Programs/Horde/Horde.Build/Perforce/PerforceServiceCache.cs
--- a/Engine/Source/Programs/Horde/Horde.Build/Perforce/PerforceServiceCache.cs
+++ b/Engine/Source/Programs/Horde/Horde.Build/Perforce/PerforceServiceCache.cs
@@ -453,7 +453,7 @@ namespace Horde.Build.Perforce
 				}
 
 				int numResults = 0;
-				_owner._logger.LogDebug("Querying Perforce cache for {StreamId} commits from {MinChange} to {MaxChange} (max: {MaxResults}, tags: {Tags})", _stream.Id, minChange ?? -1, maxChange ?? -1, maxResults ?? -1, (tags == null || tags.Count == 0) ? "none" : String.Join("/", tags.Select(x => x.ToString())));
+				_owner._logger.LogDebug("Querying Perforce cache for {StreamId} commits from {MinChange} to {MaxChange} (max: {MaxResults}, tags: {Tags})", _stream.Id, minChange ?? -2, maxChange ?? -2, maxResults ?? -1, (tags == null || tags.Count == 0) ? "none" : String.Join("/", tags.Select(x => x.ToString())));
 
 				CacheState state = await _owner._mongoService.GetSingletonAsync<CacheState>();
 				if (state.Clusters.TryGetValue(_stream.Config.ClusterName, out ClusterState? clusterState))
@@ -524,7 +524,7 @@ namespace Horde.Build.Perforce
 
 				if (maxResults == null || maxResults.Value > 0)
 				{
-					_owner._logger.LogDebug("Querying Perforce server for {StreamId} commits from {MinChange} to {MaxChange} (max: {MaxResults}, tags: {Tags})", _stream.Id, minChange ?? -1, maxChange ?? -1, maxResults ?? -1, (tags == null || tags.Count == 0) ? "none" : String.Join("/", tags.Select(x => x.ToString())));
+					_owner._logger.LogDebug("Querying Perforce server for {StreamId} commits from {MinChange} to {MaxChange} (max: {MaxResults}, tags: {Tags})", _stream.Id, minChange ?? -2, maxChange ?? -2, maxResults ?? -1, (tags == null || tags.Count == 0) ? "none" : String.Join("/", tags.Select(x => x.ToString())));
 
 					await foreach (ICommit commit in base.FindAsync(minChange, maxChange, maxResults, tags, cancellationToken))
 					{
