Horde: Fix first issue escalation using step start time.

diff --git a/Engine/Source/Programs/Horde/Horde.Build/Notifications/Sinks/SlackNotificationSink.cs b/Engine/Source/Programs/Horde/Horde.Build/Notifications/Sinks/SlackNotificationSink.cs
--- a/Engine/Source/Programs/Horde/Horde.Build/Notifications/Sinks/SlackNotificationSink.cs
+++ b/Engine/Source/Programs/Horde/Horde.Build/Notifications/Sinks/SlackNotificationSink.cs
@@ -971,7 +971,7 @@ namespace Horde.Build.Notifications.Sinks
 
 				if (workflow.EscalateAlias != null && workflow.EscalateTimes.Count > 0)
 				{
-					DateTime escalateTime = span.FirstFailure.StepTime + TimeSpan.FromMinutes(workflow.EscalateTimes[0]);
+					DateTime escalateTime = issue.CreatedAt + TimeSpan.FromMinutes(workflow.EscalateTimes[0]);
 					if (await _escalateIssues.AddAsync(issue.Id, (escalateTime - DateTime.UnixEpoch).TotalSeconds, StackExchange.Redis.When.NotExists))
 					{
 						_logger.LogInformation("First escalation time for issue {IssueId} is {Time}", issue.Id, escalateTime);
