Horde Storage - Reset version numbers as the old versioning only applies to the legacy api which has now been deleted. The intention is to start to maintain the changelog again as of this point.

diff --git a/Engine/Source/Programs/Horde/HordeStorage/changelog.md b/Engine/Source/Programs/Horde/HordeStorage/changelog.md
--- a/Engine/Source/Programs/Horde/HordeStorage/changelog.md
+++ b/Engine/Source/Programs/Horde/HordeStorage/changelog.md
@@ -1,4 +1,7 @@
 # Unreleased
 
+# 0.1.0
+* First public release of Horde Storage
+
 # Older releases
 No detailed changelog provided for older releases.
\ No newline at end of file
