Horde Storage - Fixed issue were replication was not running as the setting could not be enabled

diff --git a/Engine/Source/Programs/Horde/HordeStorage/Horde.Storage/appsettings.Production.json b/Engine/Source/Programs/Horde/HordeStorage/Horde.Storage/appsettings.Production.json
--- a/Engine/Source/Programs/Horde/HordeStorage/Horde.Storage/appsettings.Production.json
+++ b/Engine/Source/Programs/Horde/HordeStorage/Horde.Storage/appsettings.Production.json
@@ -29,11 +29,6 @@
         "RequireTls12": true
     },
 
-    "Replication": {
-        // you will need to explicitly pass replication settings if you want to enable it
-        "Enabled": false
-    },
-
     "Jupiter": {
         "PublicApiPorts": [ 80, 8081 ],
         "CorpApiPorts": [ 8008, 8082 ],
