Horde: Fix change on stream update notifications showing  -1
#rnx

diff --git a/Engine/Source/Programs/Horde/Horde.Build/Configuration/ConfigUpdateService.cs b/Engine/Source/Programs/Horde/Horde.Build/Configuration/ConfigUpdateService.cs
--- a/Engine/Source/Programs/Horde/Horde.Build/Configuration/ConfigUpdateService.cs
+++ b/Engine/Source/Programs/Horde/Horde.Build/Configuration/ConfigUpdateService.cs
@@ -433,6 +433,7 @@ namespace Horde.Build.Configuration
 							DescribeRecord describeRecord = await perforce.DescribeAsync(records[0].HeadChange, cancellationToken);
 							author = await _perforceService.FindOrAddUserAsync(configPath.Host, describeRecord.User, cancellationToken);
 							description = describeRecord.Description;
+							change = records[0].HeadChange;
 						}
 					}
 					catch (Exception ex2)
