Horde: Add template ID validation for tabs in stream config

diff --git a/Engine/Source/Programs/Horde/Horde.Build.Tests/SchedulerTests.cs b/Engine/Source/Programs/Horde/Horde.Build.Tests/SchedulerTests.cs
--- a/Engine/Source/Programs/Horde/Horde.Build.Tests/SchedulerTests.cs
+++ b/Engine/Source/Programs/Horde/Horde.Build.Tests/SchedulerTests.cs
@@ -430,6 +430,7 @@ namespace Horde.Build.Tests
 			StreamConfig config = new StreamConfig();
 			config.Name = "//UE5/Main";
 			config.Tabs.Add(new JobsTabConfig { Title = "foo", Templates = new List<TemplateId> { newTemplateRefId1, newTemplateRefId2 } });
+			config.Templates = new() { new TemplateRefConfig { Id = newTemplateRefId1 }, new TemplateRefConfig { Id = newTemplateRefId2 }  };
 
 			stream = (await CreateOrReplaceStreamAsync(StreamId, stream, ProjectId, config))!;
 
diff --git a/Engine/Source/Programs/Horde/Horde.Build.Tests/Streams/StreamCollectionTests.cs b/Engine/Source/Programs/Horde/Horde.Build.Tests/Streams/StreamCollectionTests.cs
new file mode 100644
--- /dev/null
+++ b/Engine/Source/Programs/Horde/Horde.Build.Tests/Streams/StreamCollectionTests.cs
@@ -0,0 +1,37 @@
+// Copyright Epic Games, Inc. All Rights Reserved.
+
+using System.Collections.Generic;
+using Horde.Build.Streams;
+using Horde.Build.Utilities;
+using Microsoft.VisualStudio.TestTools.UnitTesting;
+
+namespace Horde.Build.Tests.Streams
+{
+	using StreamId = StringId<IStream>;
+	using TemplateId = StringId<ITemplateRef>;
+
+	[TestClass]
+	public class StreamCollectionTests : TestSetup
+	{
+		private readonly StreamId _streamId = new ("bogusStreamId");
+		
+		[TestMethod]
+		public void ValidateUndefinedTemplateIdInTabs()
+		{
+			StreamConfig config = new()
+			{
+				Tabs = new()
+				{
+					new JobsTabConfig { Templates = new List<TemplateId> { new ("foo") }},
+					new JobsTabConfig { Templates = new List<TemplateId> { new ("bar") }}
+				},
+				Templates = new () { new TemplateRefConfig { Id = new TemplateId("foo") } }
+			};
+
+			Assert.ThrowsException<InvalidStreamException>(() => Horde.Build.Streams.StreamCollection.Validate(_streamId, config));
+
+			config.Templates.Add(new TemplateRefConfig { Id = new TemplateId("bar") });
+			Horde.Build.Streams.StreamCollection.Validate(_streamId, config);
+		}
+	}
+}
diff --git a/Engine/Source/Programs/Horde/Horde.Build/Streams/StreamCollection.cs b/Engine/Source/Programs/Horde/Horde.Build/Streams/StreamCollection.cs
--- a/Engine/Source/Programs/Horde/Horde.Build/Streams/StreamCollection.cs
+++ b/Engine/Source/Programs/Horde/Horde.Build/Streams/StreamCollection.cs
@@ -604,12 +604,21 @@ namespace Horde.Build.Streams
 				}
 			}
 
+			HashSet<TemplateId> undefinedTemplates = new();
 			// Check that all the templates are referenced by a tab
 			foreach (JobsTabConfig jobsTab in config.Tabs.OfType<JobsTabConfig>())
 			{
 				if (jobsTab.Templates != null)
 				{
 					remainingTemplates.ExceptWith(jobsTab.Templates);
+					foreach (TemplateId templateId in jobsTab.Templates)
+					{
+						if (config.Templates.Find(x => x.Id == templateId) == null)
+						{
+							undefinedTemplates.Add(templateId);
+						}
+					}
+					
 				}
 			}
 			if (remainingTemplates.Count > 0)
@@ -617,6 +626,11 @@ namespace Horde.Build.Streams
 				throw new InvalidStreamException(String.Join("\n", remainingTemplates.Select(x => $"Template '{x}' is not listed on any tab for {streamId}")));
 			}
 
+			if (undefinedTemplates.Count > 0)
+			{
+				throw new InvalidStreamException(String.Join("\n", undefinedTemplates.Select(x => $"Template '{x}' is not defined for {streamId}")));
+			}
+
 			// Check that all the agent types reference valid workspace names
 			foreach (KeyValuePair<string, AgentConfig> pair in config.AgentTypes)
 			{
