Horde: Server side of force closing issues
#rnx

diff --git a/Engine/Source/Programs/Horde/Horde.Build.Tests/IssueServiceTests.cs b/Engine/Source/Programs/Horde/Horde.Build.Tests/IssueServiceTests.cs
--- a/Engine/Source/Programs/Horde/Horde.Build.Tests/IssueServiceTests.cs
+++ b/Engine/Source/Programs/Horde/Horde.Build.Tests/IssueServiceTests.cs
@@ -1988,6 +1988,75 @@ namespace Horde.Build.Tests
 
 		}
 
+		[TestMethod]
+		public async Task ForceIssueCloseTest()
+		{
+			int issueId;
+
+			// #1
+			// Scenario: Warning in first step
+			// Expected: Default issue is created
+			{
+				IJob job = CreateJob(_mainStreamId, 105, "Test Build", _graph);
+				await AddEvent(job, 0, 0, new { level = nameof(LogLevel.Warning) }, EventSeverity.Warning);
+				await UpdateCompleteStep(job, 0, 0, JobStepOutcome.Warnings);
+
+				List<IIssue> issues = await IssueCollection.FindIssuesAsync();
+				Assert.AreEqual(1, issues.Count);
+				Assert.AreEqual(IssueSeverity.Warning, issues[0].Severity);
+
+				Assert.AreEqual("Warnings in Update Version Files", issues[0].Summary);
+
+				issueId = issues[0].Id;				
+			}
+
+			// #2
+			// Scenario: bob resolved issue
+			// Expected: Issue is marked resolved, however not verified
+			{
+				// resolved by bob
+				await IssueService.UpdateIssueAsync(issueId, resolvedById: _bobId);
+				List<IIssue> issues = await IssueCollection.FindIssuesAsync();
+				Assert.AreEqual(0, issues.Count);
+				IIssue issue = (await IssueCollection.GetIssueAsync(issueId))!;
+				Assert.AreEqual(issue.ResolvedById, _bobId);
+				Assert.IsNull(issue.VerifiedAt);
+			}
+
+			// #3
+			// Scenario: Job is force closed
+			// Expected: Existing issue is closed and verified, bob remains the resolver
+			{
+				// force closed by jerry
+				await IssueService.UpdateIssueAsync(issueId, forceClosedById: _jerryId);
+				List<IIssue> issues = await IssueCollection.FindIssuesAsync();
+				Assert.AreEqual(0, issues.Count);
+
+				IIssue issue = (await IssueCollection.GetIssueAsync(issueId))!;
+				Assert.AreEqual(issue.ResolvedById, _bobId);
+				Assert.AreEqual(issue.ForceClosedByUserId, _jerryId);
+				Assert.IsNotNull(issue.VerifiedAt);
+			}
+
+			// #4
+			// Scenario: Job fails 25 hours after it is force closed
+			// Expected: A new issue is created
+			{				
+				IJob job = CreateJob(_mainStreamId, 125, "Test Build", _graph, TimeSpan.FromHours(25));
+				await AddEvent(job, 0, 0, new { level = nameof(LogLevel.Warning) }, EventSeverity.Warning);
+				await UpdateCompleteStep(job, 0, 0, JobStepOutcome.Warnings);
+
+				List<IIssue> issues = await IssueCollection.FindIssuesAsync();
+				Assert.AreEqual(1, issues.Count);
+				Assert.AreEqual(IssueSeverity.Warning, issues[0].Severity);
+
+				Assert.AreEqual("Warnings in Update Version Files", issues[0].Summary);
+
+				Assert.AreEqual(2, issues[0].Id);
+			}
+
+		}
+
 
 		private async Task ParseAsync(LogId logId, string[] lines)
 		{
diff --git a/Engine/Source/Programs/Horde/Horde.Build/Issues/IIssue.cs b/Engine/Source/Programs/Horde/Horde.Build/Issues/IIssue.cs
--- a/Engine/Source/Programs/Horde/Horde.Build/Issues/IIssue.cs
+++ b/Engine/Source/Programs/Horde/Horde.Build/Issues/IIssue.cs
@@ -168,6 +168,12 @@ namespace Horde.Build.Issues
 		/// The UTC time when the issue was quarantined
 		/// </summary>
 		public DateTime? QuarantineTimeUtc { get; set; }
+
+		/// <summary>
+		/// User who force closed the issue
+		/// </summary>
+		public UserId? ForceClosedByUserId { get; set; }
+
 	}
 
 	/// <summary>
diff --git a/Engine/Source/Programs/Horde/Horde.Build/Issues/IIssueCollection.cs b/Engine/Source/Programs/Horde/Horde.Build/Issues/IIssueCollection.cs
--- a/Engine/Source/Programs/Horde/Horde.Build/Issues/IIssueCollection.cs
+++ b/Engine/Source/Programs/Horde/Horde.Build/Issues/IIssueCollection.cs
@@ -465,8 +465,9 @@ namespace Horde.Build.Issues
 		/// <param name="newLastSeenAt"></param>
 		/// <param name="newExternalIssueKey">Issue key for external issue tracking</param>
 		/// <param name="newQuarantinedById">The user that quarantined the issue</param>
+		/// <param name="newForceClosedById">The user that force closed the issue</param>
 		/// <returns>True if the issue was updated</returns>
-		Task<IIssue?> TryUpdateIssueAsync(IIssue issue, IssueSeverity? newSeverity = null, string? newSummary = null, string? newUserSummary = null, string? newDescription = null, bool? newPromoted = null, UserId? newOwnerId = null, UserId? newNominatedById = null, bool? newAcknowledged = null, UserId? newDeclinedById = null, int? newFixChange = null, UserId? newResolvedById = null, List<ObjectId>? newExcludeSpanIds = null, DateTime? newLastSeenAt = null, string? newExternalIssueKey = null, UserId? newQuarantinedById = null);
+		Task<IIssue?> TryUpdateIssueAsync(IIssue issue, IssueSeverity? newSeverity = null, string? newSummary = null, string? newUserSummary = null, string? newDescription = null, bool? newPromoted = null, UserId? newOwnerId = null, UserId? newNominatedById = null, bool? newAcknowledged = null, UserId? newDeclinedById = null, int? newFixChange = null, UserId? newResolvedById = null, List<ObjectId>? newExcludeSpanIds = null, DateTime? newLastSeenAt = null, string? newExternalIssueKey = null, UserId? newQuarantinedById = null, UserId? newForceClosedById = null);
 
 		/// <summary>
 		/// Updates derived data for an issue (ie. data computed from the spans attached to it). Also clears the issue's 'modified' state.
diff --git a/Engine/Source/Programs/Horde/Horde.Build/Issues/IssueCollection.cs b/Engine/Source/Programs/Horde/Horde.Build/Issues/IssueCollection.cs
--- a/Engine/Source/Programs/Horde/Horde.Build/Issues/IssueCollection.cs
+++ b/Engine/Source/Programs/Horde/Horde.Build/Issues/IssueCollection.cs
@@ -122,6 +122,9 @@ namespace Horde.Build.Issues
 			[BsonIgnoreIfNull]
 			public DateTime? QuarantineTimeUtc { get; set; }
 
+			[BsonIgnoreIfNull]
+			public UserId? ForceClosedByUserId { get; set; }
+
 			[BsonConstructor]
 			private Issue()
 			{
@@ -628,6 +631,18 @@ namespace Horde.Build.Issues
 				}
 			}
 
+			if (newIssue.ForceClosedByUserId != oldIssue.ForceClosedByUserId)
+			{
+				if (newIssue.ForceClosedByUserId != null)
+				{
+					issueLogger.LogInformation("Forced closed by {UserId}", newIssue.ForceClosedByUserId);
+				}
+				else
+				{
+					issueLogger.LogInformation("Force closed cleared");
+				}
+			}
+
 			string oldFingerprints = oldIssue.FingerprintsDesc;
 			string newFingerprints = newIssue.FingerprintsDesc;
 			if (oldFingerprints != newFingerprints)
@@ -845,7 +860,7 @@ namespace Horde.Build.Issues
 		}
 
 		/// <inheritdoc/>
-		public async Task<IIssue?> TryUpdateIssueAsync(IIssue issue, IssueSeverity? newSeverity = null, string? newSummary = null, string? newUserSummary = null, string? newDescription = null, bool? newManuallyPromoted = null, UserId? newOwnerId = null, UserId? newNominatedById = null, bool? newAcknowledged = null, UserId? newDeclinedById = null, int? newFixChange = null, UserId? newResolvedById = null, List<ObjectId>? newExcludeSpanIds = null, DateTime? newLastSeenAt = null, string? newExternaIssueKey = null, UserId? newQuarantinedById = null)
+		public async Task<IIssue?> TryUpdateIssueAsync(IIssue issue, IssueSeverity? newSeverity = null, string? newSummary = null, string? newUserSummary = null, string? newDescription = null, bool? newManuallyPromoted = null, UserId? newOwnerId = null, UserId? newNominatedById = null, bool? newAcknowledged = null, UserId? newDeclinedById = null, int? newFixChange = null, UserId? newResolvedById = null, List<ObjectId>? newExcludeSpanIds = null, DateTime? newLastSeenAt = null, string? newExternaIssueKey = null, UserId? newQuarantinedById = null, UserId? newForceClosedById = null)
 		{
 			Issue issueDocument = (Issue)issue;
 
@@ -854,6 +869,11 @@ namespace Horde.Build.Issues
 				newOwnerId = UserId.Empty;
 			}
 
+			if (issue.ResolvedById == null && (newForceClosedById != null && newResolvedById == null))
+			{
+				newResolvedById = newForceClosedById;
+			}
+
 			DateTime utcNow = DateTime.UtcNow;
 
 			List<UpdateDefinition<Issue>> updates = new List<UpdateDefinition<Issue>>();
@@ -1003,6 +1023,19 @@ namespace Horde.Build.Issues
 				}
 			}
 
+			if (newForceClosedById != null)
+			{
+				if (newForceClosedById.Value == UserId.Empty)
+				{
+					updates.Add(Builders<Issue>.Update.Unset(x => x.ForceClosedByUserId));
+				}
+				else
+				{
+					updates.Add(Builders<Issue>.Update.Set(x => x.ForceClosedByUserId, newForceClosedById.Value));					
+				}
+			}
+
+
 			if (newExternaIssueKey != null)
 			{
 				updates.Add(Builders<Issue>.Update.Set(x => x.ExternalIssueKey, newExternaIssueKey.Length == 0 ? null : newExternaIssueKey));
diff --git a/Engine/Source/Programs/Horde/Horde.Build/Issues/IssueMessages.cs b/Engine/Source/Programs/Horde/Horde.Build/Issues/IssueMessages.cs
--- a/Engine/Source/Programs/Horde/Horde.Build/Issues/IssueMessages.cs
+++ b/Engine/Source/Programs/Horde/Horde.Build/Issues/IssueMessages.cs
@@ -882,6 +882,12 @@ namespace Horde.Build.Issues
 		/// Id of user quarantining issue
 		/// </summary>
 		public string? QuarantinedById { get; set; }
+
+		/// <summary>
+		/// Id of user who is forcibly closing this issue, skipping verification checks.  This is useful for when a failing step has been removed for example
+		/// </summary>
+		public string? ForceClosedById { get; set; }
+
 	}
 
 	/// <summary>
diff --git a/Engine/Source/Programs/Horde/Horde.Build/Issues/IssueService.cs b/Engine/Source/Programs/Horde/Horde.Build/Issues/IssueService.cs
--- a/Engine/Source/Programs/Horde/Horde.Build/Issues/IssueService.cs
+++ b/Engine/Source/Programs/Horde/Horde.Build/Issues/IssueService.cs
@@ -397,8 +397,9 @@ namespace Horde.Build.Issues
 		/// <param name="removeSpanIds">Remove spans from this issue</param>
 		/// <param name="externalIssueKey">Key for external issue tracking</param>
 		/// <param name="quarantinedById">User who has quarantined the issue</param>
+		/// <param name="forceClosedById">User who has force closed the issue</param>
 		/// <returns>True if the issue was updated</returns>
-		public async Task<bool> UpdateIssueAsync(int id, string? summary = null, string? description = null, bool? promoted = null, UserId? ownerId = null, UserId? nominatedById = null, bool? acknowledged = null, UserId? declinedById = null, int? fixChange = null, UserId? resolvedById = null, List<ObjectId>? addSpanIds = null, List<ObjectId>? removeSpanIds = null, string? externalIssueKey = null, UserId? quarantinedById = null)
+		public async Task<bool> UpdateIssueAsync(int id, string? summary = null, string? description = null, bool? promoted = null, UserId? ownerId = null, UserId? nominatedById = null, bool? acknowledged = null, UserId? declinedById = null, int? fixChange = null, UserId? resolvedById = null, List<ObjectId>? addSpanIds = null, List<ObjectId>? removeSpanIds = null, string? externalIssueKey = null, UserId? quarantinedById = null, UserId? forceClosedById = null)
 		{
 			IIssue? issue;
 			for (; ; )
@@ -409,7 +410,7 @@ namespace Horde.Build.Issues
 					return false;
 				}
 
-				issue = await _issueCollection.TryUpdateIssueAsync(issue, newUserSummary: summary, newDescription: description, newPromoted: promoted, newOwnerId: ownerId ?? resolvedById, newNominatedById: nominatedById, newDeclinedById: declinedById, newAcknowledged: acknowledged, newFixChange: fixChange, newResolvedById: resolvedById, newExcludeSpanIds: removeSpanIds, newExternalIssueKey: externalIssueKey, newQuarantinedById: quarantinedById);
+				issue = await _issueCollection.TryUpdateIssueAsync(issue, newUserSummary: summary, newDescription: description, newPromoted: promoted, newOwnerId: ownerId ?? resolvedById, newNominatedById: nominatedById, newDeclinedById: declinedById, newAcknowledged: acknowledged, newFixChange: fixChange, newResolvedById: resolvedById, newExcludeSpanIds: removeSpanIds, newExternalIssueKey: externalIssueKey, newQuarantinedById: quarantinedById, newForceClosedById: forceClosedById);
 				if (issue != null)
 				{
 					break;
@@ -421,6 +422,28 @@ namespace Horde.Build.Issues
 			List<IIssue> updateIssues = new List<IIssue>();
 			updateIssues.Add(issue);
 
+			if (forceClosedById != null)
+			{
+				List<IIssueSpan> spans = await _issueCollection.FindSpansAsync(issue.Id);
+
+				foreach (IIssueSpan span in spans)
+				{
+					if (span.NextSuccess != null)
+					{
+						continue;
+					}
+
+					IIssueStep? step = span.LastFailure ?? span.FirstFailure;
+
+					if (step == null)
+					{
+						continue;
+					}
+
+					await _issueCollection.TryUpdateSpanAsync(span, newNextSuccess: new NewIssueStepData(step.Change, step.Severity, step.JobName, step.JobId, step.BatchId, step.StepId, step.StepTime, step.LogId, step.Annotations, step.PromoteByDefault));
+				}				
+			}
+
 			if (addSpanIds != null)
 			{
 				foreach (ObjectId addSpanId in addSpanIds)
diff --git a/Engine/Source/Programs/Horde/Horde.Build/Issues/IssuesController.cs b/Engine/Source/Programs/Horde/Horde.Build/Issues/IssuesController.cs
--- a/Engine/Source/Programs/Horde/Horde.Build/Issues/IssuesController.cs
+++ b/Engine/Source/Programs/Horde/Horde.Build/Issues/IssuesController.cs
@@ -606,6 +606,13 @@ namespace Horde.Build.Issues
 				newQuarantinedById = request.QuarantinedById.Length > 0 ? new UserId(request.QuarantinedById) : UserId.Empty;
 			}
 
+			UserId? newForceClosedById = null;
+			if (request.ForceClosedById != null)
+			{
+				newForceClosedById = request.ForceClosedById.Length > 0 ? new UserId(request.ForceClosedById) : UserId.Empty;
+			}
+
+
 			List<ObjectId>? addSpans = null;
 			if (request.AddSpans != null && request.AddSpans.Count > 0)
 			{
@@ -620,7 +627,7 @@ namespace Horde.Build.Issues
 
 			using (IDisposable scope = _issueCollection.GetLogger(issueId).BeginScope("User {UserId}", User.GetUserId() ?? UserId.Empty))
 			{
-				if (!await _issueService.UpdateIssueAsync(issueId, request.Summary, request.Description, request.Promoted, newOwnerId, newNominatedById, request.Acknowledged, newDeclinedById, request.FixChange, newResolvedById, addSpans, removeSpans, request.ExternalIssueKey, newQuarantinedById))
+				if (!await _issueService.UpdateIssueAsync(issueId, request.Summary, request.Description, request.Promoted, newOwnerId, newNominatedById, request.Acknowledged, newDeclinedById, request.FixChange, newResolvedById, addSpans, removeSpans, request.ExternalIssueKey, newQuarantinedById, newForceClosedById))
 				{
 					return NotFound();
 				}
