Horde: Add REST API for setting new pool sizing strategies

diff --git a/Engine/Source/Programs/Horde/Horde.Build.Tests/PoolsControllerTest.cs b/Engine/Source/Programs/Horde/Horde.Build.Tests/PoolsControllerTest.cs
--- a/Engine/Source/Programs/Horde/Horde.Build.Tests/PoolsControllerTest.cs
+++ b/Engine/Source/Programs/Horde/Horde.Build.Tests/PoolsControllerTest.cs
@@ -34,6 +34,7 @@ namespace Horde.Build.Tests
 		        Name = "Pool1",
 		        ScaleOutCooldown = 111,
 		        ScaleInCooldown = 222,
+		        SizeStrategies = new List<PoolSizeStrategyMessage> { new () { Type = PoolSizeStrategy.JobQueue, Condition = "dayOfWeek == 'monday'", Config = @"{""ScaleOutFactor"": 22.0, ""ScaleInFactor"": 33.0}" }},
 		        SizeStrategy = PoolSizeStrategy.JobQueue,
 		        JobQueueSettings = new JobQueueSettingsMessage(new JobQueueSettings(0.35, 0.85))
 	        };
@@ -47,6 +48,10 @@ namespace Horde.Build.Tests
 	        Assert.AreEqual(request.ScaleInCooldown, (int)pool.ScaleInCooldown!.Value.TotalSeconds);
 	        Assert.AreEqual(request.JobQueueSettings.ScaleOutFactor, pool.JobQueueSettings!.ScaleOutFactor, 0.0001);
 	        Assert.AreEqual(request.JobQueueSettings.ScaleInFactor, pool.JobQueueSettings!.ScaleInFactor, 0.0001);
+	        Assert.AreEqual(1, request.SizeStrategies.Count);
+	        Assert.AreEqual(request.SizeStrategies[0].Type, pool.SizeStrategies[0].Type);
+	        Assert.AreEqual(request.SizeStrategies[0].Condition!.Text, pool.SizeStrategies[0].Condition!.Text);
+	        Assert.AreEqual(request.SizeStrategies[0].Config, pool.SizeStrategies[0].Config);
         }
         
         [TestMethod]
@@ -57,6 +62,10 @@ namespace Horde.Build.Tests
 	        await PoolsController.UpdatePoolAsync(pool1.Id.ToString(), new UpdatePoolRequest
 	        {
 		        Name = "Pool1Modified",
+		        SizeStrategies = new List<PoolSizeStrategyMessage>
+		        {
+			        new () { Type = PoolSizeStrategy.JobQueue, Condition = "true", Config = "{\"someConfig\": 123"}
+		        },
 		        SizeStrategy = PoolSizeStrategy.JobQueue,
 		        JobQueueSettings = new JobQueueSettingsMessage { ScaleOutFactor = 25.0, ScaleInFactor = 0.3 }
 	        });
@@ -67,6 +76,11 @@ namespace Horde.Build.Tests
 	        Assert.AreEqual(PoolSizeStrategy.JobQueue, response.SizeStrategy);
 	        Assert.AreEqual(25.0, response.JobQueueSettings!.ScaleOutFactor);
 	        Assert.AreEqual(0.3, response.JobQueueSettings!.ScaleInFactor);
+	        
+	        Assert.AreEqual(1, response.SizeStrategies.Count);
+	        Assert.AreEqual(PoolSizeStrategy.JobQueue, response.SizeStrategies[0].Type);
+	        Assert.AreEqual("true", response.SizeStrategies[0].Condition!.Text);
+	        Assert.AreEqual("{\"someConfig\": 123", response.SizeStrategies[0].Config);
         }
     }
 }
\ No newline at end of file
diff --git a/Engine/Source/Programs/Horde/Horde.Build/Agents/Pools/IPoolCollection.cs b/Engine/Source/Programs/Horde/Horde.Build/Agents/Pools/IPoolCollection.cs
--- a/Engine/Source/Programs/Horde/Horde.Build/Agents/Pools/IPoolCollection.cs
+++ b/Engine/Source/Programs/Horde/Horde.Build/Agents/Pools/IPoolCollection.cs
@@ -28,6 +28,7 @@ namespace Horde.Build.Agents.Pools
 		/// <param name="conformInterval">Interval between conforms. Set to zero to disable.</param>
 		/// <param name="scaleOutCooldown">Cooldown time between scale-out events</param>
 		/// <param name="scaleInCooldown">Cooldown time between scale-in events</param>
+		/// <param name="sizeStrategies">Pool sizing strategies</param>
 		/// <param name="sizeStrategy">Pool sizing strategy</param>
 		/// <param name="leaseUtilizationSettings">Settings for lease utilization strategy</param>
 		/// <param name="jobQueueSettings">Settings for job queue strategy</param>
@@ -44,6 +45,7 @@ namespace Horde.Build.Agents.Pools
 			TimeSpan? conformInterval = null,
 			TimeSpan? scaleOutCooldown = null,
 			TimeSpan? scaleInCooldown = null,
+			List<PoolSizeStrategyInfo>? sizeStrategies = null,
 			PoolSizeStrategy? sizeStrategy = null,
 			LeaseUtilizationSettings? leaseUtilizationSettings = null,
 			JobQueueSettings? jobQueueSettings = null,
diff --git a/Engine/Source/Programs/Horde/Horde.Build/Agents/Pools/PoolCollection.cs b/Engine/Source/Programs/Horde/Horde.Build/Agents/Pools/PoolCollection.cs
--- a/Engine/Source/Programs/Horde/Horde.Build/Agents/Pools/PoolCollection.cs
+++ b/Engine/Source/Programs/Horde/Horde.Build/Agents/Pools/PoolCollection.cs
@@ -113,6 +113,7 @@ namespace Horde.Build.Agents.Pools
 			TimeSpan? conformInterval,
 			TimeSpan? scaleOutCooldown,
 			TimeSpan? scaleInCooldown,
+			List<PoolSizeStrategyInfo>? sizeStrategies,
 			PoolSizeStrategy? sizeStrategy,
 			LeaseUtilizationSettings? leaseUtilizationSettings,
 			JobQueueSettings? jobQueueSettings,
@@ -141,6 +142,10 @@ namespace Horde.Build.Agents.Pools
 			pool.LeaseUtilizationSettings = leaseUtilizationSettings;
 			pool.JobQueueSettings = jobQueueSettings;
 			pool.ComputeQueueAwsMetricSettings = computeQueueAwsMetricSettings;
+			if (sizeStrategies != null)
+			{
+				pool.SizeStrategies = sizeStrategies;
+			}
 			await _pools.InsertOneAsync(pool);
 			return pool;
 		}
diff --git a/Engine/Source/Programs/Horde/Horde.Build/Agents/Pools/PoolMessages.cs b/Engine/Source/Programs/Horde/Horde.Build/Agents/Pools/PoolMessages.cs
--- a/Engine/Source/Programs/Horde/Horde.Build/Agents/Pools/PoolMessages.cs
+++ b/Engine/Source/Programs/Horde/Horde.Build/Agents/Pools/PoolMessages.cs
@@ -8,6 +8,52 @@ using Horde.Build.Agents.Fleet;
 
 namespace Horde.Build.Agents.Pools
 {
+	/// <see cref="PoolSizeStrategyInfo" />
+	public class PoolSizeStrategyMessage
+	{
+		/// <summary>
+		/// Strategy implementation to use
+		/// </summary>
+		public PoolSizeStrategy Type { get; set; }
+		
+		/// <summary>
+		/// Condition if this strategy should be enabled (right now, using date/time as a distinguishing factor)
+		/// </summary>
+		public Condition? Condition { get; set; }
+
+		/// <summary>
+		/// Configuration for the strategy, serialized as JSON
+		/// </summary>
+		public string Config { get; set; } = "";
+		
+		/// <summary>
+		/// Empty constructor for JSON serialization
+		/// </summary>
+		public PoolSizeStrategyMessage()
+		{
+		}
+
+		/// <summary>
+		/// Construct a public REST API representation from the internal one
+		/// </summary>
+		/// <param name="strategy"></param>
+		public PoolSizeStrategyMessage(PoolSizeStrategyInfo strategy)
+		{
+			Type = strategy.Type;
+			Condition = strategy.Condition;
+			Config = strategy.Config;
+		}
+		
+		/// <summary>
+		/// Convert this public REST API object to an internal representation
+		/// </summary>
+		/// <returns></returns>
+		public PoolSizeStrategyInfo Convert()
+		{
+			return new PoolSizeStrategyInfo(Type, Condition, Config);
+		}
+	}
+
 	/// <see cref="LeaseUtilizationSettings" />
 	public class LeaseUtilizationSettingsMessage
 	{
@@ -151,6 +197,11 @@ namespace Horde.Build.Agents.Pools
 		/// </summary>
 		public int? ScaleInCooldown { get; set; }
 		
+		/// <summary>
+		/// Pool sizing strategies
+		/// </summary>
+		public List<PoolSizeStrategyMessage>? SizeStrategies { get; set; }
+		
 		/// <summary>
 		/// Pool sizing strategy
 		/// </summary>
@@ -248,9 +299,9 @@ namespace Horde.Build.Agents.Pools
 		public PoolSizeStrategy? SizeStrategy { get; set; }
 		
 		/// <summary>
-		/// Pool sizing strategy
+		/// Pool sizing strategies
 		/// </summary>
-		public List<PoolSizeStrategy>? SizeStrategies { get; set; }
+		public List<PoolSizeStrategyMessage>? SizeStrategies { get; set; }
 
 		/// <summary>
 		/// Set pool to use default strategy
@@ -351,22 +402,27 @@ namespace Horde.Build.Agents.Pools
 		public int? ScaleInCooldown { get; set; }
 		
 		/// <summary>
-		/// Pool sizing strategy to be used for this pool
+		/// Pool sizing strategies to be used for this pool
+		/// </summary>
+		public List<PoolSizeStrategyMessage> SizeStrategies { get; set; }
+
+		/// <summary>
+		/// Pool sizing strategy to be used for this pool (deprecated, see SizeStrategies field)
 		/// </summary>
 		public PoolSizeStrategy? SizeStrategy { get; set; }
 		
 		/// <summary>
-		/// Settings for lease utilization pool sizing strategy (if used)
+		/// Settings for lease utilization pool sizing strategy (deprecated, see SizeStrategies field)
 		/// </summary>
 		public LeaseUtilizationSettingsMessage? LeaseUtilizationSettings { get; set; }
 		
 		/// <summary>
-		/// Settings for job queue pool sizing strategy (if used) 
+		/// Settings for job queue pool sizing strategy (deprecated, see SizeStrategies field)
 		/// </summary>
 		public JobQueueSettingsMessage? JobQueueSettings { get; set; }
 
 		/// <summary>
-		/// Settings for compute queue pool sizing strategy for AWS metrics (if used) 
+		/// Settings for compute queue pool sizing strategy for AWS metrics (deprecated, see SizeStrategies field)
 		/// </summary>
 		public ComputeQueueAwsMetricSettingsMessage? ComputeQueueAwsMetricSettings { get; set; }
 
@@ -403,6 +459,7 @@ namespace Horde.Build.Agents.Pools
 			ConformInterval = pool.ConformInterval == null ? null : (int)pool.ConformInterval.Value.TotalHours;
 			ScaleOutCooldown = pool.ScaleOutCooldown == null ? null : (int)pool.ScaleOutCooldown.Value.TotalSeconds;
 			ScaleInCooldown = pool.ScaleInCooldown == null ? null : (int)pool.ScaleInCooldown.Value.TotalSeconds;
+			SizeStrategies = pool.SizeStrategies.Select(x => new PoolSizeStrategyMessage(x)).ToList();
 			SizeStrategy = pool.SizeStrategy;
 			LeaseUtilizationSettings = pool.LeaseUtilizationSettings == null ? null : new LeaseUtilizationSettingsMessage(pool.LeaseUtilizationSettings);
 			JobQueueSettings = pool.JobQueueSettings == null ? null : new JobQueueSettingsMessage(pool.JobQueueSettings);
diff --git a/Engine/Source/Programs/Horde/Horde.Build/Agents/Pools/PoolService.cs b/Engine/Source/Programs/Horde/Horde.Build/Agents/Pools/PoolService.cs
--- a/Engine/Source/Programs/Horde/Horde.Build/Agents/Pools/PoolService.cs
+++ b/Engine/Source/Programs/Horde/Horde.Build/Agents/Pools/PoolService.cs
@@ -64,6 +64,7 @@ namespace Horde.Build.Agents.Pools
 		/// <param name="conformInterval">Interval between conforms. Set to zero to disable.</param>
 		/// <param name="scaleOutCooldown">Cooldown time between scale-out events</param>
 		/// <param name="scaleInCooldown">Cooldown time between scale-in events</param>
+		/// <param name="sizeStrategies">Pool sizing strategies</param>
 		/// <param name="sizeStrategy">Pool sizing strategy</param>
 		/// <param name="leaseUtilizationSettings">Settings for lease utilization strategy</param>
 		/// <param name="jobQueueSettings">Settings for job queue strategy</param>
@@ -79,6 +80,7 @@ namespace Horde.Build.Agents.Pools
 			TimeSpan? conformInterval = null,
 			TimeSpan? scaleOutCooldown = null,
 			TimeSpan? scaleInCooldown = null,
+			List<PoolSizeStrategyInfo>? sizeStrategies = null,
 			PoolSizeStrategy? sizeStrategy = null,
 			LeaseUtilizationSettings? leaseUtilizationSettings = null,
 			JobQueueSettings? jobQueueSettings = null,
@@ -95,6 +97,7 @@ namespace Horde.Build.Agents.Pools
 				conformInterval,
 				scaleOutCooldown,
 				scaleInCooldown,
+				sizeStrategies,
 				sizeStrategy,
 				leaseUtilizationSettings,
 				jobQueueSettings,
diff --git a/Engine/Source/Programs/Horde/Horde.Build/Agents/Pools/PoolsController.cs b/Engine/Source/Programs/Horde/Horde.Build/Agents/Pools/PoolsController.cs
--- a/Engine/Source/Programs/Horde/Horde.Build/Agents/Pools/PoolsController.cs
+++ b/Engine/Source/Programs/Horde/Horde.Build/Agents/Pools/PoolsController.cs
@@ -2,6 +2,7 @@
 
 using System;
 using System.Collections.Generic;
+using System.Linq;
 using System.Threading.Tasks;
 using Horde.Build.Acls;
 using Horde.Build.Agents.Fleet;
@@ -56,6 +57,7 @@ namespace Horde.Build.Agents.Pools
 				return Forbid(AclAction.CreatePool);
 			}
 
+			List<PoolSizeStrategyInfo>? sizeStrategies = create.SizeStrategies?.Select(x => x.Convert()).ToList();
 			LeaseUtilizationSettings? luSettings = create.LeaseUtilizationSettings?.Convert();
 			JobQueueSettings? jqSettings = create.JobQueueSettings?.Convert();
 			ComputeQueueAwsMetricSettings? cqamSettings = create.ComputeQueueAwsMetricSettings?.Convert();
@@ -66,7 +68,7 @@ namespace Horde.Build.Agents.Pools
 
 			IPool newPool = await _poolService.CreatePoolAsync(
 				create.Name, create.Condition, create.EnableAutoscaling, create.MinAgents, create.NumReserveAgents, conformInterval,
-				scaleOutCooldown, scaleInCooldown, create.SizeStrategy, luSettings, jqSettings, cqamSettings, create.Properties);
+				scaleOutCooldown, scaleInCooldown, sizeStrategies, create.SizeStrategy, luSettings, jqSettings, cqamSettings, create.Properties);
 			return new CreatePoolResponse(newPool.Id.ToString());
 		}
 
@@ -145,12 +147,13 @@ namespace Horde.Build.Agents.Pools
 				return NotFound(poolIdValue);
 			}
 
+			List<PoolSizeStrategyInfo>? newSizeStrategies = update.SizeStrategies?.Select(x => x.Convert()).ToList();
 			TimeSpan? conformInterval = update.ConformInterval == null ? null : TimeSpan.FromHours(update.ConformInterval.Value);
 			TimeSpan? scaleOutCooldown = update.ScaleOutCooldown == null ? null : TimeSpan.FromSeconds(update.ScaleOutCooldown.Value);
 			TimeSpan? scaleInCooldown = update.ScaleInCooldown == null ? null : TimeSpan.FromSeconds(update.ScaleInCooldown.Value);
 
 			await _poolService.UpdatePoolAsync(pool, update.Name, update.Condition, update.EnableAutoscaling,
-				update.MinAgents, update.NumReserveAgents, update.Properties, conformInterval, scaleOutCooldown, scaleInCooldown, update.SizeStrategy, null,
+				update.MinAgents, update.NumReserveAgents, update.Properties, conformInterval, scaleOutCooldown, scaleInCooldown, update.SizeStrategy, newSizeStrategies,
 				update.LeaseUtilizationSettings?.Convert(), update.JobQueueSettings?.Convert(), update.ComputeQueueAwsMetricSettings?.Convert(), update.UseDefaultStrategy);
 			return new OkResult();
 		}
