Horde: Add force close info to issue responses
#rnx

diff --git a/Engine/Source/Programs/Horde/Horde.Build/Issues/IIssueDetails.cs b/Engine/Source/Programs/Horde/Horde.Build/Issues/IIssueDetails.cs
--- a/Engine/Source/Programs/Horde/Horde.Build/Issues/IIssueDetails.cs
+++ b/Engine/Source/Programs/Horde/Horde.Build/Issues/IIssueDetails.cs
@@ -69,6 +69,12 @@ namespace Horde.Build.Issues
 		/// </summary>
 		DateTime? QuarantineTimeUtc { get; }
 
+		/// <summary>
+		/// User who force closed the issue
+		/// </summary>
+		IUser? ForceClosedBy { get; }
+
+
 		/// <summary>
 		/// Determines whether the given user should be notified about the given issue
 		/// </summary>
diff --git a/Engine/Source/Programs/Horde/Horde.Build/Issues/IssueMessages.cs b/Engine/Source/Programs/Horde/Horde.Build/Issues/IssueMessages.cs
--- a/Engine/Source/Programs/Horde/Horde.Build/Issues/IssueMessages.cs
+++ b/Engine/Source/Programs/Horde/Horde.Build/Issues/IssueMessages.cs
@@ -529,6 +529,11 @@ namespace Horde.Build.Issues
 		/// </summary>
 		public DateTime? QuarantineTimeUtc { get; set; }
 
+		/// <summary>
+		/// User who force closed the issue
+		/// </summary>
+		public GetThinUserInfoResponse? ForceClosedByUserInfo { get; set; }
+
 		/// <summary>
 		/// Constructs a new issue
 		/// </summary>
@@ -589,6 +594,7 @@ namespace Horde.Build.Issues
 			ExternalIssueKey = details.ExternalIssueKey;
 			QuarantinedByUserInfo = details.QuarantinedBy != null ? new GetThinUserInfoResponse(details.QuarantinedBy) : null;
 			QuarantineTimeUtc = details.QuarantineTimeUtc;
+			ForceClosedByUserInfo = details.ForceClosedBy != null ? new GetThinUserInfoResponse(details.ForceClosedBy) : null;
 		}
 	}
 
diff --git a/Engine/Source/Programs/Horde/Horde.Build/Issues/IssueService.cs b/Engine/Source/Programs/Horde/Horde.Build/Issues/IssueService.cs
--- a/Engine/Source/Programs/Horde/Horde.Build/Issues/IssueService.cs
+++ b/Engine/Source/Programs/Horde/Horde.Build/Issues/IssueService.cs
@@ -57,6 +57,9 @@ namespace Horde.Build.Issues
 		/// <inheritdoc/>
 		public DateTime? QuarantineTimeUtc { get; set; }
 
+		/// <inheritdoc/>
+		public IUser? ForceClosedBy { get; set; }
+
 		/// <inheritdoc/>
 		public string? ExternalIssueKey { get; }
 
@@ -86,13 +89,14 @@ namespace Horde.Build.Issues
 		/// <param name="nominatedBy"></param>
 		/// <param name="resolvedBy"></param>
 		/// <param name="quarantinedBy"></param>
+		/// <param name="forceClosedBy"></param>
 		/// <param name="spans">Spans for this issue</param>
 		/// <param name="steps">Steps for this issue</param>
 		/// <param name="suspects"></param>
 		/// <param name="suspectUsers"></param>
 		/// <param name="showDesktopAlerts"></param>
 		/// <param name="externalIssueKey"></param>
-		public IssueDetails(IIssue issue, IUser? owner, IUser? nominatedBy, IUser? resolvedBy, IUser? quarantinedBy, IReadOnlyList<IIssueSpan> spans, IReadOnlyList<IIssueStep> steps, IReadOnlyList<IIssueSuspect> suspects, IReadOnlyList<IUser> suspectUsers, bool showDesktopAlerts, string? externalIssueKey)
+		public IssueDetails(IIssue issue, IUser? owner, IUser? nominatedBy, IUser? resolvedBy, IUser? quarantinedBy, IUser? forceClosedBy, IReadOnlyList<IIssueSpan> spans, IReadOnlyList<IIssueStep> steps, IReadOnlyList<IIssueSuspect> suspects, IReadOnlyList<IUser> suspectUsers, bool showDesktopAlerts, string? externalIssueKey)
 		{
 			Issue = issue;
 			Owner = owner;
@@ -100,6 +104,7 @@ namespace Horde.Build.Issues
 			ResolvedBy = resolvedBy;
 			QuarantinedBy = quarantinedBy;
 			QuarantineTimeUtc = issue.QuarantineTimeUtc;
+			ForceClosedBy = forceClosedBy;
 			Spans = spans;
 			Steps = steps;
 			Suspects = suspects;
@@ -338,6 +343,7 @@ namespace Horde.Build.Issues
 			IUser? nominatedBy = issue.NominatedById.HasValue ? await _userCollection.GetCachedUserAsync(issue.NominatedById.Value) : null;
 			IUser? resolvedBy = (issue.ResolvedById.HasValue && issue.ResolvedById != IIssue.ResolvedByTimeoutId && issue.ResolvedById != IIssue.ResolvedByUnknownId)? await _userCollection.GetCachedUserAsync(issue.ResolvedById.Value) : null;
 			IUser? quarantinedBy = issue.QuarantinedByUserId.HasValue ? await _userCollection.GetCachedUserAsync(issue.QuarantinedByUserId.Value) : null;
+			IUser? forceClosedBy = issue.ForceClosedByUserId.HasValue ? await _userCollection.GetCachedUserAsync(issue.ForceClosedByUserId.Value) : null;
 
 			List<IIssueSpan> spans = await _issueCollection.FindSpansAsync(issue.Id);
 			List<IIssueStep> steps = await _issueCollection.FindStepsAsync(spans.Select(x => x.Id));
@@ -353,7 +359,7 @@ namespace Horde.Build.Issues
 				}
 			}
 
-			return new IssueDetails(issue, owner, nominatedBy, resolvedBy, quarantinedBy, spans, steps, suspects, suspectUsers, ShowDesktopAlertsForIssue(issue, spans), issue.ExternalIssueKey);
+			return new IssueDetails(issue, owner, nominatedBy, resolvedBy, quarantinedBy, forceClosedBy, spans, steps, suspects, suspectUsers, ShowDesktopAlertsForIssue(issue, spans), issue.ExternalIssueKey);
 		}
 
 		/// <inheritdoc/>
