Horde: Allow configuring conditional base changelists to run preflights against.

* DefaultPreflight for a stream now contains a list of ChangeQueryConfig objects, each of which has a condition attached to them.
* Condition can query tags on the change being preflighted (eg. "tag.code == 1")
* ChangeQuery can also filter commits by tag for producing a change to build at.

diff --git a/Engine/Source/Programs/Horde/Horde.Build.Tests/JobServiceTests.cs b/Engine/Source/Programs/Horde/Horde.Build.Tests/JobServiceTests.cs
--- a/Engine/Source/Programs/Horde/Horde.Build.Tests/JobServiceTests.cs
+++ b/Engine/Source/Programs/Horde/Horde.Build.Tests/JobServiceTests.cs
@@ -2,21 +2,26 @@
 
 using System;
 using System.Collections.Generic;
+using System.Runtime.InteropServices;
 using System.Text;
+using System.Threading;
 using System.Threading.Tasks;
 using EpicGames.Core;
+using EpicGames.Horde.Common;
 using Horde.Build.Agents;
 using Horde.Build.Agents.Pools;
 using Horde.Build.Jobs;
 using Horde.Build.Jobs.Graphs;
 using Horde.Build.Jobs.Templates;
 using Horde.Build.Logs;
+using Horde.Build.Perforce;
 using Horde.Build.Projects;
 using Horde.Build.Server;
 using Horde.Build.Streams;
 using Horde.Build.Users;
 using Horde.Build.Utilities;
 using HordeCommon;
+using Microsoft.CodeAnalysis;
 using Microsoft.VisualStudio.TestTools.UnitTesting;
 
 namespace Horde.Build.Tests
@@ -74,6 +79,49 @@ namespace Horde.Build.Tests
 			Assert.AreEqual(chainedJob!.StartedByUserId, job!.StartedByUserId);
 		}
 
+		[TestMethod]
+		public async Task ChangeQueryTestAsync()
+		{
+			ProjectId projectId = new ProjectId("ue5");
+			const string ProjectConfigRevision = "projectconfig";
+			await ConfigCollection.AddConfigAsync(ProjectConfigRevision, new ProjectConfig { Name = "UE4" });
+			IProject? project = await ProjectService.Collection.AddOrUpdateAsync(projectId, ProjectConfigRevision, 0);
+			Assert.IsNotNull(project);
+
+			StreamId streamId = new StreamId("ue5-main");
+			const string StreamConfigRevision = "streamconfig";
+			await ConfigCollection.AddConfigAsync(StreamConfigRevision, new StreamConfig { });
+			IStream? stream = await CreateOrReplaceStreamAsync(streamId, null, projectId, new StreamConfig());
+			Assert.IsNotNull(stream);
+
+			IUser user = await UserCollection.FindOrAddUserByLoginAsync("Bob");
+
+			PerforceService.AddChange(streamId, 1000, user, "", new[] { "Foo.cpp" });
+			PerforceService.AddChange(streamId, 1001, user, "", new[] { "Bar.cpp" });
+			PerforceService.AddChange(streamId, 1002, user, "", new[] { "Baz.cpp" });
+			PerforceService.AddChange(streamId, 1003, user, "", new[] { "Foo.uasset" });
+			PerforceService.AddChange(streamId, 1004, user, "", new[] { "Foo.uasset" });
+
+			ICommitCollection commits = PerforceService.GetCommits(stream);
+
+			{
+				int? change = await JobService.EvaluateChangeQueryAsync(stream, new ChangeQueryConfig { CommitTag = CommitTag.Code }, null, commits, CancellationToken.None);
+				Assert.AreEqual(1002, change);
+			}
+			{
+				int? change = await JobService.EvaluateChangeQueryAsync(stream, new ChangeQueryConfig { CommitTag = CommitTag.Content }, null, commits, CancellationToken.None);
+				Assert.AreEqual(1004, change);
+			}
+			{
+				int? change = await JobService.EvaluateChangeQueryAsync(stream, new ChangeQueryConfig { CommitTag = CommitTag.Content, Condition = "tag.code == 1" }, new List<CommitTag> { CommitTag.Code }, commits, CancellationToken.None);
+				Assert.AreEqual(1004, change);
+			}
+			{
+				int? change = await JobService.EvaluateChangeQueryAsync(stream, new ChangeQueryConfig { CommitTag = CommitTag.Content, Condition = "tag.content == 1" }, new List<CommitTag> { CommitTag.Code }, commits, CancellationToken.None);
+				Assert.IsNull(change);
+			}
+		}
+
 		[TestMethod]
 		public async Task StopAnyDuplicateJobsByPreflight()
 		{
diff --git a/Engine/Source/Programs/Horde/Horde.Build.Tests/Stubs/Services/PerforceServiceStub.cs b/Engine/Source/Programs/Horde/Horde.Build.Tests/Stubs/Services/PerforceServiceStub.cs
--- a/Engine/Source/Programs/Horde/Horde.Build.Tests/Stubs/Services/PerforceServiceStub.cs
+++ b/Engine/Source/Programs/Horde/Horde.Build.Tests/Stubs/Services/PerforceServiceStub.cs
@@ -145,7 +145,7 @@ namespace Horde.Build.Tests.Stubs.Services
 			return Task.FromResult(results);
 		}
 
-		public Task<(CheckShelfResult, string?)> CheckShelfAsync(string clusterName, string streamName, int changeNumber, CancellationToken cancellationToken)
+		public Task<(CheckShelfResult, ShelfInfo?)> CheckShelfAsync(IStream stream, int changeNumber, CancellationToken cancellationToken)
 		{
 			throw new NotImplementedException();
 		}
@@ -165,26 +165,6 @@ namespace Horde.Build.Tests.Stubs.Services
 			return Task.FromResult("bogus-ticket");
 		}
 
-		public Task<int> GetCodeChangeAsync(IStream stream, int change, CancellationToken cancellationToken)
-		{
-			int codeChange = 0;
-
-			SortedDictionary<int, Commit>? streamChanges;
-			if (Changes.TryGetValue(stream.Id, out streamChanges))
-			{
-				foreach (Commit details in streamChanges.Values)
-				{
-					if (details.Number <= change && details.Files.Any(x => x.EndsWith(".h", StringComparison.OrdinalIgnoreCase) || x.EndsWith(".cpp", StringComparison.OrdinalIgnoreCase)))
-					{
-						codeChange = details.Number;
-						break;
-					}
-				}
-			}
-
-			return Task.FromResult(codeChange);
-		}
-
 		public Task<int> CreateNewChangeAsync(string clusterName, string streamName, string path, string description, CancellationToken cancellationToken)
 		{
 			throw new NotImplementedException();
diff --git a/Engine/Source/Programs/Horde/Horde.Build/Jobs/JobMessages.cs b/Engine/Source/Programs/Horde/Horde.Build/Jobs/JobMessages.cs
--- a/Engine/Source/Programs/Horde/Horde.Build/Jobs/JobMessages.cs
+++ b/Engine/Source/Programs/Horde/Horde.Build/Jobs/JobMessages.cs
@@ -68,7 +68,17 @@ namespace Horde.Build.Jobs
 		/// <summary>
 		/// Parameters to use when selecting the change to execute at.
 		/// </summary>
-		public ChangeQueryConfig? ChangeQuery { get; set; }
+		[Obsolete]
+		public ChangeQueryConfig? ChangeQuery
+		{
+			get => (ChangeQueries != null && ChangeQueries.Count > 0) ? ChangeQueries[0] : null;
+			set => ChangeQueries = (value == null) ? null : new List<ChangeQueryConfig> { value };
+		}
+			
+		/// <summary>
+		/// List of change queries to evaluate
+		/// </summary>
+		public List<ChangeQueryConfig>? ChangeQueries { get; set; }
 
 		/// <summary>
 		/// The preflight changelist number
diff --git a/Engine/Source/Programs/Horde/Horde.Build/Jobs/JobService.cs b/Engine/Source/Programs/Horde/Horde.Build/Jobs/JobService.cs
--- a/Engine/Source/Programs/Horde/Horde.Build/Jobs/JobService.cs
+++ b/Engine/Source/Programs/Horde/Horde.Build/Jobs/JobService.cs
@@ -6,6 +6,7 @@ using System.Globalization;
 using System.Linq;
 using System.Security.Claims;
 using System.Text.RegularExpressions;
+using System.Threading;
 using System.Threading.Tasks;
 using EpicGames.Core;
 using EpicGames.Perforce;
@@ -419,6 +420,74 @@ namespace Horde.Build.Jobs
 			return null;
 		}
 
+		/// <summary>
+		/// Evaluate a change query to determine which CL to run a job at
+		/// </summary>
+		/// <returns></returns>
+		public async Task<int?> EvaluateChangeQueriesAsync(IStream stream, List<ChangeQueryConfig> queries, List<CommitTag>? commitTags, ICommitCollection commits, CancellationToken cancellationToken)
+		{
+			foreach (ChangeQueryConfig query in queries)
+			{
+				int? change = await EvaluateChangeQueryAsync(stream, query, commitTags, commits, cancellationToken);
+				if (change != null)
+				{
+					return change;
+				}
+			}
+			return null;
+		}
+
+		/// <summary>
+		/// Evaluate a change query to determine which CL to run a job at
+		/// </summary>
+		/// <returns></returns>
+		public async ValueTask<int?> EvaluateChangeQueryAsync(IStream stream, ChangeQueryConfig query, List<CommitTag>? commitTags, ICommitCollection commits, CancellationToken cancellationToken)
+		{
+			if (query.Condition == null || query.Condition.Evaluate(propertyName => GetTagPropertyValues(propertyName, commitTags)))
+			{
+				// Find the last change with the listed tag
+				if (query.CommitTag != null)
+				{
+					ICommit? taggedCommit = await commits.FindAsync(null, null, 1, new[] { query.CommitTag.Value }, cancellationToken).FirstOrDefaultAsync(cancellationToken);
+					if (taggedCommit != null)
+					{
+						_logger.LogInformation("Last commit with tag '{Tag}' was {Change}", query.CommitTag.Value, taggedCommit.Number);
+						return taggedCommit.Number;
+					}
+				}
+
+				// Find the job with the given result
+				if (query.Target != null && query.TemplateId != null)
+				{
+					List<JobStepOutcome> outcomes = query.Outcomes ?? new List<JobStepOutcome> { JobStepOutcome.Success };
+
+					IList<IJob> jobs = await FindJobsAsync(streamId: stream.Id, templates: new[] { query.TemplateId.Value }, target: query.Target, state: new[] { JobStepState.Completed }, outcome: outcomes.ToArray(), count: 1, excludeUserJobs: true);
+					if (jobs.Count > 0)
+					{
+						_logger.LogInformation("Last successful build of {TemplateId} target {Target} was job {JobId} at change {Change}", query.TemplateId, query.Target, jobs[0].Id, jobs[0].Change);
+						return jobs[0].Change;
+					}
+				}
+			}
+			return null;
+		}
+
+		static IEnumerable<string> GetTagPropertyValues(string propertyName, List<CommitTag>? commitTags)
+		{
+			const string TagPrefix = "tag.";
+			if (propertyName.StartsWith(TagPrefix, StringComparison.Ordinal) && commitTags != null)
+			{
+				string tagName = propertyName.Substring(TagPrefix.Length);
+				foreach (CommitTag commitTag in commitTags)
+				{
+					if (tagName.Equals(commitTag.Text, StringComparison.OrdinalIgnoreCase))
+					{
+						yield return "1";
+					}
+				}
+			}
+		}
+
 		/// <summary>
 		/// Cancel an active lease
 		/// </summary>
diff --git a/Engine/Source/Programs/Horde/Horde.Build/Jobs/JobsController.cs b/Engine/Source/Programs/Horde/Horde.Build/Jobs/JobsController.cs
--- a/Engine/Source/Programs/Horde/Horde.Build/Jobs/JobsController.cs
+++ b/Engine/Source/Programs/Horde/Horde.Build/Jobs/JobsController.cs
@@ -3,6 +3,7 @@
 using System;
 using System.Collections.Generic;
 using System.Linq;
+using System.Threading;
 using System.Threading.Tasks;
 using EpicGames.Core;
 using Horde.Build.Acls;
@@ -128,46 +129,11 @@ namespace Horde.Build.Jobs
 				}
 			}
 
-			// Get the priority of the new job
-			Priority priority = create.Priority ?? template.Priority ?? Priority.Normal;
-
-			// New groups for the job
-			IGraph graph = await _graphs.AddAsync(template);
-
-			// Get the commits for this stream
-			ICommitCollection commits = _commitService.GetCollection(stream);
-
-			// Get the change to build
-			int change;
-			if (create.Change.HasValue && create.Change.Value != -1)
-			{
-				change = create.Change.Value;
-			}
-			else if (create.ChangeQuery != null)
-			{
-				change = await ExecuteChangeQueryAsync(stream, create.ChangeQuery.TemplateId ?? create.TemplateId, create.ChangeQuery.Target, create.ChangeQuery.Outcomes ?? new List<JobStepOutcome> { JobStepOutcome.Success }, commits);
-			}
-			else if (create.PreflightChange == null && template.SubmitNewChange != null)
-			{
-				change = await commits.CreateNewAsync(template, HttpContext.RequestAborted);
-			}
-			else
-			{
-				change = await commits.GetLatestNumberAsync(HttpContext.RequestAborted);
-			}
-
-			// And get the matching code changelist
-			ICommit? lastCodeCommit = await commits.GetLastCodeChange(change, HttpContext.RequestAborted);
-			int codeChange = lastCodeCommit?.Number ?? change;
-
-			// New properties for the job
-			List<string> arguments = create.Arguments ?? template.GetDefaultArguments();
-
 			// Check the preflight change is valid
-			string? preflightDescription = null;
+			ShelfInfo? shelfInfo = null;
 			if (create.PreflightChange != null)
 			{
-				(CheckShelfResult result, preflightDescription) = await _perforce.CheckShelfAsync(stream.Config.ClusterName, stream.Name, create.PreflightChange.Value);
+				(CheckShelfResult result, shelfInfo) = await _perforce.CheckShelfAsync(stream, create.PreflightChange.Value, HttpContext.RequestAborted);
 				switch (result)
 				{
 					case CheckShelfResult.Ok:
@@ -185,6 +151,25 @@ namespace Horde.Build.Jobs
 				}
 			}
 
+			// Get the priority of the new job
+			Priority priority = create.Priority ?? template.Priority ?? Priority.Normal;
+
+			// New groups for the job
+			IGraph graph = await _graphs.AddAsync(template);
+
+			// Get the commits for this stream
+			ICommitCollection commits = _commitService.GetCollection(stream);
+
+			// Get the change to build
+			int change = await GetChangeToBuildAsync(create, stream, template, shelfInfo, commits, HttpContext.RequestAborted);
+
+			// And get the matching code changelist
+			ICommit? lastCodeCommit = await commits.GetLastCodeChange(change, HttpContext.RequestAborted);
+			int codeChange = lastCodeCommit?.Number ?? change;
+
+			// New properties for the job
+			List<string> arguments = create.Arguments ?? template.GetDefaultArguments();
+
 			bool? updateIssues = null;
 			if (template.UpdateIssues)
 			{
@@ -196,33 +181,37 @@ namespace Horde.Build.Jobs
 			}
 
 			// Create the job
-			IJob job = await _jobService.CreateJobAsync(null, stream, templateRefId, template.Id, graph, name, change, codeChange, create.PreflightChange, null, preflightDescription, User.GetUserId(), priority, create.AutoSubmit, updateIssues, false, templateRef.Config.ChainedJobs, templateRef.Config.ShowUgsBadges, templateRef.Config.ShowUgsAlerts, templateRef.Config.NotificationChannel, templateRef.Config.NotificationChannelFilter, arguments);
+			IJob job = await _jobService.CreateJobAsync(null, stream, templateRefId, template.Id, graph, name, change, codeChange, create.PreflightChange, null, shelfInfo?.Description, User.GetUserId(), priority, create.AutoSubmit, updateIssues, false, templateRef.Config.ChainedJobs, templateRef.Config.ShowUgsBadges, templateRef.Config.ShowUgsAlerts, templateRef.Config.NotificationChannel, templateRef.Config.NotificationChannelFilter, arguments);
 			await UpdateNotificationsAsync(job.Id, new UpdateNotificationsRequest { Slack = true });
 			return new CreateJobResponse(job.Id.ToString());
 		}
 
-		/// <summary>
-		/// Evaluate a change query to determine which CL to run a job at
-		/// </summary>
-		/// <param name="stream"></param>
-		/// <param name="templateId"></param>
-		/// <param name="target"></param>
-		/// <param name="outcomes"></param>
-		/// <param name="commits"></param>
-		/// <returns></returns>
-		async Task<int> ExecuteChangeQueryAsync(IStream stream, TemplateId templateId, string? target, List<JobStepOutcome> outcomes, ICommitCollection commits)
+		async ValueTask<int> GetChangeToBuildAsync(CreateJobRequest create, IStream stream, ITemplate template, ShelfInfo? shelfInfo, ICommitCollection commits, CancellationToken cancellationToken)
 		{
-			IList<IJob> jobs = await _jobService.FindJobsAsync(streamId: stream.Id, templates: new[] { templateId }, target: target, state: new[] { JobStepState.Completed }, outcome: outcomes.ToArray(), count: 1, excludeUserJobs: true);
-			if (jobs.Count == 0)
+			// If there's an explicit change specified, use that
+			if (create.Change.HasValue && create.Change.Value != -1)
 			{
-				_logger.LogInformation("Unable to find successful build of {TemplateId} target {Target}. Using latest change instead", templateId, target);
-				return await commits.GetLatestNumberAsync();
+				return create.Change.Value;
 			}
-			else
+
+			// Evaluate the change queries
+			if (create.ChangeQueries != null && create.ChangeQueries.Count > 0)
 			{
-				_logger.LogInformation("Last successful build of {TemplateId} target {Target} was job {JobId} at change {Change}", templateId, target, jobs[0].Id, jobs[0].Change);
-				return jobs[0].Change;
+				int? change = await _jobService.EvaluateChangeQueriesAsync(stream, create.ChangeQueries, shelfInfo?.Tags, commits, cancellationToken);
+				if (change != null)
+				{
+					return change.Value;
+				}
 			}
+
+			// If we need to submit a new change, do that
+			if (create.PreflightChange == null && template.SubmitNewChange != null)
+			{
+				return await commits.CreateNewAsync(template, cancellationToken);
+			}
+
+			// Otherwise return the latest change
+			return await commits.GetLatestNumberAsync(cancellationToken);
 		}
 
 		/// <summary>
diff --git a/Engine/Source/Programs/Horde/Horde.Build/Perforce/IPerforceService.cs b/Engine/Source/Programs/Horde/Horde.Build/Perforce/IPerforceService.cs
--- a/Engine/Source/Programs/Horde/Horde.Build/Perforce/IPerforceService.cs
+++ b/Engine/Source/Programs/Horde/Horde.Build/Perforce/IPerforceService.cs
@@ -44,6 +44,31 @@ namespace Horde.Build.Perforce
 		WrongStream,
 	}
 
+	/// <summary>
+	/// Information about a shelved change
+	/// </summary>
+	public class ShelfInfo
+	{
+		/// <summary>
+		/// Description for the change
+		/// </summary>
+		public string Description { get; set; }
+
+		/// <summary>
+		/// Tags for the change
+		/// </summary>
+		public List<CommitTag> Tags { get; set; }
+
+		/// <summary>
+		/// Constructor
+		/// </summary>
+		public ShelfInfo(string description, List<CommitTag> tags)
+		{
+			Description = description;
+			Tags = tags;
+		}
+	}
+
 	/// <summary>
 	/// A connection returned by the Perforce service
 	/// </summary>
@@ -130,12 +155,11 @@ namespace Horde.Build.Perforce
 		/// <summary>
 		/// Checks a shelf is valid for the given stream
 		/// </summary>
-		/// <param name="clusterName">Name of the Perforce cluster</param>
-		/// <param name="streamName">The stream to query</param>
+		/// <param name="stream">The stream to query</param>
 		/// <param name="changeNumber">Shelved changelist number</param>
 		/// <param name="cancellationToken">Cancellation token for the operation</param>
 		/// <returns>Whether the shelf is valid, plus the description for it</returns>
-		public Task<(CheckShelfResult, string?)> CheckShelfAsync(string clusterName, string streamName, int changeNumber, CancellationToken cancellationToken = default);
+		public Task<(CheckShelfResult, ShelfInfo?)> CheckShelfAsync(IStream stream, int changeNumber, CancellationToken cancellationToken = default);
 
 		/// <summary>
 		/// Duplicates a shelved changelist
diff --git a/Engine/Source/Programs/Horde/Horde.Build/Perforce/PerforceService.cs b/Engine/Source/Programs/Horde/Horde.Build/Perforce/PerforceService.cs
--- a/Engine/Source/Programs/Horde/Horde.Build/Perforce/PerforceService.cs
+++ b/Engine/Source/Programs/Horde/Horde.Build/Perforce/PerforceService.cs
@@ -746,17 +746,16 @@ namespace Horde.Build.Perforce
 		}
 
 		/// <inheritdoc/>
-		public async Task<(CheckShelfResult, string?)> CheckShelfAsync(string clusterName, string streamName, int changeNumber, CancellationToken cancellationToken)
+		public async Task<(CheckShelfResult, ShelfInfo?)> CheckShelfAsync(IStream stream, int changeNumber, CancellationToken cancellationToken)
 		{
 			using IScope scope = GlobalTracer.Instance.BuildSpan("PerforceService.CheckPreflightAsync").StartActive();
-			scope.Span.SetTag("ClusterName", clusterName);
-			scope.Span.SetTag("StreamName", streamName);
+			scope.Span.SetTag("ClusterName", stream.Config.ClusterName);
+			scope.Span.SetTag("StreamName", stream.Name);
 			scope.Span.SetTag("ChangeNumber", changeNumber);
 
-			PerforceCluster cluster = await GetClusterAsync(clusterName);
-			IPerforceServer server = await GetServerAsync(cluster);
+			PerforceCluster cluster = await GetClusterAsync(stream.Config.ClusterName);
 
-			using (IPerforceConnection perforce = await ConnectAsync(cluster, null, cancellationToken))
+			using (IPooledPerforceConnection perforce = await ConnectAsync(cluster, null, cancellationToken))
 			{
 				PerforceResponse<DescribeRecord> response = await perforce.TryDescribeAsync(DescribeOptions.Shelved, -1, changeNumber, cancellationToken);
 				if(response.Error != null)
@@ -777,16 +776,16 @@ namespace Horde.Build.Perforce
 					return (CheckShelfResult.NoShelvedFiles, null);
 				}
 
-				StreamRecord stream = await perforce.GetStreamAsync(streamName, true, cancellationToken);
-				PerforceViewMap viewMap = PerforceViewMap.Parse(stream.View);
+				PerforceViewMap viewMap = await perforce.GetCachedStreamViewAsync(stream, cancellationToken);
+				List<string> mappedFiles = new List<string>();
 
-				bool bHasMappedFile = false;
 				bool bHasUnmappedFile = false;
 				foreach (DescribeFileRecord shelvedFile in change.Files)
 				{
-					if (viewMap.TryMapFile(shelvedFile.DepotFile, StringComparison.OrdinalIgnoreCase, out string _))
+					string mappedFile;
+					if (viewMap.TryMapFile(shelvedFile.DepotFile, StringComparison.OrdinalIgnoreCase, out mappedFile))
 					{
-						bHasMappedFile = true;
+						mappedFiles.Add(mappedFile);
 					}
 					else
 					{
@@ -796,10 +795,20 @@ namespace Horde.Build.Perforce
 
 				if (bHasUnmappedFile)
 				{
-					return (bHasMappedFile ? CheckShelfResult.MixedStream : CheckShelfResult.WrongStream, null);
+					return ((mappedFiles.Count > 0)? CheckShelfResult.MixedStream : CheckShelfResult.WrongStream, null);
+				}
+
+				List<CommitTag> tags = new List<CommitTag>();
+				foreach (CommitTagConfig tagConfig in stream.Config.CommitTags)
+				{
+					FileFilter filter = tagConfig.CreateFileFilter();
+					if (mappedFiles.Any(x => filter.Matches(x)))
+					{
+						tags.Add(tagConfig.Name);
+					}
 				}
 
-				return (CheckShelfResult.Ok, change.Description);
+				return (CheckShelfResult.Ok, new ShelfInfo(change.Description, tags));
 			}
 		}
 
diff --git a/Engine/Source/Programs/Horde/Horde.Build/Streams/StreamConfig.cs b/Engine/Source/Programs/Horde/Horde.Build/Streams/StreamConfig.cs
--- a/Engine/Source/Programs/Horde/Horde.Build/Streams/StreamConfig.cs
+++ b/Engine/Source/Programs/Horde/Horde.Build/Streams/StreamConfig.cs
@@ -7,6 +7,7 @@ using System.Diagnostics.CodeAnalysis;
 using System.Linq;
 using Datadog.Trace;
 using EpicGames.Core;
+using EpicGames.Horde.Common;
 using EpicGames.Perforce;
 using Horde.Build.Acls;
 using Horde.Build.Agents.Pools;
@@ -464,6 +465,16 @@ namespace Horde.Build.Streams
 		/// Whether to match a job that produced warnings
 		/// </summary>
 		public List<JobStepOutcome>? Outcomes { get; set; }
+
+		/// <summary>
+		/// Finds the last commit with this tag
+		/// </summary>
+		public CommitTag? CommitTag { get; set; }
+
+		/// <summary>
+		/// Condition to evaluate before deciding to use this query. May query tags in a preflight.
+		/// </summary>
+		public Condition? Condition { get; set; }
 	}
 
 	/// <summary>
@@ -479,7 +490,17 @@ namespace Horde.Build.Streams
 		/// <summary>
 		/// Query for the change to use
 		/// </summary>
-		public ChangeQueryConfig? Change { get; set; }
+		[Obsolete("Use ChangeQueries instead")]
+		public ChangeQueryConfig? Change
+		{
+			get => (ChangeQueries != null && ChangeQueries.Count > 0) ? ChangeQueries[0] : null;
+			set => ChangeQueries = ((value == null) ? null : new List<ChangeQueryConfig> { value });
+		}
+
+		/// <summary>
+		/// List of queries to execute to determine the base change for this preflight
+		/// </summary>
+		public List<ChangeQueryConfig>? ChangeQueries { get; set; }
 	}
 
 	/// <summary>
