Horde: Force close dashboard side, and align issue modal top elements
#rnx

diff --git a/Engine/Source/Programs/Horde/HordeDashboard/src/backend/Api.ts b/Engine/Source/Programs/Horde/HordeDashboard/src/backend/Api.ts
--- a/Engine/Source/Programs/Horde/HordeDashboard/src/backend/Api.ts
+++ b/Engine/Source/Programs/Horde/HordeDashboard/src/backend/Api.ts
@@ -2914,6 +2914,9 @@ export type GetIssueResponse = {
 	/** The UTC time when the issue was quarantined */
 	quarantineTimeUtc?: Date | string;
 
+	/** User info for who force closed the issue */
+	forceClosedByUserInfo?: GetThinUserInfoResponse;
+
 }
 
 /**Request an issue to be updated */
@@ -2958,6 +2961,9 @@ export type UpdateIssueRequest = {
 	/** Id of user quarantining issue */
 	quarantinedById?: string;
 
+	/** Id of user force closing the issue */
+	forceClosedById?: string;
+
 }
 
 export type GetUtilizationTelemetryStream = {
diff --git a/Engine/Source/Programs/Horde/HordeDashboard/src/components/IssueViewV2.tsx b/Engine/Source/Programs/Horde/HordeDashboard/src/components/IssueViewV2.tsx
--- a/Engine/Source/Programs/Horde/HordeDashboard/src/components/IssueViewV2.tsx
+++ b/Engine/Source/Programs/Horde/HordeDashboard/src/components/IssueViewV2.tsx
@@ -1167,7 +1167,7 @@ const IssueHeader: React.FC = () => {
    const jiraStatus = jiraIssue?.resolutionName ?? jiraIssue?.statusName ?? "";
 
    return <Stack horizontal tokens={{ childrenGap: 48 }} >
-      {(!!jiraIssue || !!issue.quarantineTimeUtc) && <Stack style={{ width: 340, padding: 0 }} tokens={{ childrenGap: 12 }}>
+      {(!!jiraIssue || !!issue.quarantineTimeUtc || !!issue.forceClosedByUserInfo) && <Stack style={{ width: 340, padding: 0 }} tokens={{ childrenGap: 12 }}>
          {!!jiraIssue && <Stack tokens={{ childrenGap: 4 }}>
             {renderSummaryItem(`Jira`, jiraIssue.key, jiraIssue.link!, !!jiraIssue.resolutionName)}
             {(!!jiraAssignee) && renderSummaryItem(`Jira Assignee`, jiraAssignee, jiraIssue!.link!, !!jiraIssue!.resolutionName)}
@@ -1176,6 +1176,9 @@ const IssueHeader: React.FC = () => {
          {!!issue.quarantinedByUserInfo && <Stack>
             {renderSummaryItem(`Quarantined By`, `${issue.quarantinedByUserInfo.name} on ${getShortNiceTime(issue.quarantineTimeUtc)}`)}
          </Stack>}
+         {!!issue.forceClosedByUserInfo && <Stack>
+            {renderSummaryItem(`Force Closed By`, `${issue.forceClosedByUserInfo.name}`)}
+         </Stack>}
       </Stack>}
       {!!issue.description && <Stack>
          <Label style={{ padding: 0 }}>{`Description:`}</Label>
@@ -1189,7 +1192,7 @@ const IssueSummaryPanel: React.FC = () => {
 
    const issue = details.issue!;
 
-   if (!issue || (!details.jiraIssue && !issue.description && !issue.quarantinedByUserInfo)) {
+   if (!issue || (!details.jiraIssue && !issue.description && !issue.quarantinedByUserInfo && !issue.forceClosedByUserInfo)) {
       return null;
    }
 
@@ -1457,6 +1460,7 @@ const IssueCommandBar: React.FC = () => {
    const [externalIssueLinkShown, setExternalIssueLinkShown] = useState(false);
    const [externalIssueCreateShown, setExternalIssueCreateShown] = useState(false);
    const [quarantineShown, setQuarantineShown] = useState(false);
+   const [forceCloseShown, setForceCloseShown] = useState(false);
    const [testFixShown, setTestFixShown] = useState(false);
 
    const issue = details.issue!;
@@ -1608,11 +1612,6 @@ const IssueCommandBar: React.FC = () => {
 
    const moreList: IContextualMenuItem[] = [
       {
-         key: 'more_test_fix',
-         iconProps: { iconName: "Wrench" },
-         text: `Test Fix`,
-         onClick: () => { setTestFixShown(true); }
-      }, {
          key: 'more_history_link',
          iconProps: { iconName: "SearchTemplate" },
          text: `History`,
@@ -1624,10 +1623,19 @@ const IssueCommandBar: React.FC = () => {
          iconProps: { iconName: "Flash" },
          text: `Quarantine`,
          onClick: () => { setQuarantineShown(true); }
+      }, {
+         key: 'more_forceclose_link',
+         iconProps: { iconName: "Delete" },
+         text: `Force Close`,
+         onClick: () => { setForceCloseShown(true); }
+      }, {
+         key: 'more_test_fix',
+         iconProps: { iconName: "Wrench" },
+         text: `Test Fix`,
+         onClick: () => { setTestFixShown(true); }
       }];
 
 
-
    const moreItems: ICommandBarItemProps[] = [
       {
          key: 'more_items',
@@ -1662,28 +1670,29 @@ const IssueCommandBar: React.FC = () => {
       {externalIssueLinkShown && <LinkExternalIssueModal onClose={() => { setExternalIssueLinkShown(false) }} />}
       {externalIssueCreateShown && <CreateExternalIssueModal onClose={() => { setExternalIssueCreateShown(false) }} />}
       {quarantineShown && <IssueQuarantineModal onClose={() => { setQuarantineShown(false) }} />}
+      {forceCloseShown && <IssueForceCloseModal onClose={() => { setForceCloseShown(false) }} />}
       {testFixShown && <TestFixModal onClose={() => { setTestFixShown(false) }} />}
 
       <Stack horizontal>
          <Stack>
-            <Stack styles={{ root: { paddingBottom: 8, paddingRight: 32 } }}>
+            <Stack styles={{ root: { paddingBottom: 0, paddingRight: 32 } }}>
                <Stack className={hordeClasses.commandBarSmall} horizontal verticalAlign="center" tokens={{ childrenGap: 8 }}>
                   <Stack horizontal style={{ paddingLeft: 12 }}>
-                     <Stack horizontal verticalAlign="center" style={{ paddingTop: 4, paddingRight: 18 }}>
-                        {!!suspectRange && <a href={suspectRange} target="blank"> <Stack styles={{ root: { paddingTop: 3 } }}>
+                     <Stack horizontal verticalAlign="center" style={{paddingRight: 18 }}>
+                        {!!suspectRange && <a href={suspectRange} target="blank"> <Stack>
                            <CommandBarButton className={hordeClasses.commandBarSmall} styles={{ root: { padding: "10px 8px 10px 8px" } }} iconProps={{ iconName: "Locate" }} text="Suspects" />
                         </Stack>
                         </a>}
 
                      </Stack>
-                     <Stack horizontal verticalAlign="center" style={{ paddingTop: 2 }} tokens={{ childrenGap: 18 }}>
+                     <Stack horizontal verticalAlign="center" verticalFill={true} tokens={{ childrenGap: 18 }}>
                         <Stack>
                            <CommandBar styles={{ root: { height: 32, padding: 0, paddingRight: 8, paddingLeft: 16, backgroundColor: "unset" } }}
                               className={customClasses.actionBar}
                               items={commandItems}
                               onReduceData={() => undefined} />
                         </Stack>
-                        {canAck && <Stack className="horde-no-darktheme" style={{ paddingRight: 8, paddingTop: 8 }}>
+                        {canAck && <Stack className="horde-no-darktheme" style={{ paddingRight: 8 }}>
                            <PrimaryButton style={{ animationName: "red-pulse", animationDuration: "2s", animationIterationCount: "infinite", backgroundColor: "#FF0000", border: "1px solid #FF0000", padding: 0, paddingLeft: 4, paddingRight: 4, height: 22, fontWeight: "unset", fontSize: 12 }} text="Acknowledge" onClick={() => { setAckShown(true) }}></PrimaryButton>
                         </Stack>}
                         <Stack>
@@ -1780,13 +1789,13 @@ export const IssueModalV2: React.FC<{ popHistoryOnClose: boolean, issueId?: stri
          <Stack style={{ height: "100%" }}>
             <Stack style={{ flexBasis: "70px", flexShrink: 0 }}>
                <Stack horizontal styles={{ root: { padding: 8 } }} style={{ padding: 20, paddingBottom: 8 }}>
-                  <Stack horizontal style={{ paddingTop: 5, width: 1024 }} tokens={{ childrenGap: 24 }} verticalAlign="start">
+                  <Stack horizontal style={{ width: 1024 }} tokens={{ childrenGap: 24 }} verticalAlign="center" verticalFill={true}>
                      <Stack style={{ maxWidth: 880, lineBreak: "anywhere" }} >
                         <Text styles={{ root: { fontWeight: "unset", fontFamily: "Horde Open Sans SemiBold", fontSize: "15px", color: "#087BC4", textDecoration: details.issue?.resolvedAt ? "line-through" : undefined } }}>{title}</Text>
                      </Stack>
 
                      <Stack onClick={() => { setEditShown(true) }} style={{ cursor: "pointer" }}>
-                        <Icon styles={{ root: { margin: '0px', padding: '0px', paddingTop: "4px" } }} iconName="Edit" className={hordeClasses.iconBlue} />
+                        <Icon styles={{ root: { margin: '0px', padding: '0px', paddingTop: "0px" } }} iconName="Edit" className={hordeClasses.iconBlue} />
                      </Stack>
                      <Stack grow />
                   </Stack>
@@ -2857,6 +2866,91 @@ export const IssueQuarantineModal: React.FC<{ onClose: () => void }> = ({ onClos
 
 };
 
+
+export const IssueForceCloseModal: React.FC<{ onClose: () => void }> = ({ onClose }) => {
+
+   const issue = details.issue!;
+   const issueId = issue.id;
+
+   const [state, setState] = useState<{ submitting?: boolean, forceClosed?: boolean, error?: string }>({ forceClosed: !!issue.forceClosedByUserInfo });
+
+   const onSave = async () => {
+
+      setState({ ...state, submitting: true });
+
+      try {
+
+         await backend.updateIssue(issueId, {
+            forceClosedById: state.forceClosed ? dashboard.userId : ""
+         });
+
+         details.refresh(issueId).finally(() => {
+            setState({ ...state, submitting: false, error: "" });
+            onClose();
+         })
+
+      } catch (reason) {
+
+         setState({ ...state, submitting: false, error: reason as string });
+      }
+
+   };
+
+
+   const height = state.error ? 260 : 240;
+
+   return <Modal className={hordeClasses.modal} isOpen={true} topOffsetFixed={true} styles={{ main: { padding: 8, width: 540, height: height, minHeight: height, hasBeenOpened: false, top: "24px", position: "absolute" } }} onDismiss={() => { onClose() }}>
+      <Stack horizontal styles={{ root: { padding: 8 } }}>
+         <Stack.Item grow={2}>
+            <Text variant="mediumPlus">Force Close Issue</Text>
+         </Stack.Item>
+         <Stack.Item grow={0}>
+            <IconButton
+               iconProps={{ iconName: 'Cancel' }}
+               ariaLabel="Close popup modal"
+               onClick={() => { onClose(); }}
+            />
+         </Stack.Item>
+      </Stack>
+
+      <Stack styles={{ root: { paddingLeft: 18, paddingRight: 18, width: 540, paddingTop: 12 } }}>
+         <Stack tokens={{ childrenGap: 12 }}>
+            {!!state.error && <MessageBar
+               messageBarType={MessageBarType.error}
+               isMultiline={false}> {state.error} </MessageBar>}
+            <Stack>
+               <Text variant="medium">Warning: Force closing issues without verification is not regular workflow.  It is intended for special cases, such as steps being removed from a job.  Please proceed with caution.</Text>
+            </Stack>
+            <Stack>
+               <Checkbox disabled={!!issue.forceClosedByUserInfo} label={!issue.forceClosedByUserInfo ? "Force Close Issue" : `Force Closed by ${issue.forceClosedByUserInfo.name}`}
+                  defaultChecked={state.forceClosed}
+                  onChange={(ev, newValue) => {
+                     setState({ ...state, forceClosed: !!newValue });
+                  }} />
+            </Stack>
+         </Stack>
+      </Stack>
+
+      <Stack tokens={{ childrenGap: 16 }} styles={{ root: { paddingTop: 32, paddingLeft: 8, paddingBottom: 8 } }}>
+         <Stack horizontal>
+            <Stack>
+            </Stack>
+            <Stack grow />
+            <Stack>
+               {!!issue.forceClosedByUserInfo && <Stack horizontal tokens={{ childrenGap: 12 }} style={{ paddingRight: 24 }}>                  
+                  <DefaultButton text="Ok" disabled={state.submitting} onClick={() => { onClose(); }} />
+               </Stack>}
+               {!issue.forceClosedByUserInfo && <Stack horizontal tokens={{ childrenGap: 12 }} style={{ paddingRight: 24 }}>
+                  <PrimaryButton text="Save" disabled={state.submitting ?? false} onClick={() => { onSave() }} />
+                  <DefaultButton text="Cancel" disabled={state.submitting} onClick={() => { onClose(); }} />
+               </Stack>}
+            </Stack>
+         </Stack>
+      </Stack>
+   </Modal>;
+
+};
+
 const TestFixModal: React.FC<{ onClose: () => void }> = ({ onClose }) => {
 
    const history = useHistory();
