Horde: Pass a list of tags for a change being preflighted via the UE_HORDE_PREFLIGHT_TAGS environment variable.

diff --git a/Engine/Source/Programs/Horde/Horde.Build.Tests/Fixture.cs b/Engine/Source/Programs/Horde/Horde.Build.Tests/Fixture.cs
--- a/Engine/Source/Programs/Horde/Horde.Build.Tests/Fixture.cs
+++ b/Engine/Source/Programs/Horde/Horde.Build.Tests/Fixture.cs
@@ -112,7 +112,8 @@ namespace Horde.Build.Tests
 				false,
 				null,
 				null,
-				arguments: new List<string>()
+				arguments: new List<string>(),
+				null
 			);
 			Job1 = (await jobService.GetJobAsync(Job1.Id))!;
 
@@ -138,7 +139,8 @@ namespace Horde.Build.Tests
 				false,
 				null,
 				null,
-				arguments: new List<string>()
+				arguments: new List<string>(),
+				null
 			);
 			Job2 = (await jobService.GetJobAsync(Job2.Id))!;
 
diff --git a/Engine/Source/Programs/Horde/Horde.Build.Tests/Fleet/JobQueueStrategyTest.cs b/Engine/Source/Programs/Horde/Horde.Build.Tests/Fleet/JobQueueStrategyTest.cs
--- a/Engine/Source/Programs/Horde/Horde.Build.Tests/Fleet/JobQueueStrategyTest.cs
+++ b/Engine/Source/Programs/Horde/Horde.Build.Tests/Fleet/JobQueueStrategyTest.cs
@@ -152,7 +152,7 @@ namespace Horde.Build.Tests.Fleet
 			IJob job = await JobCollection.AddAsync(ObjectId<IJob>.GenerateNewId(), streamId,
 				new StringId<ITemplateRef>("bogusTemplateRefId"), ContentHash.Empty, graph, "bogusJobName",
 				1000, 1000, null, null, null, null, null, null, null, null, null, false,
-				false, null, null, new List<string> { "-Target=" + nodeNameToExecute });
+				false, null, null, new List<string> { "-Target=" + nodeNameToExecute }, null);
 
 			return job;
 		}
diff --git a/Engine/Source/Programs/Horde/Horde.Build.Tests/JobCollectionTests.cs b/Engine/Source/Programs/Horde/Horde.Build.Tests/JobCollectionTests.cs
--- a/Engine/Source/Programs/Horde/Horde.Build.Tests/JobCollectionTests.cs
+++ b/Engine/Source/Programs/Horde/Horde.Build.Tests/JobCollectionTests.cs
@@ -81,7 +81,7 @@ namespace Horde.Build.Tests
 			arguments.Add("-Target=Publish Client");
 			arguments.Add("-Target=Post-Publish Client");
 
-			IJob job = await JobCollection.AddAsync(JobId.GenerateNewId(), new StreamId("ue4-main"), new TemplateId("test-build"), ContentHash.SHA1("hello"), baseGraph, "Test job", 123, 123, null, null, null, null, null, null, null, null, null, false, false, null, null, arguments);
+			IJob job = await JobCollection.AddAsync(JobId.GenerateNewId(), new StreamId("ue4-main"), new TemplateId("test-build"), ContentHash.SHA1("hello"), baseGraph, "Test job", 123, 123, null, null, null, null, null, null, null, null, null, false, false, null, null, arguments, null);
 
 			job = await StartBatch(job, baseGraph, 0);
 			job = await RunStep(job, baseGraph, 0, 0, JobStepOutcome.Success); // Setup Build
@@ -156,7 +156,7 @@ namespace Horde.Build.Tests
 			arguments.Add("-Target=Step 1");
 			arguments.Add("-Target=Step 3");
 
-			IJob job = await JobCollection.AddAsync(JobId.GenerateNewId(), new StreamId("ue4-main"), new TemplateId("test-build"), ContentHash.SHA1("hello"), baseGraph, "Test job", 123, 123, null, null, null, null, null, null, null, null, null, false, false, null, null, arguments);
+			IJob job = await JobCollection.AddAsync(JobId.GenerateNewId(), new StreamId("ue4-main"), new TemplateId("test-build"), ContentHash.SHA1("hello"), baseGraph, "Test job", 123, 123, null, null, null, null, null, null, null, null, null, false, false, null, null, arguments, null);
 
 			job = await StartBatch(job, baseGraph, 0);
 			job = await RunStep(job, baseGraph, 0, 0, JobStepOutcome.Success); // Setup Build
@@ -212,7 +212,7 @@ namespace Horde.Build.Tests
 			arguments.Add("-Target=Step 1");
 			arguments.Add("-Target=Step 3");
 
-			IJob job = await JobCollection.AddAsync(JobId.GenerateNewId(), new StreamId("ue4-main"), new TemplateId("test-build"), ContentHash.SHA1("hello"), baseGraph, "Test job", 123, 123, null, null, null, null, null, null, null, null, null, false, false, null, null, arguments);
+			IJob job = await JobCollection.AddAsync(JobId.GenerateNewId(), new StreamId("ue4-main"), new TemplateId("test-build"), ContentHash.SHA1("hello"), baseGraph, "Test job", 123, 123, null, null, null, null, null, null, null, null, null, false, false, null, null, arguments, null);
 			Assert.AreEqual(1, job.Batches.Count);
 
 			job = await StartBatch(job, baseGraph, 0);
@@ -263,7 +263,7 @@ namespace Horde.Build.Tests
 			arguments.Add("-Target=Step 1");
 			arguments.Add("-Target=Step 3");
 
-			IJob job = await JobCollection.AddAsync(JobId.GenerateNewId(), new StreamId("ue4-main"), new TemplateId("test-build"), ContentHash.SHA1("hello"), baseGraph, "Test job", 123, 123, null, null, null, null, null, null, null, null, null, false, false, null, null, arguments);
+			IJob job = await JobCollection.AddAsync(JobId.GenerateNewId(), new StreamId("ue4-main"), new TemplateId("test-build"), ContentHash.SHA1("hello"), baseGraph, "Test job", 123, 123, null, null, null, null, null, null, null, null, null, false, false, null, null, arguments, null);
 			Assert.AreEqual(1, job.Batches.Count);
 
 			// First retry
diff --git a/Engine/Source/Programs/Horde/Horde.Build.Tests/JobServiceTests.cs b/Engine/Source/Programs/Horde/Horde.Build.Tests/JobServiceTests.cs
--- a/Engine/Source/Programs/Horde/Horde.Build.Tests/JobServiceTests.cs
+++ b/Engine/Source/Programs/Horde/Horde.Build.Tests/JobServiceTests.cs
@@ -60,7 +60,7 @@ namespace Horde.Build.Tests
 			IStream? stream = await StreamService.GetStreamAsync(streamId);
 			stream = await CreateOrReplaceStreamAsync(new StreamId("ue5-main"), stream, projectId, streamConfig);
 
-			IJob job = await JobService.CreateJobAsync(null, stream!, templateRefId1, template.Id, graph, "Hello", 1234, 1233, 999, null, null, null, null, null, null, null, stream!.Templates[templateRefId1].Config.ChainedJobs, true, true, null, null, new List<string>());
+			IJob job = await JobService.CreateJobAsync(null, stream!, templateRefId1, template.Id, graph, "Hello", 1234, 1233, 999, null, null, null, null, null, null, null, stream!.Templates[templateRefId1].Config.ChainedJobs, true, true, null, null, new List<string>(), null);
 			Assert.AreEqual(1, job.ChainedJobs.Count);
 
 			job = Deref(await JobService.UpdateBatchAsync(job, job.Batches[0].Id, LogId.GenerateNewId(), JobStepBatchState.Running));
@@ -175,7 +175,8 @@ namespace Horde.Build.Tests
 				false,
 				null,
 				null,
-				arguments: new List<string>(arguments)
+				arguments: new List<string>(arguments),
+				null
 			);
 		}
 		
@@ -231,7 +232,7 @@ namespace Horde.Build.Tests
 
 			graph = await GraphCollection.AppendAsync(graph, new List<NewGroup> { groupA, groupB });
 
-			IJob job = await JobService.CreateJobAsync(null, stream!, new TemplateId("temp"), template.Id, graph, "Hello", 1234, 1233, 999, null, null, null, null, null, null, null, null, true, true, null, null, new List<string> { "-Target=Pak" });
+			IJob job = await JobService.CreateJobAsync(null, stream!, new TemplateId("temp"), template.Id, graph, "Hello", 1234, 1233, 999, null, null, null, null, null, null, null, null, true, true, null, null, new List<string> { "-Target=Pak" }, null);
 
 			job = Deref(await JobService.UpdateBatchAsync(job, job.Batches[0].Id, LogId.GenerateNewId(), JobStepBatchState.Running));
 			job = Deref(await JobService.UpdateStepAsync(job, job.Batches[0].Id, job.Batches[0].Steps[0].Id, JobStepState.Running));
diff --git a/Engine/Source/Programs/Horde/Horde.Build.Tests/JobTaskSourceTests.cs b/Engine/Source/Programs/Horde/Horde.Build.Tests/JobTaskSourceTests.cs
--- a/Engine/Source/Programs/Horde/Horde.Build.Tests/JobTaskSourceTests.cs
+++ b/Engine/Source/Programs/Horde/Horde.Build.Tests/JobTaskSourceTests.cs
@@ -135,7 +135,7 @@ namespace Horde.Build.Tests
 			IJob job = await JobCollection.AddAsync(JobId.GenerateNewId(), Stream.Id,
 				fixture.TemplateRefId1, fixture.Template.Id, graph, "Test Paused Step Job",
 				1000, 1000, null, null, null, null, Priority.Highest, null, null, null, null, false,
-				false, null, null, new List<string> { "-Target=" + "Step That Depends on Paused Step;Step That Depends on Update Version Files" });
+				false, null, null, new List<string> { "-Target=" + "Step That Depends on Paused Step;Step That Depends on Update Version Files" }, null);
 
 			// validate
 			Assert.AreEqual(0, JobTaskSource.GetQueueForTesting().Count);
diff --git a/Engine/Source/Programs/Horde/Horde.Build.Tests/SchedulerTests.cs b/Engine/Source/Programs/Horde/Horde.Build.Tests/SchedulerTests.cs
--- a/Engine/Source/Programs/Horde/Horde.Build.Tests/SchedulerTests.cs
+++ b/Engine/Source/Programs/Horde/Horde.Build.Tests/SchedulerTests.cs
@@ -445,7 +445,7 @@ namespace Horde.Build.Tests
 			Assert.AreEqual(0, jobs2.Count);
 
 			// Create a job and fail it
-			IJob job1 = await JobService.CreateJobAsync(null, stream, newTemplateRefId1, _template.Id, graphA, "Hello", 1234, 1233, 999, null, null, null, null, null, null, null, null, true, true, null, null, new List<string> { "-Target=TriggerNext" });
+			IJob job1 = await JobService.CreateJobAsync(null, stream, newTemplateRefId1, _template.Id, graphA, "Hello", 1234, 1233, 999, null, null, null, null, null, null, null, null, true, true, null, null, new List<string> { "-Target=TriggerNext" }, null);
 			SubResourceId batchId1 = job1.Batches[0].Id;
 			SubResourceId stepId1 = job1.Batches[0].Steps[0].Id;
 			job1 = Deref(await JobService.UpdateBatchAsync(job1, batchId1, LogId.GenerateNewId(), JobStepBatchState.Running));
@@ -461,7 +461,7 @@ namespace Horde.Build.Tests
 			Assert.AreEqual(0, jobs3.Count);
 
 			// Create a job and make it succeed
-			IJob job2 = await JobService.CreateJobAsync(null, stream, newTemplateRefId1, _template.Id, graphA, "Hello", 1234, 1233, 999, null, null, null, null, null, null, null, null, true, true, null, null, new List<string> { "-Target=TriggerNext" });
+			IJob job2 = await JobService.CreateJobAsync(null, stream, newTemplateRefId1, _template.Id, graphA, "Hello", 1234, 1233, 999, null, null, null, null, null, null, null, null, true, true, null, null, new List<string> { "-Target=TriggerNext" }, null);
 			SubResourceId batchId2 = job2.Batches[0].Id;
 			SubResourceId stepId2 = job2.Batches[0].Steps[0].Id;
 			job2 = Deref(await JobService.UpdateBatchAsync(job2, batchId2, LogId.GenerateNewId(), JobStepBatchState.Running));
@@ -531,7 +531,7 @@ namespace Horde.Build.Tests
 			{
 				int codeChange = (change < 1233) ? 1230 : 1233;
 
-				IJob job1 = await JobService.CreateJobAsync(null, stream, newTemplateRefId1, _template.Id, graphA, "Hello", change, codeChange, null, null, null, null, null, null, false, null, null, true, true, null, null, new List<string> { "-Target=TriggerNext" });
+				IJob job1 = await JobService.CreateJobAsync(null, stream, newTemplateRefId1, _template.Id, graphA, "Hello", change, codeChange, null, null, null, null, null, null, false, null, null, true, true, null, null, new List<string> { "-Target=TriggerNext" }, null);
 				for (int batchIdx = 0; batchIdx < job1.Batches.Count; batchIdx++)
 				{
 					SubResourceId batchId1 = job1.Batches[batchIdx].Id;
diff --git a/Engine/Source/Programs/Horde/Horde.Build/Jobs/IJob.cs b/Engine/Source/Programs/Horde/Horde.Build/Jobs/IJob.cs
--- a/Engine/Source/Programs/Horde/Horde.Build/Jobs/IJob.cs
+++ b/Engine/Source/Programs/Horde/Horde.Build/Jobs/IJob.cs
@@ -666,6 +666,11 @@ namespace Horde.Build.Jobs
 		/// </summary>
 		public IReadOnlyList<string> Arguments { get; }
 
+		/// <summary>
+		/// Environment variables for the job
+		/// </summary>
+		public IReadOnlyDictionary<string, string> Environment { get; }
+
 		/// <summary>
 		/// Issues associated with this job
 		/// </summary>
diff --git a/Engine/Source/Programs/Horde/Horde.Build/Jobs/IJobCollection.cs b/Engine/Source/Programs/Horde/Horde.Build/Jobs/IJobCollection.cs
--- a/Engine/Source/Programs/Horde/Horde.Build/Jobs/IJobCollection.cs
+++ b/Engine/Source/Programs/Horde/Horde.Build/Jobs/IJobCollection.cs
@@ -57,8 +57,9 @@ namespace Horde.Build.Jobs
 		/// <param name="notificationChannel">Notification channel for this job</param>
 		/// <param name="notificationChannelFilter">Notification channel filter for this job</param>
 		/// <param name="arguments">Arguments for the job</param>
+		/// <param name="environment">Additional environment variables for the job</param>
 		/// <returns>The new job document</returns>
-		Task<IJob> AddAsync(JobId jobId, StreamId streamId, TemplateId templateRefId, ContentHash templateHash, IGraph graph, string name, int change, int codeChange, int? preflightChange, int? clonedPreflightChange, string? preflightDescription, UserId? startedByUserId, Priority? priority, bool? autoSubmit, bool? updateIssues, bool? promoteIssuesByDefault, List<ChainedJobTemplateConfig>? jobTriggers, bool showUgsBadges, bool showUgsAlerts, string? notificationChannel, string? notificationChannelFilter, List<string>? arguments);
+		Task<IJob> AddAsync(JobId jobId, StreamId streamId, TemplateId templateRefId, ContentHash templateHash, IGraph graph, string name, int change, int codeChange, int? preflightChange, int? clonedPreflightChange, string? preflightDescription, UserId? startedByUserId, Priority? priority, bool? autoSubmit, bool? updateIssues, bool? promoteIssuesByDefault, List<ChainedJobTemplateConfig>? jobTriggers, bool showUgsBadges, bool showUgsAlerts, string? notificationChannel, string? notificationChannelFilter, List<string>? arguments, IReadOnlyDictionary<string, string>? environment);
 
 		/// <summary>
 		/// Gets a job with the given unique id
diff --git a/Engine/Source/Programs/Horde/Horde.Build/Jobs/JobCollection.cs b/Engine/Source/Programs/Horde/Horde.Build/Jobs/JobCollection.cs
--- a/Engine/Source/Programs/Horde/Horde.Build/Jobs/JobCollection.cs
+++ b/Engine/Source/Programs/Horde/Horde.Build/Jobs/JobCollection.cs
@@ -24,6 +24,7 @@ using MongoDB.Bson;
 using MongoDB.Bson.IO;
 using MongoDB.Bson.Serialization;
 using MongoDB.Bson.Serialization.Attributes;
+using MongoDB.Bson.Serialization.Options;
 using MongoDB.Driver;
 using OpenTracing;
 using OpenTracing.Util;
@@ -257,6 +258,10 @@ namespace Horde.Build.Jobs
 			public List<JobStepBatchDocument> Batches { get; set; } = new List<JobStepBatchDocument>();
 			public List<Report>? Reports { get; set; }
 			public List<string> Arguments { get; set; } = new List<string>();
+
+			[BsonDictionaryOptions(DictionaryRepresentation.ArrayOfDocuments)]
+			public Dictionary<string, string> Environment { get; set; } = new Dictionary<string, string>();
+
 			public List<int> ReferencedByIssues { get; set; } = new List<int>();
 			public ObjectId? NotificationTriggerId { get; set; }
 			public bool ShowUgsBadges { get; set; }
@@ -288,6 +293,7 @@ namespace Horde.Build.Jobs
 			IReadOnlyList<string> IJob.Arguments => Arguments;
 			IReadOnlyList<int> IJob.Issues => ReferencedByIssues;
 			IReadOnlyDictionary<int, ObjectId> IJob.LabelIdxToTriggerId => _labelNotifications.ToDictionary(x => x._labelIdx, x => x._triggerId);
+			IReadOnlyDictionary<string, string> IJob.Environment => Environment;
 			IReadOnlyList<IChainedJob> IJob.ChainedJobs => ChainedJobs;
 
 			[BsonConstructor]
@@ -297,7 +303,7 @@ namespace Horde.Build.Jobs
 				GraphHash = null!;
 			}
 
-			public JobDocument(JobId id, StreamId streamId, TemplateId templateId, ContentHash templateHash, ContentHash graphHash, string name, int change, int codeChange, int preflightChange, int clonedPreflightChange, string? preflightDescription, UserId? startedByUserId, Priority? priority, bool? autoSubmit, bool? updateIssues, bool? promoteIssuesByDefault, DateTime createTimeUtc, List<ChainedJobDocument> chainedJobs, bool showUgsBadges, bool showUgsAlerts, string? notificationChannel, string? notificationChannelFilter, List<string>? arguments)
+			public JobDocument(JobId id, StreamId streamId, TemplateId templateId, ContentHash templateHash, ContentHash graphHash, string name, int change, int codeChange, int preflightChange, int clonedPreflightChange, string? preflightDescription, UserId? startedByUserId, Priority? priority, bool? autoSubmit, bool? updateIssues, bool? promoteIssuesByDefault, DateTime createTimeUtc, List<ChainedJobDocument> chainedJobs, bool showUgsBadges, bool showUgsAlerts, string? notificationChannel, string? notificationChannelFilter, List<string>? arguments, IReadOnlyDictionary<string, string>? environment)
 			{
 				Id = id;
 				StreamId = streamId;
@@ -322,6 +328,10 @@ namespace Horde.Build.Jobs
 				NotificationChannel = notificationChannel;
 				NotificationChannelFilter = notificationChannelFilter;
 				Arguments = arguments ?? Arguments;
+				if (environment != null)
+				{
+					Environment = new Dictionary<string, string>(environment);
+				}
 				NextSubResourceId = SubResourceId.Random();
 				UpdateTimeUtc = createTimeUtc;
 			}
@@ -396,7 +406,7 @@ namespace Horde.Build.Jobs
 		}
 
 		/// <inheritdoc/>
-		public async Task<IJob> AddAsync(JobId jobId, StreamId streamId, TemplateId templateRefId, ContentHash templateHash, IGraph graph, string name, int change, int codeChange, int? preflightChange, int? clonedPreflightChange, string? preflightDescription, UserId? startedByUserId, Priority? priority, bool? autoSubmit, bool? updateIssues, bool? promoteIssuesByDefault, List<ChainedJobTemplateConfig>? chainedJobs, bool showUgsBadges, bool showUgsAlerts, string? notificationChannel, string? notificationChannelFilter, List<string>? arguments)
+		public async Task<IJob> AddAsync(JobId jobId, StreamId streamId, TemplateId templateRefId, ContentHash templateHash, IGraph graph, string name, int change, int codeChange, int? preflightChange, int? clonedPreflightChange, string? preflightDescription, UserId? startedByUserId, Priority? priority, bool? autoSubmit, bool? updateIssues, bool? promoteIssuesByDefault, List<ChainedJobTemplateConfig>? chainedJobs, bool showUgsBadges, bool showUgsAlerts, string? notificationChannel, string? notificationChannelFilter, List<string>? arguments, IReadOnlyDictionary<string, string>? environment)
 		{
 			List<ChainedJobDocument> jobTriggers = new List<ChainedJobDocument>();
 			if (chainedJobs == null)
@@ -408,7 +418,7 @@ namespace Horde.Build.Jobs
 				jobTriggers = chainedJobs.ConvertAll(x => new ChainedJobDocument(x));
 			}
 
-			JobDocument newJob = new JobDocument(jobId, streamId, templateRefId, templateHash, graph.Id, name, change, codeChange, preflightChange ?? 0, clonedPreflightChange ?? 0, preflightDescription, startedByUserId, priority, autoSubmit, updateIssues, promoteIssuesByDefault, DateTime.UtcNow, jobTriggers, showUgsBadges, showUgsAlerts, notificationChannel, notificationChannelFilter, arguments);
+			JobDocument newJob = new JobDocument(jobId, streamId, templateRefId, templateHash, graph.Id, name, change, codeChange, preflightChange ?? 0, clonedPreflightChange ?? 0, preflightDescription, startedByUserId, priority, autoSubmit, updateIssues, promoteIssuesByDefault, DateTime.UtcNow, jobTriggers, showUgsBadges, showUgsAlerts, notificationChannel, notificationChannelFilter, arguments, environment);
 			CreateBatches(newJob, graph, _logger);
 
 			await _jobs.InsertOneAsync(newJob);
diff --git a/Engine/Source/Programs/Horde/Horde.Build/Jobs/JobService.cs b/Engine/Source/Programs/Horde/Horde.Build/Jobs/JobService.cs
--- a/Engine/Source/Programs/Horde/Horde.Build/Jobs/JobService.cs
+++ b/Engine/Source/Programs/Horde/Horde.Build/Jobs/JobService.cs
@@ -189,8 +189,9 @@ namespace Horde.Build.Jobs
 		/// <param name="notificationChannel">Notification Channel for this job</param>
 		/// <param name="notificationChannelFilter">Notification Channel filter for this job</param>
 		/// <param name="arguments">Arguments for the job</param>
+		/// <param name="environment">Additional environment variables to set, in the form X = Y</param>
 		/// <returns>Unique id representing the job</returns>
-		public async Task<IJob> CreateJobAsync(JobId? jobId, IStream stream, TemplateId templateRefId, ContentHash templateHash, IGraph graph, string name, int change, int codeChange, int? preflightChange, int? clonedPreflightChange, string? preflightDescription, UserId? startedByUserId, Priority? priority, bool? autoSubmit, bool? updateIssues, bool? promoteIssuesByDefault, List<ChainedJobTemplateConfig>? jobTriggers, bool showUgsBadges, bool showUgsAlerts, string? notificationChannel, string? notificationChannelFilter, IReadOnlyList<string> arguments)
+		public async Task<IJob> CreateJobAsync(JobId? jobId, IStream stream, TemplateId templateRefId, ContentHash templateHash, IGraph graph, string name, int change, int codeChange, int? preflightChange, int? clonedPreflightChange, string? preflightDescription, UserId? startedByUserId, Priority? priority, bool? autoSubmit, bool? updateIssues, bool? promoteIssuesByDefault, List<ChainedJobTemplateConfig>? jobTriggers, bool showUgsBadges, bool showUgsAlerts, string? notificationChannel, string? notificationChannelFilter, IReadOnlyList<string> arguments, IReadOnlyDictionary<string, string>? environment)
 		{
 			using IScope traceScope = GlobalTracer.Instance.BuildSpan("JobService.CreateJobAsync").StartActive();
 			traceScope.Span.SetTag("JobId", jobId);
@@ -254,7 +255,7 @@ namespace Horde.Build.Jobs
 
 			name = StringUtils.ExpandProperties(name, properties);
 
-			IJob newJob = await _jobs.AddAsync(jobIdValue, stream.Id, templateRefId, templateHash, graph, name, change, codeChange, preflightChange, clonedPreflightChange, preflightDescription, startedByUserId, priority, autoSubmit, updateIssues, promoteIssuesByDefault, jobTriggers, showUgsBadges, showUgsAlerts, notificationChannel, notificationChannelFilter, expandedArguments);
+			IJob newJob = await _jobs.AddAsync(jobIdValue, stream.Id, templateRefId, templateHash, graph, name, change, codeChange, preflightChange, clonedPreflightChange, preflightDescription, startedByUserId, priority, autoSubmit, updateIssues, promoteIssuesByDefault, jobTriggers, showUgsBadges, showUgsAlerts, notificationChannel, notificationChannelFilter, expandedArguments, environment);
 			_jobTaskSource.UpdateQueuedJob(newJob, graph);
 
 			await _jobTaskSource.UpdateUgsBadges(newJob, graph, new List<(LabelState, LabelOutcome)>());
@@ -1316,7 +1317,7 @@ namespace Horde.Build.Jobs
 					IGraph triggerGraph = await _graphs.AddAsync(template);
 					_logger.LogInformation("Creating downstream job {ChainedJobId} from job {JobId}", chainedJobId, newJob.Id);
 
-					await CreateJobAsync(chainedJobId, stream, jobTrigger.TemplateRefId, templateRef.Hash, triggerGraph, templateRef.Config.Name, newJob.Change, newJob.CodeChange, newJob.PreflightChange, newJob.ClonedPreflightChange, newJob.PreflightDescription, newJob.StartedByUserId, template.Priority, null, newJob.UpdateIssues, newJob.PromoteIssuesByDefault, templateRef.Config.ChainedJobs, false, false, templateRef.Config.NotificationChannel, templateRef.Config.NotificationChannelFilter, template.Arguments);
+					await CreateJobAsync(chainedJobId, stream, jobTrigger.TemplateRefId, templateRef.Hash, triggerGraph, templateRef.Config.Name, newJob.Change, newJob.CodeChange, newJob.PreflightChange, newJob.ClonedPreflightChange, newJob.PreflightDescription, newJob.StartedByUserId, template.Priority, null, newJob.UpdateIssues, newJob.PromoteIssuesByDefault, templateRef.Config.ChainedJobs, false, false, templateRef.Config.NotificationChannel, templateRef.Config.NotificationChannelFilter, template.Arguments, null);
 					return newJob;
 				}
 
diff --git a/Engine/Source/Programs/Horde/Horde.Build/Jobs/JobsController.cs b/Engine/Source/Programs/Horde/Horde.Build/Jobs/JobsController.cs
--- a/Engine/Source/Programs/Horde/Horde.Build/Jobs/JobsController.cs
+++ b/Engine/Source/Programs/Horde/Horde.Build/Jobs/JobsController.cs
@@ -129,6 +129,9 @@ namespace Horde.Build.Jobs
 				}
 			}
 
+			// Environment variables for the job
+			Dictionary<string, string> environment = new Dictionary<string, string>();
+
 			// Check the preflight change is valid
 			ShelfInfo? shelfInfo = null;
 			if (create.PreflightChange != null)
@@ -149,6 +152,12 @@ namespace Horde.Build.Jobs
 					default:
 						return BadRequest(KnownLogEvents.Horde_InvalidPreflight, "CL {Change} cannot be preflighted ({Result})", create.PreflightChange, result);
 				}
+
+				if (shelfInfo!.Tags != null)
+				{
+					string tagList = String.Join(";", shelfInfo.Tags.Select(x => x.ToString()));
+					environment.Add("UE_HORDE_PREFLIGHT_TAGS", tagList);
+				}
 			}
 
 			// Get the priority of the new job
@@ -181,7 +190,7 @@ namespace Horde.Build.Jobs
 			}
 
 			// Create the job
-			IJob job = await _jobService.CreateJobAsync(null, stream, templateRefId, template.Id, graph, name, change, codeChange, create.PreflightChange, null, shelfInfo?.Description, User.GetUserId(), priority, create.AutoSubmit, updateIssues, false, templateRef.Config.ChainedJobs, templateRef.Config.ShowUgsBadges, templateRef.Config.ShowUgsAlerts, templateRef.Config.NotificationChannel, templateRef.Config.NotificationChannelFilter, arguments);
+			IJob job = await _jobService.CreateJobAsync(null, stream, templateRefId, template.Id, graph, name, change, codeChange, create.PreflightChange, null, shelfInfo?.Description, User.GetUserId(), priority, create.AutoSubmit, updateIssues, false, templateRef.Config.ChainedJobs, templateRef.Config.ShowUgsBadges, templateRef.Config.ShowUgsAlerts, templateRef.Config.NotificationChannel, templateRef.Config.NotificationChannelFilter, arguments, environment);
 			await UpdateNotificationsAsync(job.Id, new UpdateNotificationsRequest { Slack = true });
 			return new CreateJobResponse(job.Id.ToString());
 		}
diff --git a/Engine/Source/Programs/Horde/Horde.Build/Jobs/Schedules/ScheduleService.cs b/Engine/Source/Programs/Horde/Horde.Build/Jobs/Schedules/ScheduleService.cs
--- a/Engine/Source/Programs/Horde/Horde.Build/Jobs/Schedules/ScheduleService.cs
+++ b/Engine/Source/Programs/Horde/Horde.Build/Jobs/Schedules/ScheduleService.cs
@@ -462,7 +462,7 @@ namespace Horde.Build.Jobs.Schedules
 			{
 				cancellationToken.ThrowIfCancellationRequested();
 				List<string> defaultArguments = template.GetDefaultArguments();
-				IJob newJob = await _jobService.CreateJobAsync(null, stream, templateId, template.Id, graph, template.Name, change, codeChange, null, null, null, null, template.Priority, null, null, template.PromoteIssuesByDefault, templateRef.Config.ChainedJobs, templateRef.Config.ShowUgsBadges, templateRef.Config.ShowUgsAlerts, templateRef.Config.NotificationChannel, templateRef.Config.NotificationChannelFilter, defaultArguments);
+				IJob newJob = await _jobService.CreateJobAsync(null, stream, templateId, template.Id, graph, template.Name, change, codeChange, null, null, null, null, template.Priority, null, null, template.PromoteIssuesByDefault, templateRef.Config.ChainedJobs, templateRef.Config.ShowUgsBadges, templateRef.Config.ShowUgsAlerts, templateRef.Config.NotificationChannel, templateRef.Config.NotificationChannelFilter, defaultArguments, null);
 				_logger.LogInformation("Started new job for {StreamId} template {TemplateId} at CL {Change} (Code CL {CodeChange}): {JobId}", stream.Id, templateId, change, codeChange, newJob.Id);
 				await _streamService.UpdateScheduleTriggerAsync(stream, templateId, utcNow, change, new List<JobId> { newJob.Id }, new List<JobId>());
 			}
diff --git a/Engine/Source/Programs/Horde/Horde.Build/Server/RpcService.cs b/Engine/Source/Programs/Horde/Horde.Build/Server/RpcService.cs
--- a/Engine/Source/Programs/Horde/Horde.Build/Server/RpcService.cs
+++ b/Engine/Source/Programs/Horde/Horde.Build/Server/RpcService.cs
@@ -635,6 +635,11 @@ namespace Horde.Build.Server
 				response.EnvVars.Add("UE_HORDE_TEMPLATENAME", templateName);
 				response.EnvVars.Add("UE_HORDE_STEPNAME", node.Name);
 
+				if (job.Environment != null)
+				{
+					response.EnvVars.Add(job.Environment);
+				}
+
 				IJobStepRef? lastStep = await _jobStepRefCollection.GetPrevStepForNodeAsync(job.StreamId, job.TemplateId, node.Name, job.Change);
 				if (lastStep != null)
 				{
