Horde: Use a class to pass options for creating new jobs.

diff --git a/Engine/Source/Programs/Horde/Horde.Build.Tests/AgentServiceTest.cs b/Engine/Source/Programs/Horde/Horde.Build.Tests/AgentServiceTest.cs
--- a/Engine/Source/Programs/Horde/Horde.Build.Tests/AgentServiceTest.cs
+++ b/Engine/Source/Programs/Horde/Horde.Build.Tests/AgentServiceTest.cs
@@ -2,8 +2,10 @@
 
 using System;
 using System.Collections.Generic;
+using System.Linq;
 using System.Security.Claims;
 using System.Threading.Tasks;
+using EpicGames.Core;
 using Horde.Build.Agents;
 using Horde.Build.Jobs;
 using Horde.Build.Utilities;
@@ -25,9 +27,13 @@ public class AgentServiceTest : TestSetup
 		await CreateFixtureAsync();
 
 		ActionResult<List<object>> res = await JobsController.FindJobsAsync();
-		Assert.AreEqual(2, res.Value!.Count);
-		Assert.AreEqual("hello2", (res.Value[0] as GetJobResponse)!.Name);
-		Assert.AreEqual("hello1", (res.Value[1] as GetJobResponse)!.Name);
+
+		List<GetJobResponse> responses = res.Value!.ConvertAll(x => (GetJobResponse)x);
+		responses.SortBy(x => x.Change);
+
+		Assert.AreEqual(2, responses.Count);
+		Assert.AreEqual("hello1", responses[0].Name);
+		Assert.AreEqual("hello2", responses[1].Name);
 
 		res = await JobsController.FindJobsAsync(includePreflight: false);
 		Assert.AreEqual(1, res.Value!.Count);
diff --git a/Engine/Source/Programs/Horde/Horde.Build.Tests/Fixture.cs b/Engine/Source/Programs/Horde/Horde.Build.Tests/Fixture.cs
--- a/Engine/Source/Programs/Horde/Horde.Build.Tests/Fixture.cs
+++ b/Engine/Source/Programs/Horde/Horde.Build.Tests/Fixture.cs
@@ -99,21 +99,7 @@ namespace Horde.Build.Tests
 				name: "hello1",
 				change: 1000001,
 				codeChange: 1000002,
-				preflightChange: 1001,
-				clonedPreflightChange: null,
-				preflightDescription: null,
-				startedByUserId: null,
-				priority: Priority.Normal,
-				null,
-				null,
-				null,
-				null,
-				false,
-				false,
-				null,
-				null,
-				arguments: new List<string>(),
-				null
+				new CreateJobOptions { PreflightChange = 1001 }
 			);
 			Job1 = (await jobService.GetJobAsync(Job1.Id))!;
 
@@ -126,21 +112,7 @@ namespace Horde.Build.Tests
 				name: "hello2",
 				change: 2000001,
 				codeChange: 2000002,
-				preflightChange: null,
-				clonedPreflightChange: null,
-				preflightDescription: null,
-				startedByUserId: null,
-				priority: Priority.Normal,
-				null,
-				null,
-				null,
-				null,
-				false,
-				false,
-				null,
-				null,
-				arguments: new List<string>(),
-				null
+				new CreateJobOptions()
 			);
 			Job2 = (await jobService.GetJobAsync(Job2.Id))!;
 
diff --git a/Engine/Source/Programs/Horde/Horde.Build.Tests/Fleet/JobQueueStrategyTest.cs b/Engine/Source/Programs/Horde/Horde.Build.Tests/Fleet/JobQueueStrategyTest.cs
--- a/Engine/Source/Programs/Horde/Horde.Build.Tests/Fleet/JobQueueStrategyTest.cs
+++ b/Engine/Source/Programs/Horde/Horde.Build.Tests/Fleet/JobQueueStrategyTest.cs
@@ -149,10 +149,12 @@ namespace Horde.Build.Tests.Fleet
 		
 		private async Task<IJob> AddPlaceholderJob(IGraph graph, StreamId streamId, string nodeNameToExecute)
 		{
+			CreateJobOptions options = new CreateJobOptions();
+			options.Arguments.Add($"-Target={nodeNameToExecute}");
+
 			IJob job = await JobCollection.AddAsync(ObjectId<IJob>.GenerateNewId(), streamId,
 				new StringId<ITemplateRef>("bogusTemplateRefId"), ContentHash.Empty, graph, "bogusJobName",
-				1000, 1000, null, null, null, null, null, null, null, null, null, false,
-				false, null, null, new List<string> { "-Target=" + nodeNameToExecute }, null);
+				1000, 1000, options);
 
 			return job;
 		}
diff --git a/Engine/Source/Programs/Horde/Horde.Build.Tests/JobCollectionTests.cs b/Engine/Source/Programs/Horde/Horde.Build.Tests/JobCollectionTests.cs
--- a/Engine/Source/Programs/Horde/Horde.Build.Tests/JobCollectionTests.cs
+++ b/Engine/Source/Programs/Horde/Horde.Build.Tests/JobCollectionTests.cs
@@ -77,11 +77,11 @@ namespace Horde.Build.Tests
 
 			IGraph baseGraph = await GraphCollection.AddAsync(templateMock.Object);
 
-			List<string> arguments = new List<string>();
-			arguments.Add("-Target=Publish Client");
-			arguments.Add("-Target=Post-Publish Client");
+			CreateJobOptions options = new CreateJobOptions();
+			options.Arguments.Add("-Target=Publish Client");
+			options.Arguments.Add("-Target=Post-Publish Client");
 
-			IJob job = await JobCollection.AddAsync(JobId.GenerateNewId(), new StreamId("ue4-main"), new TemplateId("test-build"), ContentHash.SHA1("hello"), baseGraph, "Test job", 123, 123, null, null, null, null, null, null, null, null, null, false, false, null, null, arguments, null);
+			IJob job = await JobCollection.AddAsync(JobId.GenerateNewId(), new StreamId("ue4-main"), new TemplateId("test-build"), ContentHash.SHA1("hello"), baseGraph, "Test job", 123, 123, options);
 
 			job = await StartBatch(job, baseGraph, 0);
 			job = await RunStep(job, baseGraph, 0, 0, JobStepOutcome.Success); // Setup Build
@@ -152,11 +152,11 @@ namespace Horde.Build.Tests
 
 			IGraph baseGraph = await GraphCollection.AddAsync(templateMock.Object);
 
-			List<string> arguments = new List<string>();
-			arguments.Add("-Target=Step 1");
-			arguments.Add("-Target=Step 3");
+			CreateJobOptions options = new CreateJobOptions();
+			options.Arguments.Add("-Target=Step 1");
+			options.Arguments.Add("-Target=Step 3");
 
-			IJob job = await JobCollection.AddAsync(JobId.GenerateNewId(), new StreamId("ue4-main"), new TemplateId("test-build"), ContentHash.SHA1("hello"), baseGraph, "Test job", 123, 123, null, null, null, null, null, null, null, null, null, false, false, null, null, arguments, null);
+			IJob job = await JobCollection.AddAsync(JobId.GenerateNewId(), new StreamId("ue4-main"), new TemplateId("test-build"), ContentHash.SHA1("hello"), baseGraph, "Test job", 123, 123, options);
 
 			job = await StartBatch(job, baseGraph, 0);
 			job = await RunStep(job, baseGraph, 0, 0, JobStepOutcome.Success); // Setup Build
@@ -208,11 +208,11 @@ namespace Horde.Build.Tests
 
 			IGraph baseGraph = await GraphCollection.AddAsync(templateMock.Object);
 
-			List<string> arguments = new List<string>();
-			arguments.Add("-Target=Step 1");
-			arguments.Add("-Target=Step 3");
+			CreateJobOptions options = new CreateJobOptions();
+			options.Arguments.Add("-Target=Step 1");
+			options.Arguments.Add("-Target=Step 3");
 
-			IJob job = await JobCollection.AddAsync(JobId.GenerateNewId(), new StreamId("ue4-main"), new TemplateId("test-build"), ContentHash.SHA1("hello"), baseGraph, "Test job", 123, 123, null, null, null, null, null, null, null, null, null, false, false, null, null, arguments, null);
+			IJob job = await JobCollection.AddAsync(JobId.GenerateNewId(), new StreamId("ue4-main"), new TemplateId("test-build"), ContentHash.SHA1("hello"), baseGraph, "Test job", 123, 123, options);
 			Assert.AreEqual(1, job.Batches.Count);
 
 			job = await StartBatch(job, baseGraph, 0);
@@ -259,11 +259,11 @@ namespace Horde.Build.Tests
 
 			IGraph baseGraph = await GraphCollection.AddAsync(templateMock.Object);
 
-			List<string> arguments = new List<string>();
-			arguments.Add("-Target=Step 1");
-			arguments.Add("-Target=Step 3");
+			CreateJobOptions options = new CreateJobOptions();
+			options.Arguments.Add("-Target=Step 1");
+			options.Arguments.Add("-Target=Step 3");
 
-			IJob job = await JobCollection.AddAsync(JobId.GenerateNewId(), new StreamId("ue4-main"), new TemplateId("test-build"), ContentHash.SHA1("hello"), baseGraph, "Test job", 123, 123, null, null, null, null, null, null, null, null, null, false, false, null, null, arguments, null);
+			IJob job = await JobCollection.AddAsync(JobId.GenerateNewId(), new StreamId("ue4-main"), new TemplateId("test-build"), ContentHash.SHA1("hello"), baseGraph, "Test job", 123, 123, options);
 			Assert.AreEqual(1, job.Batches.Count);
 
 			// First retry
diff --git a/Engine/Source/Programs/Horde/Horde.Build.Tests/JobServiceTests.cs b/Engine/Source/Programs/Horde/Horde.Build.Tests/JobServiceTests.cs
--- a/Engine/Source/Programs/Horde/Horde.Build.Tests/JobServiceTests.cs
+++ b/Engine/Source/Programs/Horde/Horde.Build.Tests/JobServiceTests.cs
@@ -60,7 +60,11 @@ namespace Horde.Build.Tests
 			IStream? stream = await StreamService.GetStreamAsync(streamId);
 			stream = await CreateOrReplaceStreamAsync(new StreamId("ue5-main"), stream, projectId, streamConfig);
 
-			IJob job = await JobService.CreateJobAsync(null, stream!, templateRefId1, template.Id, graph, "Hello", 1234, 1233, 999, null, null, null, null, null, null, null, stream!.Templates[templateRefId1].Config.ChainedJobs, true, true, null, null, new List<string>(), null);
+			CreateJobOptions options = new CreateJobOptions();
+			options.PreflightChange = 999;
+			options.JobTriggers.AddRange(stream!.Templates[templateRefId1].Config.ChainedJobs!);
+
+			IJob job = await JobService.CreateJobAsync(null, stream!, templateRefId1, template.Id, graph, "Hello", 1234, 1233, options);
 			Assert.AreEqual(1, job.ChainedJobs.Count);
 
 			job = Deref(await JobService.UpdateBatchAsync(job, job.Batches[0].Id, LogId.GenerateNewId(), JobStepBatchState.Running));
@@ -153,6 +157,12 @@ namespace Horde.Build.Tests
 		private async Task<IJob> CreatePreflightJob(Fixture fixture, string templateRefId, string templateHash, string startedByUserName, int preflightChange, string[] arguments)
 		{
 			IUser user = await UserCollection.FindOrAddUserByLoginAsync(startedByUserName);
+
+			CreateJobOptions options = new CreateJobOptions();
+			options.PreflightChange = preflightChange;
+			options.StartedByUserId = user.Id;
+			options.Arguments.AddRange(arguments);
+
 			return await JobService.CreateJobAsync(
 				jobId: JobId.GenerateNewId(),
 				stream: fixture!.Stream!,
@@ -162,21 +172,7 @@ namespace Horde.Build.Tests
 				name: "hello1",
 				change: 1000001,
 				codeChange: 1000002,
-				preflightChange: preflightChange,
-				clonedPreflightChange: null,
-				preflightDescription: null,
-				startedByUserId: user.Id,
-				priority: Priority.Normal,
-				null,
-				null,
-				null,
-				null,
-				false,
-				false,
-				null,
-				null,
-				arguments: new List<string>(arguments),
-				null
+				options
 			);
 		}
 		
@@ -232,7 +228,11 @@ namespace Horde.Build.Tests
 
 			graph = await GraphCollection.AppendAsync(graph, new List<NewGroup> { groupA, groupB });
 
-			IJob job = await JobService.CreateJobAsync(null, stream!, new TemplateId("temp"), template.Id, graph, "Hello", 1234, 1233, 999, null, null, null, null, null, null, null, null, true, true, null, null, new List<string> { "-Target=Pak" }, null);
+			CreateJobOptions options = new CreateJobOptions();
+			options.PreflightChange = 999;
+			options.Arguments.Add("-Target=Pak");
+
+			IJob job = await JobService.CreateJobAsync(null, stream!, new TemplateId("temp"), template.Id, graph, "Hello", 1234, 1233, options);
 
 			job = Deref(await JobService.UpdateBatchAsync(job, job.Batches[0].Id, LogId.GenerateNewId(), JobStepBatchState.Running));
 			job = Deref(await JobService.UpdateStepAsync(job, job.Batches[0].Id, job.Batches[0].Steps[0].Id, JobStepState.Running));
diff --git a/Engine/Source/Programs/Horde/Horde.Build.Tests/JobTaskSourceTests.cs b/Engine/Source/Programs/Horde/Horde.Build.Tests/JobTaskSourceTests.cs
--- a/Engine/Source/Programs/Horde/Horde.Build.Tests/JobTaskSourceTests.cs
+++ b/Engine/Source/Programs/Horde/Horde.Build.Tests/JobTaskSourceTests.cs
@@ -132,10 +132,12 @@ namespace Horde.Build.Tests
 			}
 
 			// create a new job
+			CreateJobOptions options = new CreateJobOptions();
+			options.Arguments.Add("-Target=Step That Depends on Paused Step;Step That Depends on Update Version Files");
+
 			IJob job = await JobCollection.AddAsync(JobId.GenerateNewId(), Stream.Id,
 				fixture.TemplateRefId1, fixture.Template.Id, graph, "Test Paused Step Job",
-				1000, 1000, null, null, null, null, Priority.Highest, null, null, null, null, false,
-				false, null, null, new List<string> { "-Target=" + "Step That Depends on Paused Step;Step That Depends on Update Version Files" }, null);
+				1000, 1000, options);
 
 			// validate
 			Assert.AreEqual(0, JobTaskSource.GetQueueForTesting().Count);
diff --git a/Engine/Source/Programs/Horde/Horde.Build.Tests/JobsControllerDbTest.cs b/Engine/Source/Programs/Horde/Horde.Build.Tests/JobsControllerDbTest.cs
--- a/Engine/Source/Programs/Horde/Horde.Build.Tests/JobsControllerDbTest.cs
+++ b/Engine/Source/Programs/Horde/Horde.Build.Tests/JobsControllerDbTest.cs
@@ -2,6 +2,7 @@
 
 using System.Collections.Generic;
 using System.Threading.Tasks;
+using EpicGames.Core;
 using Horde.Build.Jobs;
 using Horde.Build.Jobs.Graphs;
 using Microsoft.AspNetCore.Mvc;
@@ -23,10 +24,14 @@ namespace Horde.Build.Tests
 			await CreateFixtureAsync();
 
 			ActionResult<List<object>> res = await JobsController.FindJobsAsync();
-	        Assert.AreEqual(2, res.Value!.Count);
-	        Assert.AreEqual("hello2", (res.Value[0] as GetJobResponse)!.Name);
-	        Assert.AreEqual("hello1", (res.Value[1] as GetJobResponse)!.Name);
-	        
+
+			List<GetJobResponse> responses = res.Value!.ConvertAll(x => (GetJobResponse)x);
+			responses.SortBy(x => x.Change);
+
+			Assert.AreEqual(2, responses.Count);
+			Assert.AreEqual("hello1", responses[0].Name);
+			Assert.AreEqual("hello2", responses[1].Name);
+
 	        res = await JobsController.FindJobsAsync(includePreflight: false);
 	        Assert.AreEqual(1, res.Value!.Count);
 	        Assert.AreEqual("hello2", (res.Value[0] as GetJobResponse)!.Name);
diff --git a/Engine/Source/Programs/Horde/Horde.Build.Tests/SchedulerTests.cs b/Engine/Source/Programs/Horde/Horde.Build.Tests/SchedulerTests.cs
--- a/Engine/Source/Programs/Horde/Horde.Build.Tests/SchedulerTests.cs
+++ b/Engine/Source/Programs/Horde/Horde.Build.Tests/SchedulerTests.cs
@@ -445,7 +445,11 @@ namespace Horde.Build.Tests
 			Assert.AreEqual(0, jobs2.Count);
 
 			// Create a job and fail it
-			IJob job1 = await JobService.CreateJobAsync(null, stream, newTemplateRefId1, _template.Id, graphA, "Hello", 1234, 1233, 999, null, null, null, null, null, null, null, null, true, true, null, null, new List<string> { "-Target=TriggerNext" }, null);
+			CreateJobOptions options1 = new CreateJobOptions();
+			options1.PreflightChange = 999;
+			options1.Arguments.Add("-Target=TriggerNext");
+
+			IJob job1 = await JobService.CreateJobAsync(null, stream, newTemplateRefId1, _template.Id, graphA, "Hello", 1234, 1233, options1);
 			SubResourceId batchId1 = job1.Batches[0].Id;
 			SubResourceId stepId1 = job1.Batches[0].Steps[0].Id;
 			job1 = Deref(await JobService.UpdateBatchAsync(job1, batchId1, LogId.GenerateNewId(), JobStepBatchState.Running));
@@ -461,7 +465,11 @@ namespace Horde.Build.Tests
 			Assert.AreEqual(0, jobs3.Count);
 
 			// Create a job and make it succeed
-			IJob job2 = await JobService.CreateJobAsync(null, stream, newTemplateRefId1, _template.Id, graphA, "Hello", 1234, 1233, 999, null, null, null, null, null, null, null, null, true, true, null, null, new List<string> { "-Target=TriggerNext" }, null);
+			CreateJobOptions options2 = new CreateJobOptions();
+			options2.PreflightChange = 999;
+			options2.Arguments.Add("-Target=TriggerNext");
+
+			IJob job2 = await JobService.CreateJobAsync(null, stream, newTemplateRefId1, _template.Id, graphA, "Hello", 1234, 1233, options2);
 			SubResourceId batchId2 = job2.Batches[0].Id;
 			SubResourceId stepId2 = job2.Batches[0].Steps[0].Id;
 			job2 = Deref(await JobService.UpdateBatchAsync(job2, batchId2, LogId.GenerateNewId(), JobStepBatchState.Running));
@@ -531,7 +539,10 @@ namespace Horde.Build.Tests
 			{
 				int codeChange = (change < 1233) ? 1230 : 1233;
 
-				IJob job1 = await JobService.CreateJobAsync(null, stream, newTemplateRefId1, _template.Id, graphA, "Hello", change, codeChange, null, null, null, null, null, null, false, null, null, true, true, null, null, new List<string> { "-Target=TriggerNext" }, null);
+				CreateJobOptions options1 = new CreateJobOptions();
+				options1.Arguments.Add("-Target=TriggerNext");
+
+				IJob job1 = await JobService.CreateJobAsync(null, stream, newTemplateRefId1, _template.Id, graphA, "Hello", change, codeChange, options1);
 				for (int batchIdx = 0; batchIdx < job1.Batches.Count; batchIdx++)
 				{
 					SubResourceId batchId1 = job1.Batches[batchIdx].Id;
diff --git a/Engine/Source/Programs/Horde/Horde.Build/Jobs/IJobCollection.cs b/Engine/Source/Programs/Horde/Horde.Build/Jobs/IJobCollection.cs
--- a/Engine/Source/Programs/Horde/Horde.Build/Jobs/IJobCollection.cs
+++ b/Engine/Source/Programs/Horde/Horde.Build/Jobs/IJobCollection.cs
@@ -27,6 +27,110 @@ namespace Horde.Build.Jobs
 	using TemplateId = StringId<ITemplateRef>;
 	using UserId = ObjectId<IUser>;
 
+	/// <summary>
+	/// Options for creating a new job
+	/// </summary>
+	public class CreateJobOptions
+	{
+		/// <summary>
+		/// Optional changelist to preflight
+		/// </summary>
+		public int? PreflightChange { get; set; }
+
+		/// <summary>
+		/// Optional cloned preflight changelist
+		/// </summary>
+		public int? ClonedPreflightChange { get; set; }
+
+		/// <summary>
+		/// Description for the change being preflighted
+		/// </summary>
+		public string? PreflightDescription { get; set; }
+
+		/// <summary>
+		/// User that started the job
+		/// </summary>
+		public UserId? StartedByUserId { get; set; }
+
+		/// <summary>
+		/// Priority of the job
+		/// </summary>
+		public Priority? Priority { get; set; }
+
+		/// <summary>
+		/// Whether to automatically submit the preflighted change on completion
+		/// </summary>
+		public bool? AutoSubmit { get; set; }
+
+		/// <summary>
+		/// Whether to update issues based on the outcome of this job
+		/// </summary>
+		public bool? UpdateIssues { get; set; }
+
+		/// <summary>
+		/// Whether to promote issues based on the outcome of this job by default
+		/// </summary>
+		public bool? PromoteIssuesByDefault { get; set; }
+
+		/// <summary>
+		/// List of downstream job triggers
+		/// </summary>
+		public List<ChainedJobTemplateConfig> JobTriggers { get; } = new List<ChainedJobTemplateConfig>();
+
+		/// <summary>
+		/// Whether to show badges in UGS for this job
+		/// </summary>
+		public bool ShowUgsBadges { get; set; }
+
+		/// <summary>
+		/// Whether to show alerts in UGS for this job
+		/// </summary>
+		public bool ShowUgsAlerts { get; set; }
+
+		/// <summary>
+		/// Notification channel for this job
+		/// </summary>
+		public string? NotificationChannel { get; set; }
+
+		/// <summary>
+		/// Notification channel filter for this job
+		/// </summary>
+		public string? NotificationChannelFilter { get; set; }
+
+		/// <summary>
+		/// Arguments for the job
+		/// </summary>
+		public List<string> Arguments { get; } = new List<string>();
+
+		/// <summary>
+		/// Environment variables for the job
+		/// </summary>
+		public Dictionary<string, string> Environment { get; } = new Dictionary<string, string>(StringComparer.OrdinalIgnoreCase);
+
+		/// <summary>
+		/// Default constructor
+		/// </summary>
+		public CreateJobOptions()
+		{
+		}
+
+		/// <summary>
+		/// Constructor
+		/// </summary>
+		public CreateJobOptions(TemplateRefConfig templateRefConfig)
+		{
+			PromoteIssuesByDefault = templateRefConfig.PromoteIssuesByDefault;
+			if (templateRefConfig.ChainedJobs != null)
+			{
+				JobTriggers.AddRange(templateRefConfig.ChainedJobs);
+			}
+			ShowUgsBadges = templateRefConfig.ShowUgsBadges;
+			ShowUgsAlerts = templateRefConfig.ShowUgsAlerts;
+			NotificationChannel = templateRefConfig.NotificationChannel;
+			NotificationChannelFilter = templateRefConfig.NotificationChannelFilter;
+		}
+	}
+
 	/// <summary>
 	/// Interface for a collection of job documents
 	/// </summary>
@@ -43,23 +147,9 @@ namespace Horde.Build.Jobs
 		/// <param name="name">Name of the job</param>
 		/// <param name="change">The change to build</param>
 		/// <param name="codeChange">The corresponding code changelist number</param>
-		/// <param name="preflightChange">Optional changelist to preflight</param>
-		/// <param name="clonedPreflightChange">Optional cloned preflight changelist</param>
-		/// <param name="preflightDescription">Description for the change being preflighted</param>
-		/// <param name="startedByUserId">User that started the job</param>
-		/// <param name="priority">Priority of the job</param>
-		/// <param name="autoSubmit">Whether to automatically submit the preflighted change on completion</param>
-		/// <param name="updateIssues">Whether to update issues based on the outcome of this job</param>
-		/// <param name="promoteIssuesByDefault">Whether to promote issues based on the outcome of this job by default</param>
-		/// <param name="jobTriggers">List of downstream job triggers</param>
-		/// <param name="showUgsBadges">Whether to show badges in UGS for this job</param>
-		/// <param name="showUgsAlerts">Whether to show alerts in UGS for this job</param>
-		/// <param name="notificationChannel">Notification channel for this job</param>
-		/// <param name="notificationChannelFilter">Notification channel filter for this job</param>
-		/// <param name="arguments">Arguments for the job</param>
-		/// <param name="environment">Additional environment variables for the job</param>
+		/// <param name="options">Additional options for the new job</param>
 		/// <returns>The new job document</returns>
-		Task<IJob> AddAsync(JobId jobId, StreamId streamId, TemplateId templateRefId, ContentHash templateHash, IGraph graph, string name, int change, int codeChange, int? preflightChange, int? clonedPreflightChange, string? preflightDescription, UserId? startedByUserId, Priority? priority, bool? autoSubmit, bool? updateIssues, bool? promoteIssuesByDefault, List<ChainedJobTemplateConfig>? jobTriggers, bool showUgsBadges, bool showUgsAlerts, string? notificationChannel, string? notificationChannelFilter, List<string>? arguments, IReadOnlyDictionary<string, string>? environment);
+		Task<IJob> AddAsync(JobId jobId, StreamId streamId, TemplateId templateRefId, ContentHash templateHash, IGraph graph, string name, int change, int codeChange, CreateJobOptions options);
 
 		/// <summary>
 		/// Gets a job with the given unique id
diff --git a/Engine/Source/Programs/Horde/Horde.Build/Jobs/JobCollection.cs b/Engine/Source/Programs/Horde/Horde.Build/Jobs/JobCollection.cs
--- a/Engine/Source/Programs/Horde/Horde.Build/Jobs/JobCollection.cs
+++ b/Engine/Source/Programs/Horde/Horde.Build/Jobs/JobCollection.cs
@@ -303,7 +303,7 @@ namespace Horde.Build.Jobs
 				GraphHash = null!;
 			}
 
-			public JobDocument(JobId id, StreamId streamId, TemplateId templateId, ContentHash templateHash, ContentHash graphHash, string name, int change, int codeChange, int preflightChange, int clonedPreflightChange, string? preflightDescription, UserId? startedByUserId, Priority? priority, bool? autoSubmit, bool? updateIssues, bool? promoteIssuesByDefault, DateTime createTimeUtc, List<ChainedJobDocument> chainedJobs, bool showUgsBadges, bool showUgsAlerts, string? notificationChannel, string? notificationChannelFilter, List<string>? arguments, IReadOnlyDictionary<string, string>? environment)
+			public JobDocument(JobId id, StreamId streamId, TemplateId templateId, ContentHash templateHash, ContentHash graphHash, string name, int change, int codeChange, CreateJobOptions options, DateTime createTimeUtc)
 			{
 				Id = id;
 				StreamId = streamId;
@@ -313,25 +313,27 @@ namespace Horde.Build.Jobs
 				Name = name;
 				Change = change;
 				CodeChange = codeChange;
-				PreflightChange = preflightChange;
-				ClonedPreflightChange = clonedPreflightChange;
-				PreflightDescription = preflightDescription;
-				StartedByUserId = startedByUserId;
-				Priority = priority ?? HordeCommon.Priority.Normal;
-				AutoSubmit = autoSubmit ?? false;
-				UpdateIssues = updateIssues ?? (startedByUserId == null && preflightChange == 0);
-				PromoteIssuesByDefault = promoteIssuesByDefault ?? false;
+				PreflightChange = options.PreflightChange ?? 0;
+				ClonedPreflightChange = options.ClonedPreflightChange ?? 0;
+				PreflightDescription = options.PreflightDescription;
+				StartedByUserId = options.StartedByUserId;
+				Priority = options.Priority ?? HordeCommon.Priority.Normal;
+				AutoSubmit = options.AutoSubmit ?? false;
+				UpdateIssues = options.UpdateIssues ?? (options.StartedByUserId == null && options.PreflightChange == 0);
+				PromoteIssuesByDefault = options.PromoteIssuesByDefault ?? false;
 				CreateTimeUtc = createTimeUtc;
-				ChainedJobs = chainedJobs;
-				ShowUgsBadges = showUgsBadges;
-				ShowUgsAlerts = showUgsAlerts;
-				NotificationChannel = notificationChannel;
-				NotificationChannelFilter = notificationChannelFilter;
-				Arguments = arguments ?? Arguments;
-				if (environment != null)
+				ChainedJobs.AddRange(options.JobTriggers.Select(x => new ChainedJobDocument(x)));
+				ShowUgsBadges = options.ShowUgsBadges;
+				ShowUgsAlerts = options.ShowUgsAlerts;
+				NotificationChannel = options.NotificationChannel;
+				NotificationChannelFilter = options.NotificationChannelFilter;
+				Arguments.AddRange(options.Arguments);
+
+				foreach (KeyValuePair<string, string> pair in options.Environment)
 				{
-					Environment = new Dictionary<string, string>(environment);
+					Environment[pair.Key] = pair.Value;
 				}
+
 				NextSubResourceId = SubResourceId.Random();
 				UpdateTimeUtc = createTimeUtc;
 			}
@@ -406,19 +408,9 @@ namespace Horde.Build.Jobs
 		}
 
 		/// <inheritdoc/>
-		public async Task<IJob> AddAsync(JobId jobId, StreamId streamId, TemplateId templateRefId, ContentHash templateHash, IGraph graph, string name, int change, int codeChange, int? preflightChange, int? clonedPreflightChange, string? preflightDescription, UserId? startedByUserId, Priority? priority, bool? autoSubmit, bool? updateIssues, bool? promoteIssuesByDefault, List<ChainedJobTemplateConfig>? chainedJobs, bool showUgsBadges, bool showUgsAlerts, string? notificationChannel, string? notificationChannelFilter, List<string>? arguments, IReadOnlyDictionary<string, string>? environment)
+		public async Task<IJob> AddAsync(JobId jobId, StreamId streamId, TemplateId templateRefId, ContentHash templateHash, IGraph graph, string name, int change, int codeChange, CreateJobOptions options)
 		{
-			List<ChainedJobDocument> jobTriggers = new List<ChainedJobDocument>();
-			if (chainedJobs == null)
-			{
-				jobTriggers = new List<ChainedJobDocument>();
-			}
-			else
-			{
-				jobTriggers = chainedJobs.ConvertAll(x => new ChainedJobDocument(x));
-			}
-
-			JobDocument newJob = new JobDocument(jobId, streamId, templateRefId, templateHash, graph.Id, name, change, codeChange, preflightChange ?? 0, clonedPreflightChange ?? 0, preflightDescription, startedByUserId, priority, autoSubmit, updateIssues, promoteIssuesByDefault, DateTime.UtcNow, jobTriggers, showUgsBadges, showUgsAlerts, notificationChannel, notificationChannelFilter, arguments, environment);
+			JobDocument newJob = new JobDocument(jobId, streamId, templateRefId, templateHash, graph.Id, name, change, codeChange, options, DateTime.UtcNow);
 			CreateBatches(newJob, graph, _logger);
 
 			await _jobs.InsertOneAsync(newJob);
diff --git a/Engine/Source/Programs/Horde/Horde.Build/Jobs/JobService.cs b/Engine/Source/Programs/Horde/Horde.Build/Jobs/JobService.cs
--- a/Engine/Source/Programs/Horde/Horde.Build/Jobs/JobService.cs
+++ b/Engine/Source/Programs/Horde/Horde.Build/Jobs/JobService.cs
@@ -175,23 +175,9 @@ namespace Horde.Build.Jobs
 		/// <param name="name">Name of the job</param>
 		/// <param name="change">The change to build</param>
 		/// <param name="codeChange">The corresponding code changelist</param>
-		/// <param name="preflightChange">Optional changelist to preflight</param>
-		/// <param name="clonedPreflightChange">Duplicated preflight change</param>
-		/// <param name="preflightDescription">Description for the preflight</param>
-		/// <param name="startedByUserId">Id of the user that started the job</param>
-		/// <param name="priority">Priority of the job</param>
-		/// <param name="autoSubmit">Whether to automatically submit the preflighted change on completion</param>
-		/// <param name="updateIssues">Whether to update issues when this job completes</param>
-		/// <param name="promoteIssuesByDefault">Whether to promote issues from this job by default</param>
-		/// <param name="jobTriggers">List of downstream job triggers</param>
-		/// <param name="showUgsBadges">Whether to show badges in UGS for this job</param>
-		/// <param name="showUgsAlerts">Whether to show alerts in UGS for this job</param>
-		/// <param name="notificationChannel">Notification Channel for this job</param>
-		/// <param name="notificationChannelFilter">Notification Channel filter for this job</param>
-		/// <param name="arguments">Arguments for the job</param>
-		/// <param name="environment">Additional environment variables to set, in the form X = Y</param>
+		/// <param name="options">Options for the new job</param>
 		/// <returns>Unique id representing the job</returns>
-		public async Task<IJob> CreateJobAsync(JobId? jobId, IStream stream, TemplateId templateRefId, ContentHash templateHash, IGraph graph, string name, int change, int codeChange, int? preflightChange, int? clonedPreflightChange, string? preflightDescription, UserId? startedByUserId, Priority? priority, bool? autoSubmit, bool? updateIssues, bool? promoteIssuesByDefault, List<ChainedJobTemplateConfig>? jobTriggers, bool showUgsBadges, bool showUgsAlerts, string? notificationChannel, string? notificationChannelFilter, IReadOnlyList<string> arguments, IReadOnlyDictionary<string, string>? environment)
+		public async Task<IJob> CreateJobAsync(JobId? jobId, IStream stream, TemplateId templateRefId, ContentHash templateHash, IGraph graph, string name, int change, int codeChange, CreateJobOptions options)
 		{
 			using IScope traceScope = GlobalTracer.Instance.BuildSpan("JobService.CreateJobAsync").StartActive();
 			traceScope.Span.SetTag("JobId", jobId);
@@ -202,67 +188,62 @@ namespace Horde.Build.Jobs
 			traceScope.Span.SetTag("Name", name);
 			traceScope.Span.SetTag("Change", change);
 			traceScope.Span.SetTag("CodeChange", codeChange);
-			traceScope.Span.SetTag("PreflightChange", preflightChange);
-			traceScope.Span.SetTag("ClonedPreflightChange", clonedPreflightChange);
-			traceScope.Span.SetTag("StartedByUserId", startedByUserId.ToString());
-			traceScope.Span.SetTag("Priority", priority.ToString());
-			traceScope.Span.SetTag("ShowUgsBadges", showUgsBadges);
-			traceScope.Span.SetTag("NotificationChannel", notificationChannel ?? "null");
-			traceScope.Span.SetTag("NotificationChannelFilter", notificationChannelFilter ?? "null");
-			traceScope.Span.SetTag("Arguments", String.Join(',', arguments));
-
-			if (autoSubmit != null)
+			traceScope.Span.SetTag("PreflightChange", options.PreflightChange);
+			traceScope.Span.SetTag("ClonedPreflightChange", options.ClonedPreflightChange);
+			traceScope.Span.SetTag("StartedByUserId", options.StartedByUserId.ToString());
+			traceScope.Span.SetTag("Priority", options.Priority.ToString());
+			traceScope.Span.SetTag("ShowUgsBadges", options.ShowUgsBadges);
+			traceScope.Span.SetTag("NotificationChannel", options.NotificationChannel ?? "null");
+			traceScope.Span.SetTag("NotificationChannelFilter", options.NotificationChannelFilter ?? "null");
+			traceScope.Span.SetTag("Arguments", String.Join(',', options.Arguments));
+
+			if (options.AutoSubmit != null)
 			{
-				traceScope.Span.SetTag("AutoSubmit", autoSubmit.Value);
+				traceScope.Span.SetTag("AutoSubmit", options.AutoSubmit.Value);
 			}
-			if (updateIssues != null)
+			if (options.UpdateIssues != null)
 			{
-				traceScope.Span.SetTag("UpdateIssues", updateIssues.Value);
+				traceScope.Span.SetTag("UpdateIssues", options.UpdateIssues.Value);
 			}
-			if (jobTriggers != null)
+			if (options.JobTriggers != null)
 			{
-				traceScope.Span.SetTag("JobTriggers.Count", jobTriggers.Count);
+				traceScope.Span.SetTag("JobTriggers.Count", options.JobTriggers.Count);
 			}
 
 			JobId jobIdValue = jobId ?? Horde.Build.Utilities.ObjectId<IJob>.GenerateNewId();
 			using IDisposable scope = _logger.BeginScope("CreateJobAsync({JobId})", jobIdValue);
 
-			if (preflightChange != null && ShouldClonePreflightChange(stream.Id))
+			if (options.PreflightChange != null && ShouldClonePreflightChange(stream.Id))
 			{
-				clonedPreflightChange = await CloneShelvedChangeAsync(stream.Config.ClusterName, clonedPreflightChange ?? preflightChange.Value);
+				options.ClonedPreflightChange = await CloneShelvedChangeAsync(stream.Config.ClusterName, options.ClonedPreflightChange ?? options.PreflightChange.Value);
 			}
 
-			_logger.LogInformation("Creating job at CL {Change}, code CL {CodeChange}, preflight CL {PreflightChange}, cloned CL {ClonedPreflightChange}", change, codeChange, preflightChange, clonedPreflightChange);
+			_logger.LogInformation("Creating job at CL {Change}, code CL {CodeChange}, preflight CL {PreflightChange}, cloned CL {ClonedPreflightChange}", change, codeChange, options.PreflightChange, options.ClonedPreflightChange);
 
 			Dictionary<string, string> properties = new Dictionary<string, string>();
 			properties["Change"] = change.ToString(CultureInfo.InvariantCulture);
 			properties["CodeChange"] = codeChange.ToString(CultureInfo.InvariantCulture);
-			properties["PreflightChange"] = preflightChange?.ToString(CultureInfo.InvariantCulture) ?? String.Empty;
-			properties["ClonedPreflightChange"] = clonedPreflightChange?.ToString(CultureInfo.InvariantCulture) ?? String.Empty;
+			properties["PreflightChange"] = options.PreflightChange?.ToString(CultureInfo.InvariantCulture) ?? String.Empty;
+			properties["ClonedPreflightChange"] = options.ClonedPreflightChange?.ToString(CultureInfo.InvariantCulture) ?? String.Empty;
 			properties["StreamId"] = stream.Id.ToString();
 			properties["TemplateId"] = templateRefId.ToString();
 			properties["JobId"] = jobIdValue.ToString();
 
-			List<string> expandedArguments = new List<string>();
-			if (arguments != null)
+			for(int idx = 0; idx < options.Arguments.Count; idx++)
 			{
-				foreach (string argument in arguments)
-				{
-					string expandedArgument = StringUtils.ExpandProperties(argument, properties);
-					expandedArguments.Add(expandedArgument);
-				}
+				options.Arguments[idx] = StringUtils.ExpandProperties(options.Arguments[idx], properties);
 			}
 
 			name = StringUtils.ExpandProperties(name, properties);
 
-			IJob newJob = await _jobs.AddAsync(jobIdValue, stream.Id, templateRefId, templateHash, graph, name, change, codeChange, preflightChange, clonedPreflightChange, preflightDescription, startedByUserId, priority, autoSubmit, updateIssues, promoteIssuesByDefault, jobTriggers, showUgsBadges, showUgsAlerts, notificationChannel, notificationChannelFilter, expandedArguments, environment);
+			IJob newJob = await _jobs.AddAsync(jobIdValue, stream.Id, templateRefId, templateHash, graph, name, change, codeChange, options);
 			_jobTaskSource.UpdateQueuedJob(newJob, graph);
 
 			await _jobTaskSource.UpdateUgsBadges(newJob, graph, new List<(LabelState, LabelOutcome)>());
 
-			if (startedByUserId != null)
+			if (options.StartedByUserId != null)
 			{
-				await _userCollection.UpdateSettingsAsync(startedByUserId.Value, addPinnedJobIds: new[] { newJob.Id });
+				await _userCollection.UpdateSettingsAsync(options.StartedByUserId.Value, addPinnedJobIds: new[] { newJob.Id });
 			}
 
 			await AbortAnyDuplicateJobs(newJob);
@@ -1317,7 +1298,11 @@ namespace Horde.Build.Jobs
 					IGraph triggerGraph = await _graphs.AddAsync(template);
 					_logger.LogInformation("Creating downstream job {ChainedJobId} from job {JobId}", chainedJobId, newJob.Id);
 
-					await CreateJobAsync(chainedJobId, stream, jobTrigger.TemplateRefId, templateRef.Hash, triggerGraph, templateRef.Config.Name, newJob.Change, newJob.CodeChange, newJob.PreflightChange, newJob.ClonedPreflightChange, newJob.PreflightDescription, newJob.StartedByUserId, template.Priority, null, newJob.UpdateIssues, newJob.PromoteIssuesByDefault, templateRef.Config.ChainedJobs, false, false, templateRef.Config.NotificationChannel, templateRef.Config.NotificationChannelFilter, template.Arguments, null);
+					CreateJobOptions options = new CreateJobOptions(templateRef.Config);
+					options.PreflightChange = newJob.PreflightChange;
+					options.PreflightDescription = newJob.PreflightDescription;
+
+					await CreateJobAsync(chainedJobId, stream, jobTrigger.TemplateRefId, templateRef.Hash, triggerGraph, templateRef.Config.Name, newJob.Change, newJob.CodeChange, options);
 					return newJob;
 				}
 
diff --git a/Engine/Source/Programs/Horde/Horde.Build/Jobs/JobsController.cs b/Engine/Source/Programs/Horde/Horde.Build/Jobs/JobsController.cs
--- a/Engine/Source/Programs/Horde/Horde.Build/Jobs/JobsController.cs
+++ b/Engine/Source/Programs/Horde/Horde.Build/Jobs/JobsController.cs
@@ -189,8 +189,23 @@ namespace Horde.Build.Jobs
 				updateIssues = create.UpdateIssues.Value;
 			}
 
+			// Create options for the new job
+			CreateJobOptions options = new CreateJobOptions(templateRef.Config);
+			options.PreflightChange = create.PreflightChange;
+			options.PreflightDescription = shelfInfo?.Description;
+			options.StartedByUserId = User.GetUserId();
+			options.Priority = priority;
+			options.AutoSubmit = create.AutoSubmit;
+			options.UpdateIssues = updateIssues;
+			options.Arguments.AddRange(arguments);
+
+			foreach ((string key, string value) in environment)
+			{
+				options.Environment[key] = value;
+			}
+
 			// Create the job
-			IJob job = await _jobService.CreateJobAsync(null, stream, templateRefId, template.Id, graph, name, change, codeChange, create.PreflightChange, null, shelfInfo?.Description, User.GetUserId(), priority, create.AutoSubmit, updateIssues, false, templateRef.Config.ChainedJobs, templateRef.Config.ShowUgsBadges, templateRef.Config.ShowUgsAlerts, templateRef.Config.NotificationChannel, templateRef.Config.NotificationChannelFilter, arguments, environment);
+			IJob job = await _jobService.CreateJobAsync(null, stream, templateRefId, template.Id, graph, name, change, codeChange, options);
 			await UpdateNotificationsAsync(job.Id, new UpdateNotificationsRequest { Slack = true });
 			return new CreateJobResponse(job.Id.ToString());
 		}
diff --git a/Engine/Source/Programs/Horde/Horde.Build/Jobs/Schedules/ScheduleService.cs b/Engine/Source/Programs/Horde/Horde.Build/Jobs/Schedules/ScheduleService.cs
--- a/Engine/Source/Programs/Horde/Horde.Build/Jobs/Schedules/ScheduleService.cs
+++ b/Engine/Source/Programs/Horde/Horde.Build/Jobs/Schedules/ScheduleService.cs
@@ -462,7 +462,12 @@ namespace Horde.Build.Jobs.Schedules
 			{
 				cancellationToken.ThrowIfCancellationRequested();
 				List<string> defaultArguments = template.GetDefaultArguments();
-				IJob newJob = await _jobService.CreateJobAsync(null, stream, templateId, template.Id, graph, template.Name, change, codeChange, null, null, null, null, template.Priority, null, null, template.PromoteIssuesByDefault, templateRef.Config.ChainedJobs, templateRef.Config.ShowUgsBadges, templateRef.Config.ShowUgsAlerts, templateRef.Config.NotificationChannel, templateRef.Config.NotificationChannelFilter, defaultArguments, null);
+
+				CreateJobOptions options = new CreateJobOptions(templateRef.Config);
+				options.Priority = template.Priority;
+				options.Arguments.AddRange(template.GetDefaultArguments());
+
+				IJob newJob = await _jobService.CreateJobAsync(null, stream, templateId, template.Id, graph, template.Name, change, codeChange, options);
 				_logger.LogInformation("Started new job for {StreamId} template {TemplateId} at CL {Change} (Code CL {CodeChange}): {JobId}", stream.Id, templateId, change, codeChange, newJob.Id);
 				await _streamService.UpdateScheduleTriggerAsync(stream, templateId, utcNow, change, new List<JobId> { newJob.Id }, new List<JobId>());
 			}
