Horde: Fix missing code/content tags from preflight changes.

diff --git a/Engine/Source/Programs/Horde/Horde.Build/Perforce/PerforceService.cs b/Engine/Source/Programs/Horde/Horde.Build/Perforce/PerforceService.cs
--- a/Engine/Source/Programs/Horde/Horde.Build/Perforce/PerforceService.cs
+++ b/Engine/Source/Programs/Horde/Horde.Build/Perforce/PerforceService.cs
@@ -799,7 +799,7 @@ namespace Horde.Build.Perforce
 				}
 
 				List<CommitTag> tags = new List<CommitTag>();
-				foreach (CommitTagConfig tagConfig in stream.Config.CommitTags)
+				foreach (CommitTagConfig tagConfig in stream.Config.GetAllCommitTags())
 				{
 					FileFilter filter = tagConfig.CreateFileFilter();
 					if (mappedFiles.Any(x => filter.Matches(x)))
