Horde: Perform the first issue escalation at an interval from the current time, rather than using the issue creation time. Since we only update the escalation time if it doesn't exist, this will be the time that the notification is added.

diff --git a/Engine/Source/Programs/Horde/Horde.Build/Notifications/Sinks/SlackNotificationSink.cs b/Engine/Source/Programs/Horde/Horde.Build/Notifications/Sinks/SlackNotificationSink.cs
--- a/Engine/Source/Programs/Horde/Horde.Build/Notifications/Sinks/SlackNotificationSink.cs
+++ b/Engine/Source/Programs/Horde/Horde.Build/Notifications/Sinks/SlackNotificationSink.cs
@@ -971,7 +971,7 @@ namespace Horde.Build.Notifications.Sinks
 
 				if (workflow.EscalateAlias != null && workflow.EscalateTimes.Count > 0)
 				{
-					DateTime escalateTime = issue.CreatedAt + TimeSpan.FromMinutes(workflow.EscalateTimes[0]);
+					DateTime escalateTime = _clock.UtcNow + TimeSpan.FromMinutes(workflow.EscalateTimes[0]);
 					if (await _escalateIssues.AddAsync(issue.Id, (escalateTime - DateTime.UnixEpoch).TotalSeconds, StackExchange.Redis.When.NotExists))
 					{
 						_logger.LogInformation("First escalation time for issue {IssueId} is {Time}", issue.Id, escalateTime);
