Horde: Log step-specific environment variables to logs on startup.

diff --git a/Engine/Source/Programs/Horde/Horde.Agent/Execution/BuildGraphExecutor.cs b/Engine/Source/Programs/Horde/Horde.Agent/Execution/BuildGraphExecutor.cs
--- a/Engine/Source/Programs/Horde/Horde.Agent/Execution/BuildGraphExecutor.cs
+++ b/Engine/Source/Programs/Horde/Horde.Agent/Execution/BuildGraphExecutor.cs
@@ -757,7 +757,7 @@ namespace Horde.Agent.Execution
 
 		async Task<int> ExecuteCommandAsync(BeginStepResponse step, DirectoryReference workspaceDir, string fileName, string arguments, IReadOnlyDictionary<string, string> envVars, IReadOnlyDictionary<string, string> credentials, ILogger logger, CancellationToken cancellationToken)
 		{
-			Dictionary<string, string> newEnvironment = new Dictionary<string, string>(envVars);
+			Dictionary<string, string> newEnvironment = new Dictionary<string, string>();
 			foreach (object? envVar in Environment.GetEnvironmentVariables())
 			{
 				System.Collections.DictionaryEntry entry = (System.Collections.DictionaryEntry)envVar!;
@@ -772,6 +772,11 @@ namespace Horde.Agent.Execution
 				logger.LogInformation("Setting env var: {Key}={Value}", envVar.Key, envVar.Value);
 				newEnvironment[envVar.Key] = envVar.Value;
 			}
+			foreach (KeyValuePair<string, string> envVar in envVars)
+			{
+				logger.LogInformation("Setting env var: {Key}={Value}", envVar.Key, envVar.Value);
+				newEnvironment[envVar.Key] = envVar.Value;
+			}
 			foreach (KeyValuePair<string, string> envVar in credentials)
 			{
 				logger.LogInformation("Setting env var: {Key}=[redacted]", envVar.Key);
