Horde: Add a DefaultChange setting to job templates, which specify queries to be executed to determine the change to run at. If no query matches, the latest change will be used.

diff --git a/Engine/Source/Programs/Horde/Horde.Build/Jobs/Templates/TemplateMessages.cs b/Engine/Source/Programs/Horde/Horde.Build/Jobs/Templates/TemplateMessages.cs
--- a/Engine/Source/Programs/Horde/Horde.Build/Jobs/Templates/TemplateMessages.cs
+++ b/Engine/Source/Programs/Horde/Horde.Build/Jobs/Templates/TemplateMessages.cs
@@ -4,6 +4,7 @@ using System;
 using System.Collections.Generic;
 using System.ComponentModel.DataAnnotations;
 using EpicGames.Core;
+using Horde.Build.Streams;
 using HordeCommon;
 
 namespace Horde.Build.Jobs.Templates
@@ -430,6 +431,11 @@ namespace Horde.Build.Jobs.Templates
 		/// </summary>
 		public string? SubmitDescription { get; set; }
 
+		/// <summary>
+		/// Default change to build at. Each object has a condition parameter which can evaluated by the server to determine which change to use.
+		/// </summary>
+		public List<ChangeQueryConfig>? DefaultChange { get; set; }
+
 		/// <summary>
 		/// Fixed arguments for the new job
 		/// </summary>
diff --git a/Engine/Source/Programs/Horde/Horde.Build/Streams/StreamConfig.cs b/Engine/Source/Programs/Horde/Horde.Build/Streams/StreamConfig.cs
--- a/Engine/Source/Programs/Horde/Horde.Build/Streams/StreamConfig.cs
+++ b/Engine/Source/Programs/Horde/Horde.Build/Streams/StreamConfig.cs
@@ -451,6 +451,16 @@ namespace Horde.Build.Streams
 	/// </summary>
 	public class ChangeQueryConfig
 	{
+		/// <summary>
+		/// Name of this query, for display on the dashboard.
+		/// </summary>
+		public string? Name { get; set; }
+
+		/// <summary>
+		/// Condition to evaluate before deciding to use this query. May query tags in a preflight.
+		/// </summary>
+		public Condition? Condition { get; set; }
+		
 		/// <summary>
 		/// The template id to query
 		/// </summary>
@@ -470,11 +480,6 @@ namespace Horde.Build.Streams
 		/// Finds the last commit with this tag
 		/// </summary>
 		public CommitTag? CommitTag { get; set; }
-
-		/// <summary>
-		/// Condition to evaluate before deciding to use this query. May query tags in a preflight.
-		/// </summary>
-		public Condition? Condition { get; set; }
 	}
 
 	/// <summary>
@@ -490,17 +495,8 @@ namespace Horde.Build.Streams
 		/// <summary>
 		/// Query for the change to use
 		/// </summary>
-		[Obsolete("Use ChangeQueries instead")]
-		public ChangeQueryConfig? Change
-		{
-			get => (ChangeQueries != null && ChangeQueries.Count > 0) ? ChangeQueries[0] : null;
-			set => ChangeQueries = ((value == null) ? null : new List<ChangeQueryConfig> { value });
-		}
-
-		/// <summary>
-		/// List of queries to execute to determine the base change for this preflight
-		/// </summary>
-		public List<ChangeQueryConfig>? ChangeQueries { get; set; }
+		[Obsolete("Set on template instead")]
+		public ChangeQueryConfig? Change { get; set; }
 	}
 
 	/// <summary>
