Horde Storage - Added some warnings when we fail to measure latency to help find out why this seems to be failing.

diff --git a/Engine/Source/Programs/Horde/HordeStorage/Horde.Storage/Implementation/PeerStatusService.cs b/Engine/Source/Programs/Horde/HordeStorage/Horde.Storage/Implementation/PeerStatusService.cs
--- a/Engine/Source/Programs/Horde/HordeStorage/Horde.Storage/Implementation/PeerStatusService.cs
+++ b/Engine/Source/Programs/Horde/HordeStorage/Horde.Storage/Implementation/PeerStatusService.cs
@@ -9,6 +9,8 @@ using System.Threading.Tasks;
 using EpicGames.Core;
 using Jupiter;
 using Microsoft.Extensions.Options;
+using Serilog;
+using Log = Serilog.Log;
 
 namespace Horde.Storage.Implementation
 {
@@ -46,6 +48,7 @@ namespace Horde.Storage.Implementation
         private readonly IHttpClientFactory _clientFactory;
         private readonly Dictionary<string, PeerStatus> _peers = new Dictionary<string, PeerStatus>(StringComparer.InvariantCultureIgnoreCase);
         private volatile bool _alreadyPolling = false;
+        private readonly ILogger _logger = Log.ForContext<PeerStatusService>();
 
         public PeerStatus? GetPeerStatus(string regionName)
         {
@@ -161,13 +164,15 @@ namespace Horde.Storage.Implementation
                 // ignore error responses as they may not have reached the actual instance
                 if (!result.IsSuccessStatusCode)
                 {
+                    _logger.Warning("Non-success status code ({StatusCode}) when attempting to measure latency to {Endpoint}", result.StatusCode, endpoint.Url);
                     return int.MaxValue;
                 }
 
                 return (int)stopwatch.ElapsedMilliseconds;
             }
-            catch (Exception)
+            catch (Exception e)
             {
+                _logger.Warning(e, "Exception when attempting to measure latency from {Endpoint}", endpoint.Url);
                 // error reaching the endpoint is just considered to max latency
                 return int.MaxValue;
             }
