Horde: Include the "DefaultChange" config value in template responses.

diff --git a/Engine/Source/Programs/Horde/Horde.Build/Streams/StreamMessages.cs b/Engine/Source/Programs/Horde/Horde.Build/Streams/StreamMessages.cs
--- a/Engine/Source/Programs/Horde/Horde.Build/Streams/StreamMessages.cs
+++ b/Engine/Source/Programs/Horde/Horde.Build/Streams/StreamMessages.cs
@@ -254,6 +254,11 @@ namespace Horde.Build.Streams
 		/// </summary>
 		public List<GetTemplateStepStateResponse>? StepStates { get; set; }
 
+		/// <summary>
+		/// List of queries for the default changelist
+		/// </summary>
+		public List<ChangeQueryConfig>? DefaultChange { get; set; }
+
 		/// <summary>
 		/// ACL for this template
 		/// </summary>
@@ -279,6 +284,7 @@ namespace Horde.Build.Streams
 			Schedule = (templateRef.Schedule != null) ? new GetScheduleResponse(templateRef.Schedule) : null;
 			ChainedJobs = (templateRef.Config.ChainedJobs != null && templateRef.Config.ChainedJobs.Count > 0) ? templateRef.Config.ChainedJobs : null;
 			StepStates = stepStates;
+			DefaultChange = templateRef.Config.DefaultChange;
 			Acl = (bIncludeAcl && templateRef.Acl != null) ? new GetAclResponse(templateRef.Acl) : null;
 		}
 	}
