Horde: Attempt to fix agent tests on Linux.

diff --git a/Engine/Source/Programs/Horde/Horde.Agent.Tests/LogParserTests.cs b/Engine/Source/Programs/Horde/Horde.Agent.Tests/LogParserTests.cs
--- a/Engine/Source/Programs/Horde/Horde.Agent.Tests/LogParserTests.cs
+++ b/Engine/Source/Programs/Horde/Horde.Agent.Tests/LogParserTests.cs
@@ -5,6 +5,7 @@ using System.Buffers;
 using System.Collections.Generic;
 using System.IO;
 using System.Linq;
+using System.Runtime.InteropServices;
 using System.Text;
 using EpicGames.Core;
 using EpicGames.Perforce;
@@ -21,6 +22,10 @@ namespace Horde.Agent.Tests
 	{
 		const string LogLine = "LogLine";
 
+		static readonly DirectoryReference s_rootDir = new DirectoryReference(RuntimeInformation.IsOSPlatform(OSPlatform.Windows) ? "C:\\Horde" : "/horde");
+
+		static string MakeAbsolutePath(string path) => FileReference.Combine(s_rootDir, path).FullName;
+
 		class LoggerCapture : ILogger
 		{
 			int _logLineIndex;
@@ -237,10 +242,12 @@ namespace Horde.Agent.Tests
 		{
 			// C# compile error from UBT
 			{
+				string absPath = MakeAbsolutePath(@"Engine\Source\Runtime\CoreOnline\CoreOnline.Build.cs");
+
 				string[] lines =
 				{
-					@"  ERROR: c:\Horde\Engine\Source\Runtime\CoreOnline\CoreOnline.Build.cs(4,7): error CS0246: The type or namespace name 'Tools' could not be found (are you missing a using directive or an assembly reference?)",
-					@"  WARNING: C:\horde\Engine\Source\Runtime\CoreOnline\CoreOnline.Build.cs(4,7): warning CS0246: The type or namespace name 'Tools' could not be found (are you missing a using directive or an assembly reference?)"
+					@"  ERROR: " + absPath + @"(4,7): error CS0246: The type or namespace name 'Tools' could not be found (are you missing a using directive or an assembly reference?)",
+					@"  WARNING: " + absPath + @"(4,7): warning CS0246: The type or namespace name 'Tools' could not be found (are you missing a using directive or an assembly reference?)"
 				};
 
 				List<LogEvent> logEvents = Parse(String.Join("\n", lines));
@@ -251,7 +258,7 @@ namespace Horde.Agent.Tests
 				Assert.AreEqual("4", logEvents[0].GetProperty("line").ToString());
 
 				LogValue fileProperty = (LogValue)logEvents[0].GetProperty("file");
-				Assert.AreEqual(@"c:\Horde\Engine\Source\Runtime\CoreOnline\CoreOnline.Build.cs", fileProperty.Text);
+				Assert.AreEqual(absPath, fileProperty.Text);
 				Assert.AreEqual(@"SourceFile", fileProperty.Type);
 				Assert.AreEqual(@"Engine/Source/Runtime/CoreOnline/CoreOnline.Build.cs", fileProperty.Properties!["relativePath"].ToString());
 				Assert.AreEqual(@"//UE4/Main/Engine/Source/Runtime/CoreOnline/CoreOnline.Build.cs@12345", fileProperty.Properties!["depotPath"].ToString());
@@ -261,7 +268,7 @@ namespace Horde.Agent.Tests
 				Assert.AreEqual("4", logEvents[1].GetProperty("line").ToString());
 
 				LogValue fileProperty1 = logEvents[1].GetProperty<LogValue>("file");
-				Assert.AreEqual(@"C:\horde\Engine\Source\Runtime\CoreOnline\CoreOnline.Build.cs", fileProperty1.Text);
+				Assert.AreEqual(absPath, fileProperty1.Text);
 				Assert.AreEqual(@"SourceFile", fileProperty1.Type);
 				Assert.AreEqual(@"Engine/Source/Runtime/CoreOnline/CoreOnline.Build.cs", fileProperty1.Properties!["relativePath"].ToString());
 				Assert.AreEqual(@"//UE4/Main/Engine/Source/Runtime/CoreOnline/CoreOnline.Build.cs@12345", fileProperty1.Properties!["depotPath"].ToString());
@@ -833,7 +840,7 @@ namespace Horde.Agent.Tests
 
 		static List<LogEvent> Parse(string text)
 		{
-			return Parse(text, new DirectoryReference("C:\\Horde".Replace('\\', Path.DirectorySeparatorChar)));
+			return Parse(text, s_rootDir);
 		}
 
 		static List<LogEvent> Parse(string text, DirectoryReference workspaceDir)
