Horde: More fixes for Linux tests.

diff --git a/Engine/Source/Programs/Horde/Horde.Agent.Tests/LogParserTests.cs b/Engine/Source/Programs/Horde/Horde.Agent.Tests/LogParserTests.cs
--- a/Engine/Source/Programs/Horde/Horde.Agent.Tests/LogParserTests.cs
+++ b/Engine/Source/Programs/Horde/Horde.Agent.Tests/LogParserTests.cs
@@ -184,8 +184,8 @@ namespace Horde.Agent.Tests
 			{
 				string[] lines =
 				{
-					@"  GenerateSigningRequestDialog.cs(22,7): error CS0246: The type or namespace name 'Org' could not be found (are you missing a using directive or an assembly reference?) [c:\Horde\Engine\Source\Programs\IOS\iPhonePackager\iPhonePackager.csproj]",
-					@"  Utilities.cs(16,7): error CS0246: The type or namespace name 'Org' could not be found (are you missing a using directive or an assembly reference?) [c:\Horde\Engine\Source\Programs\IOS\iPhonePackager\iPhonePackager.csproj]"
+					@"  GenerateSigningRequestDialog.cs(22,7): error CS0246: The type or namespace name 'Org' could not be found (are you missing a using directive or an assembly reference?) [" + MakeAbsolutePath(@"Engine\Source\Programs\IOS\iPhonePackager\iPhonePackager.csproj") + @"]",
+					@"  Utilities.cs(16,7): error CS0246: The type or namespace name 'Org' could not be found (are you missing a using directive or an assembly reference?) [" + MakeAbsolutePath(@"Engine\Source\Programs\IOS\iPhonePackager\iPhonePackager.csproj") + @"]"
 				};
 
 				List<LogEvent> logEvents = Parse(String.Join("\n", lines));
@@ -220,7 +220,7 @@ namespace Horde.Agent.Tests
 			{
 				string[] lines =
 				{
-					@"  Configuration\TargetRules.cs(1497,58): warning CS8625: Cannot convert null literal to non-nullable reference type. [C:\Horde\Engine\Source\Programs\UnrealBuildTool\UnrealBuildTool.csproj]",
+					@"  Configuration\TargetRules.cs(1497,58): warning CS8625: Cannot convert null literal to non-nullable reference type. [" + MakeAbsolutePath(@"Engine\Source\Programs\UnrealBuildTool\UnrealBuildTool.csproj") + @"]",
 				};
 
 				List<LogEvent> events = Parse(String.Join("\n", lines));
@@ -294,8 +294,8 @@ namespace Horde.Agent.Tests
 			{
 				string[] lines =
 				{
-					@" C:\Horde\Foo\Bar.txt(20): warning TL2012: Some error message",
-					@" C:\Horde\Foo\Bar.txt(20, 30) : warning TL2034: Some error message",
+					@" " + MakeAbsolutePath(@"Foo\Bar.txt") + @"(20): warning TL2012: Some error message",
+					@" " + MakeAbsolutePath(@"Foo\Bar.txt") + @"(20, 30) : warning TL2034: Some error message",
 					@" CSC : error CS2012: Cannot open 'D:\Build\++UE4\Sync\Engine\Source\Programs\Enterprise\Datasmith\DatasmithRevitExporter\Resources\obj\Release\DatasmithRevitResources.dll' for writing -- 'The process cannot access the file 'D:\Build\++UE4\Sync\Engine\Source\Programs\Enterprise\Datasmith\DatasmithRevitExporter\Resources\obj\Release\DatasmithRevitResources.dll' because it is being used by another process.' [D:\Build\++UE4\Sync\Engine\Source\Programs\Enterprise\Datasmith\DatasmithRevitExporter\Resources\DatasmithRevitResources.csproj]"
 				};
 
@@ -335,10 +335,10 @@ namespace Horde.Agent.Tests
 			{
 				string[] lines =
 				{
-					@"C:\Horde\Engine\Plugins\Experimental\VirtualCamera\Source\VirtualCamera\Private\VCamBlueprintFunctionLibrary.cpp(249): error C2220: the following warning is treated as an error",
-					@"C:\Horde\Engine\Plugins\Experimental\VirtualCamera\Source\VirtualCamera\Private\VCamBlueprintFunctionLibrary.cpp(249): warning C4996: 'UEditorLevelLibrary::PilotLevelActor': The Editor Scripting Utilities Plugin is deprecated - Use the function in Level Editor Subsystem Please update your code to the new API before upgrading to the next release, otherwise your project will no longer compile.",
+					MakeAbsolutePath(@"Engine\Plugins\Experimental\VirtualCamera\Source\VirtualCamera\Private\VCamBlueprintFunctionLibrary.cpp") + @"(249): error C2220: the following warning is treated as an error",
+					MakeAbsolutePath(@"Engine\Plugins\Experimental\VirtualCamera\Source\VirtualCamera\Private\VCamBlueprintFunctionLibrary.cpp") + @"(249): warning C4996: 'UEditorLevelLibrary::PilotLevelActor': The Editor Scripting Utilities Plugin is deprecated - Use the function in Level Editor Subsystem Please update your code to the new API before upgrading to the next release, otherwise your project will no longer compile.",
 					@"..\Plugins\Editor\EditorScriptingUtilities\Source\EditorScriptingUtilities\Public\EditorLevelLibrary.h(122): note: see declaration of 'UEditorLevelLibrary::PilotLevelActor'",
-					@"C:\Horde\Engine\Plugins\Experimental\VirtualCamera\Source\VirtualCamera\Private\VCamBlueprintFunctionLibrary.cpp(314): warning C4996: 'UEditorLevelLibrary::EditorSetGameView': The Editor Scripting Utilities Plugin is deprecated - Use the function in Level Editor Subsystem Please update your code to the new API before upgrading to the next release, otherwise your project will no longer compile.",
+					MakeAbsolutePath(@"Engine\Plugins\Experimental\VirtualCamera\Source\VirtualCamera\Private\VCamBlueprintFunctionLibrary.cpp") + @"(314): warning C4996: 'UEditorLevelLibrary::EditorSetGameView': The Editor Scripting Utilities Plugin is deprecated - Use the function in Level Editor Subsystem Please update your code to the new API before upgrading to the next release, otherwise your project will no longer compile.",
 					@"..\Plugins\Editor\EditorScriptingUtilities\Source\EditorScriptingUtilities\Public\EditorLevelLibrary.h(190): note: see declaration of 'UEditorLevelLibrary::EditorSetGameView'"
 				};
 
@@ -368,10 +368,10 @@ namespace Horde.Agent.Tests
 		{
 			string[] lines =
 			{
-				@"LogBlueprint: Warning: [AssetLog] D:\build\++UE5\Sync\QAGame\Plugins\NiagaraFluids\Content\Blueprints\Phsyarum_BP.uasset: [Compiler] Fill Texture 2D : Usage of 'Fill Texture 2D' has been deprecated. This function has been replaced by object user variables on the emitter to specify render targets to fill with data.",
+				@"LogBlueprint: Warning: [AssetLog] " + MakeAbsolutePath(@"QAGame\Plugins\NiagaraFluids\Content\Blueprints\Phsyarum_BP.uasset") + @": [Compiler] Fill Texture 2D : Usage of 'Fill Texture 2D' has been deprecated. This function has been replaced by object user variables on the emitter to specify render targets to fill with data.",
 			};
 
-			List<LogEvent> events = Parse(String.Join("\n", lines), new DirectoryReference("D:\\build\\++UE5\\Sync"));
+			List<LogEvent> events = Parse(String.Join("\n", lines));
 			Assert.AreEqual(1, events.Count);
 			Assert.AreEqual(events[0].Id, KnownLogEvents.Engine_AssetLog);
 
@@ -405,13 +405,10 @@ namespace Horde.Agent.Tests
 		[TestMethod]
 		public void MSBuildEventMatcher()
 		{
-			List<LogEvent> logEvents = Parse(@"  C:\Program Files (x86)\Microsoft Visual Studio\2019\BuildTools\MSBuild\Current\Bin\Microsoft.Common.CurrentVersion.targets(4207,5): warning MSB3026: Could not copy ""obj\Development\DotNETUtilities.dll"" to ""..\..\..\..\Binaries\DotNET\DotNETUtilities.dll"". Beginning retry 2 in 1000ms. The process cannot access the file '..\..\..\..\Binaries\DotNET\DotNETUtilities.dll' because it is being used by another process. The file is locked by: ""UnrealAutomationTool(13236)"" [C:\Horde\Engine\Source\Programs\DotNETCommon\DotNETUtilities\DotNETUtilities.csproj]");
+			List<LogEvent> logEvents = Parse(@"  C:\Program Files (x86)\Microsoft Visual Studio\2019\BuildTools\MSBuild\Current\Bin\Microsoft.Common.CurrentVersion.targets(4207,5): warning MSB3026: Could not copy ""obj\Development\DotNETUtilities.dll"" to ""..\..\..\..\Binaries\DotNET\DotNETUtilities.dll"". Beginning retry 2 in 1000ms. The process cannot access the file '..\..\..\..\Binaries\DotNET\DotNETUtilities.dll' because it is being used by another process. The file is locked by: ""UnrealAutomationTool(13236)"" [" + MakeAbsolutePath(@"Engine\Source\Programs\DotNETCommon\DotNETUtilities\DotNETUtilities.csproj") + @"]");
 			CheckEventGroup(logEvents, 0, 1, LogLevel.Information, KnownLogEvents.Systemic_MSBuild);
 
 			Assert.AreEqual("warning", logEvents[0].GetProperty("severity").ToString());
-			
-			// FIXME: Fails on Linux. Properties dict is empty
-			//Assert.AreEqual(@"Engine\Source\Programs\DotNETCommon\DotNETUtilities\DotNETUtilities.csproj", GetSubProperty(Events[0], "file", "relativePath"));
 		}
 
 		[TestMethod]
