Horde: Add an experimental v2 endpoint for streams, which includes configuration data separately to main response.

diff --git a/Engine/Source/Programs/Horde/Horde.Build/Streams/StreamMessages.cs b/Engine/Source/Programs/Horde/Horde.Build/Streams/StreamMessages.cs
--- a/Engine/Source/Programs/Horde/Horde.Build/Streams/StreamMessages.cs
+++ b/Engine/Source/Programs/Horde/Horde.Build/Streams/StreamMessages.cs
@@ -11,9 +11,15 @@ using Horde.Build.Users;
 using Horde.Build.Utilities;
 using Horde.Build.Issues;
 using Horde.Build.Perforce;
+using Horde.Build.Projects;
+using Horde.Build.Jobs;
+using Microsoft.AspNetCore.DataProtection;
 
 namespace Horde.Build.Streams
 {
+	using JobId = ObjectId<IJob>;
+	using StreamId = StringId<IStream>;
+	using ProjectId = StringId<IProject>;
 	using TemplateId = StringId<ITemplateRef>;
 
 	/// <summary>
@@ -432,6 +438,132 @@ namespace Horde.Build.Streams
 		}
 	}
 
+	/// <summary>
+	/// Response describing a stream
+	/// </summary>
+	public class GetStreamResponseV2
+	{
+		readonly IStream _stream;
+
+		/// <inheritdoc cref="IStream.Id"/>
+		public StreamId Id => _stream.Id;
+
+		/// <inheritdoc cref="IStream.ProjectId"/>
+		public ProjectId ProjectId => _stream.ProjectId;
+
+		/// <inheritdoc cref="IStream.Name"/>
+		public string Name => _stream.Name;
+
+		/// <inheritdoc cref="IStream.Name"/>
+		public string ConfigRevision => _stream.ConfigRevision;
+
+		/// <inheritdoc cref="IStream.Config"/>
+		public StreamConfig? Config { get; }
+
+		/// <inheritdoc cref="IStream.Deleted"/>
+		public bool Deleted => _stream.Deleted;
+
+		/// <inheritdoc cref="IStream.Templates"/>
+		public List<GetTemplateRefResponseV2> Templates { get; }
+
+		/// <inheritdoc cref="IStream.PausedUntil"/>
+		public DateTime? PausedUntil => _stream.PausedUntil;
+
+		/// <inheritdoc cref="IStream.PauseComment"/>
+		public string? PauseComment => _stream.PauseComment;
+
+		internal GetStreamResponseV2(IStream stream, bool config, List<GetTemplateRefResponseV2> templates)
+		{
+			_stream = stream;
+			Templates = templates;
+
+			if (config)
+			{
+				Config = stream.Config;
+			}
+		}
+	}
+
+	/// <summary>
+	/// Job template in a stream
+	/// </summary>
+	public class GetTemplateRefResponseV2
+	{
+		readonly ITemplateRef _templateRef;
+
+		/// <inheritdoc cref="ITemplateRef.Id"/>
+		public TemplateId Id => _templateRef.Id;
+
+		/// <inheritdoc cref="ITemplateRef.Hash"/>
+		public ContentHash Hash => _templateRef.Hash;
+
+		/// <inheritdoc cref="ITemplateRef.Schedule"/>
+		public GetTemplateScheduleResponseV2? Schedule { get; }
+
+		/// <inheritdoc cref="ITemplateRef.StepStates"/>
+		public List<GetTemplateStepResponseV2>? StepStates { get; }
+
+		internal GetTemplateRefResponseV2(ITemplateRef templateRef, List<GetTemplateStepResponseV2>? stepStates)
+		{
+			_templateRef = templateRef;
+
+			StepStates = stepStates;
+
+			if (templateRef.Schedule != null)
+			{
+				Schedule = new GetTemplateScheduleResponseV2(templateRef.Schedule);
+			}
+		}
+	}
+
+	/// <summary>
+	/// Schedule for a template
+	/// </summary>
+	public class GetTemplateScheduleResponseV2
+	{
+		readonly ITemplateSchedule _schedule;
+
+		/// <inheritdoc cref="ITemplateSchedule.LastTriggerChange"/>
+		public int LastTriggerChange => _schedule.LastTriggerChange;
+
+		/// <inheritdoc cref="ITemplateSchedule.LastTriggerTimeUtc"/>
+		public DateTime LastTriggerTimeUtc => _schedule.LastTriggerTimeUtc;
+
+		/// <inheritdoc cref="ITemplateSchedule.LastTriggerTimeUtc"/>
+		public IReadOnlyList<JobId> ActiveJobs => _schedule.ActiveJobs;
+
+		internal GetTemplateScheduleResponseV2(ITemplateSchedule schedule)
+		{
+			_schedule = schedule;
+		}
+	}
+
+	/// <summary>
+	/// State information for a step in the stream
+	/// </summary>
+	public class GetTemplateStepResponseV2
+	{
+		readonly ITemplateStep _step;
+
+		/// <inheritdoc cref="ITemplateStep.Name"/>
+		public string Name => _step.Name;
+
+		/// <inheritdoc cref="ITemplateStep.PausedByUserId"/>
+		public GetThinUserInfoResponse PausedByUserInfo { get; }
+
+		/// <inheritdoc cref="ITemplateStep.PauseTimeUtc"/>
+		public DateTime PauseTimeUtc => _step.PauseTimeUtc;
+
+		/// <summary>
+		/// Constructor
+		/// </summary>
+		public GetTemplateStepResponseV2(ITemplateStep step, GetThinUserInfoResponse pausedByUserInfo)
+		{
+			_step = step;
+			PausedByUserInfo = pausedByUserInfo;
+		}
+	}
+
 	/// <summary>
 	/// Information about a commit
 	/// </summary>
diff --git a/Engine/Source/Programs/Horde/Horde.Build/Streams/StreamsController.cs b/Engine/Source/Programs/Horde/Horde.Build/Streams/StreamsController.cs
--- a/Engine/Source/Programs/Horde/Horde.Build/Streams/StreamsController.cs
+++ b/Engine/Source/Programs/Horde/Horde.Build/Streams/StreamsController.cs
@@ -79,6 +79,34 @@ namespace Horde.Build.Streams
 			return responses.OrderBy(x => x.Order).ThenBy(x => x.Name).Select(x => PropertyFilter.Apply(x, filter)).ToList();
 		}
 
+		/// <summary>
+		/// Query all the streams for a particular project.
+		/// </summary>
+		/// <param name="projectIds">Unique id of the project to query</param>
+		/// <param name="filter">Filter for the properties to return</param>
+		/// <returns>Information about all the projects</returns>
+		[HttpGet]
+		[Route("/api/v2/streams")]
+		[ProducesResponseType(typeof(List<GetStreamResponse>), 200)]
+		public async Task<ActionResult<List<object>>> GetStreamsAsyncV2([FromQuery(Name = "ProjectId")] string[] projectIds, [FromQuery] PropertyFilter? filter = null)
+		{
+			ProjectId[] projectIdValues = Array.ConvertAll(projectIds, x => new ProjectId(x));
+
+			List<IStream> streams = await _streamService.GetStreamsAsync(projectIdValues);
+			ProjectPermissionsCache permissionsCache = new ProjectPermissionsCache();
+
+			List<GetStreamResponseV2> responses = new List<GetStreamResponseV2>();
+			foreach (IStream stream in streams)
+			{
+				if (await _streamService.AuthorizeAsync(stream, AclAction.ViewStream, User, permissionsCache))
+				{
+					GetStreamResponseV2 response = await CreateGetStreamResponseAsyncV2(stream, false, permissionsCache);
+					responses.Add(response);
+				}
+			}
+			return responses.OrderBy(x => x.Id).Select(x => PropertyFilter.Apply(x, filter)).ToList();
+		}
+
 		/// <summary>
 		/// Retrieve information about a specific stream.
 		/// </summary>
@@ -105,6 +133,33 @@ namespace Horde.Build.Streams
 			return PropertyFilter.Apply(await CreateGetStreamResponse(stream, permissionsCache), filter);
 		}
 
+		/// <summary>
+		/// Retrieve information about a specific stream.
+		/// </summary>
+		/// <param name="streamId">Id of the stream to get information about</param>
+		/// <param name="config">Whether to include config data for the stream</param>
+		/// <param name="filter">Filter for the properties to return</param>
+		/// <returns>Information about the requested project</returns>
+		[HttpGet]
+		[Route("/api/v2/streams/{streamId}")]
+		[ProducesResponseType(typeof(GetStreamResponseV2), 200)]
+		public async Task<ActionResult<object>> GetStreamAsyncV2(StreamId streamId, [FromQuery] bool config = true, [FromQuery] PropertyFilter? filter = null)
+		{
+			IStream? stream = await _streamService.GetStreamAsync(streamId);
+			if (stream == null)
+			{
+				return NotFound(streamId);
+			}
+
+			ProjectPermissionsCache permissionsCache = new ProjectPermissionsCache();
+			if (!await _streamService.AuthorizeAsync(stream, AclAction.ViewStream, User, permissionsCache))
+			{
+				return Forbid(AclAction.ViewStream, streamId);
+			}
+
+			return PropertyFilter.Apply(await CreateGetStreamResponseAsyncV2(stream, config, permissionsCache), filter);
+		}
+
 		/// <summary>
 		/// Create a stream response object, including all the templates
 		/// </summary>
@@ -155,6 +210,39 @@ namespace Horde.Build.Streams
 			return stream.ToApiResponse(bIncludeAcl, apiTemplateRefs);
 		}
 
+		/// <summary>
+		/// Create a stream response object, including all the templates
+		/// </summary>
+		/// <param name="stream">Stream to create response for</param>
+		/// <param name="config">Whether to include config data in the response</param>
+		/// <param name="cache">Permissions cache</param>
+		/// <returns>Response object</returns>
+		async Task<GetStreamResponseV2> CreateGetStreamResponseAsyncV2(IStream stream, bool config, ProjectPermissionsCache cache)
+		{
+			using IScope scope = GlobalTracer.Instance.BuildSpan("CreateGetStreamResponseAsyncV2").StartActive();
+			scope.Span.SetTag("streamId", stream.Id);
+
+			List<GetTemplateRefResponseV2> templateRefResponses = new List<GetTemplateRefResponseV2>();
+			foreach (ITemplateRef templateRef in stream.Templates.Values)
+			{
+				List<GetTemplateStepResponseV2>? templateStepResponses = null;
+				if (templateRef.StepStates != null && templateRef.StepStates.Count > 0)
+				{
+					templateStepResponses = new List<GetTemplateStepResponseV2>();
+					foreach (ITemplateStep templateStep in templateRef.StepStates)
+					{
+						GetThinUserInfoResponse? pausedByUserInfo = new GetThinUserInfoResponse(await _userCollection.GetCachedUserAsync(templateStep.PausedByUserId));
+						templateStepResponses.Add(new GetTemplateStepResponseV2(templateStep, pausedByUserInfo));
+					}
+				}
+
+				GetTemplateRefResponseV2 templateRefResponse = new GetTemplateRefResponseV2(templateRef, templateStepResponses);
+				templateRefResponses.Add(templateRefResponse);
+			}
+
+			return new GetStreamResponseV2(stream, config, templateRefResponses);
+		}
+
 		/// <summary>
 		/// Gets a list of changes for a stream
 		/// </summary>
