Horde: Separate endpoints for retrieving config data and state data.

diff --git a/Engine/Source/Programs/Horde/Horde.Build/Streams/StreamMessages.cs b/Engine/Source/Programs/Horde/Horde.Build/Streams/StreamMessages.cs
--- a/Engine/Source/Programs/Horde/Horde.Build/Streams/StreamMessages.cs
+++ b/Engine/Source/Programs/Horde/Horde.Build/Streams/StreamMessages.cs
@@ -438,6 +438,25 @@ namespace Horde.Build.Streams
 		}
 	}
 
+	/// <summary>
+	/// Returns the configuration data for a stream
+	/// </summary>
+	public class GetStreamConfigResponseV2
+	{
+		readonly IStream _stream;
+
+		/// <inheritdoc cref="IStream.ConfigRevision"/>
+		public string Revision => _stream.ConfigRevision;
+
+		/// <inheritdoc cref="IStream.Config"/>
+		public StreamConfig Config => _stream.Config;
+
+		internal GetStreamConfigResponseV2(IStream stream)
+		{
+			_stream = stream;
+		}
+	}
+
 	/// <summary>
 	/// Response describing a stream
 	/// </summary>
@@ -457,9 +476,6 @@ namespace Horde.Build.Streams
 		/// <inheritdoc cref="IStream.Name"/>
 		public string ConfigRevision => _stream.ConfigRevision;
 
-		/// <inheritdoc cref="IStream.Config"/>
-		public StreamConfig? Config { get; }
-
 		/// <inheritdoc cref="IStream.Deleted"/>
 		public bool Deleted => _stream.Deleted;
 
@@ -472,15 +488,10 @@ namespace Horde.Build.Streams
 		/// <inheritdoc cref="IStream.PauseComment"/>
 		public string? PauseComment => _stream.PauseComment;
 
-		internal GetStreamResponseV2(IStream stream, bool config, List<GetTemplateRefResponseV2> templates)
+		internal GetStreamResponseV2(IStream stream, List<GetTemplateRefResponseV2> templates)
 		{
 			_stream = stream;
 			Templates = templates;
-
-			if (config)
-			{
-				Config = stream.Config;
-			}
 		}
 	}
 
diff --git a/Engine/Source/Programs/Horde/Horde.Build/Streams/StreamsController.cs b/Engine/Source/Programs/Horde/Horde.Build/Streams/StreamsController.cs
--- a/Engine/Source/Programs/Horde/Horde.Build/Streams/StreamsController.cs
+++ b/Engine/Source/Programs/Horde/Horde.Build/Streams/StreamsController.cs
@@ -100,7 +100,7 @@ namespace Horde.Build.Streams
 			{
 				if (await _streamService.AuthorizeAsync(stream, AclAction.ViewStream, User, permissionsCache))
 				{
-					GetStreamResponseV2 response = await CreateGetStreamResponseAsyncV2(stream, false, permissionsCache);
+					GetStreamResponseV2 response = await CreateGetStreamResponseAsyncV2(stream, permissionsCache);
 					responses.Add(response);
 				}
 			}
@@ -137,13 +137,12 @@ namespace Horde.Build.Streams
 		/// Retrieve information about a specific stream.
 		/// </summary>
 		/// <param name="streamId">Id of the stream to get information about</param>
-		/// <param name="config">Whether to include config data for the stream</param>
 		/// <param name="filter">Filter for the properties to return</param>
 		/// <returns>Information about the requested project</returns>
 		[HttpGet]
 		[Route("/api/v2/streams/{streamId}")]
 		[ProducesResponseType(typeof(GetStreamResponseV2), 200)]
-		public async Task<ActionResult<object>> GetStreamAsyncV2(StreamId streamId, [FromQuery] bool config = true, [FromQuery] PropertyFilter? filter = null)
+		public async Task<ActionResult<object>> GetStreamAsyncV2(StreamId streamId, [FromQuery] PropertyFilter? filter = null)
 		{
 			IStream? stream = await _streamService.GetStreamAsync(streamId);
 			if (stream == null)
@@ -157,7 +156,33 @@ namespace Horde.Build.Streams
 				return Forbid(AclAction.ViewStream, streamId);
 			}
 
-			return PropertyFilter.Apply(await CreateGetStreamResponseAsyncV2(stream, config, permissionsCache), filter);
+			return PropertyFilter.Apply(await CreateGetStreamResponseAsyncV2(stream, permissionsCache), filter);
+		}
+
+		/// <summary>
+		/// Retrieve config information for a specific stream.
+		/// </summary>
+		/// <param name="streamId">Id of the stream to get information about</param>
+		/// <param name="filter">Filter for the properties to return</param>
+		/// <returns>Information about the requested project</returns>
+		[HttpGet]
+		[Route("/api/v2/streams/{streamId}/config")]
+		[ProducesResponseType(typeof(GetStreamConfigResponseV2), 200)]
+		public async Task<ActionResult<object>> GetStreamConfigAsyncV2(StreamId streamId, [FromQuery] PropertyFilter? filter = null)
+		{
+			IStream? stream = await _streamService.GetStreamAsync(streamId);
+			if (stream == null)
+			{
+				return NotFound(streamId);
+			}
+
+			ProjectPermissionsCache permissionsCache = new ProjectPermissionsCache();
+			if (!await _streamService.AuthorizeAsync(stream, AclAction.ViewStream, User, permissionsCache))
+			{
+				return Forbid(AclAction.ViewStream, streamId);
+			}
+
+			return PropertyFilter.Apply(new GetStreamConfigResponseV2(stream), filter);
 		}
 
 		/// <summary>
@@ -214,10 +239,9 @@ namespace Horde.Build.Streams
 		/// Create a stream response object, including all the templates
 		/// </summary>
 		/// <param name="stream">Stream to create response for</param>
-		/// <param name="config">Whether to include config data in the response</param>
 		/// <param name="cache">Permissions cache</param>
 		/// <returns>Response object</returns>
-		async Task<GetStreamResponseV2> CreateGetStreamResponseAsyncV2(IStream stream, bool config, ProjectPermissionsCache cache)
+		async Task<GetStreamResponseV2> CreateGetStreamResponseAsyncV2(IStream stream, ProjectPermissionsCache cache)
 		{
 			using IScope scope = GlobalTracer.Instance.BuildSpan("CreateGetStreamResponseAsyncV2").StartActive();
 			scope.Span.SetTag("streamId", stream.Id);
@@ -240,7 +264,7 @@ namespace Horde.Build.Streams
 				templateRefResponses.Add(templateRefResponse);
 			}
 
-			return new GetStreamResponseV2(stream, config, templateRefResponses);
+			return new GetStreamResponseV2(stream, templateRefResponses);
 		}
 
 		/// <summary>
