Automated Test Session - fix test results being remove because of being part of several test sessions.

diff --git a/Engine/Source/Programs/Horde/HordeDashboard/src/Plugins/AutomatedTestSession/models/AutomatedTestSession.ts b/Engine/Source/Programs/Horde/HordeDashboard/src/Plugins/AutomatedTestSession/models/AutomatedTestSession.ts
--- a/Engine/Source/Programs/Horde/HordeDashboard/src/Plugins/AutomatedTestSession/models/AutomatedTestSession.ts
+++ b/Engine/Source/Programs/Horde/HordeDashboard/src/Plugins/AutomatedTestSession/models/AutomatedTestSession.ts
@@ -1003,18 +1003,25 @@ export class TestSessionCollection implements Loader {
     }
 
     private indexTestsFromSession(session: TestSessionWrapper) {
-        const tests = session.TestSessionInfo!.Tests;
-        Object.values(tests).some((item) => {
+        const tests = Object.values(session.TestSessionInfo!.Tests);
+        if(tests.length === 0) {
+            return;
+        }
+        const first = this._testcases.get(tests[0].TestUID);
+        if (first && first.SessionName !== session.Name) {
+            return; // Skip conflicting session name, most likely session was renamed.
+        }
+        tests.forEach((item) => {
             const testuid = item.TestUID;
             let testcase = this._testcases.get(testuid);
             if(!testcase) {
                 testcase = new TestCase(item.Name, testuid, item.Suite, session.Name);
                 this._testcases.set(testuid, testcase);
-            } else if (testcase.SessionName !== session.Name) {
-                return true; // Skip conficting session name, most likely session was renamed.
+            }
+            else if(testcase.SessionName !== session.Name) {
+                console.warn(`'${testcase.DisplayName}' is in another session. Appears in '${testcase.SessionName}' and '${session.Name}'.`);
             }
             testcase.addMeta(item);
-            return false;
         });
     }
 
