Horde: Add an index for querying jobs by stream, then template id, then creation time.

diff --git a/Engine/Source/Programs/Horde/Horde.Build/Jobs/JobCollection.cs b/Engine/Source/Programs/Horde/Horde.Build/Jobs/JobCollection.cs
--- a/Engine/Source/Programs/Horde/Horde.Build/Jobs/JobCollection.cs
+++ b/Engine/Source/Programs/Horde/Horde.Build/Jobs/JobCollection.cs
@@ -359,6 +359,7 @@ namespace Horde.Build.Jobs
 		readonly IMongoCollection<JobDocument> _jobs;
 		readonly MongoIndex<JobDocument> _createTimeIndex;
 		readonly MongoIndex<JobDocument> _updateTimeIndex;
+		readonly MongoIndex<JobDocument> _streamThenTemplateThenCreationTimeIndex;
 		readonly IClock _clock;
 		readonly ILogger<JobCollection> _logger;
 
@@ -375,6 +376,7 @@ namespace Horde.Build.Jobs
 
 			List<MongoIndex<JobDocument>> indexes = new List<MongoIndex<JobDocument>>();
 			indexes.Add(keys => keys.Ascending(x => x.StreamId));
+			indexes.Add(_streamThenTemplateThenCreationTimeIndex = MongoIndex.Create<JobDocument>(keys => keys.Ascending(x => x.StreamId).Ascending(x => x.TemplateId).Descending(x => x.CreateTimeUtc)));
 			indexes.Add(keys => keys.Ascending(x => x.Change));
 			indexes.Add(keys => keys.Ascending(x => x.PreflightChange));
 			indexes.Add(_createTimeIndex = MongoIndex.Create<JobDocument>(keys => keys.Descending(x => x.CreateTimeUtc)));
@@ -542,11 +544,7 @@ namespace Horde.Build.Jobs
 		/// <inheritdoc/>
 		public async Task<List<IJob>> FindLatestByStreamWithTemplatesAsync(StreamId streamId, TemplateId[] templates, UserId? preflightStartedByUser, DateTimeOffset? maxCreateTime, DateTimeOffset? modifiedAfter, int? index, int? count, bool consistentRead)
 		{
-			string indexHint = _createTimeIndex.Name;
-			if (modifiedAfter != null)
-			{
-				indexHint = _updateTimeIndex.Name;
-			}
+			string indexHint = _streamThenTemplateThenCreationTimeIndex.Name;
 			
 			// This find call uses an index hint. Modifying the parameter passed to FindAsync can affect execution time a lot as the query planner is forced to use the specified index.
 			// Casting to interface to benefit from default parameter values
