Horde: Adding job abort back to context menu
#rnx
#skipci

diff --git a/Engine/Source/Programs/Horde/HordeDashboard/src/components/JobOperationsContextMenu.tsx b/Engine/Source/Programs/Horde/HordeDashboard/src/components/JobOperationsContextMenu.tsx
--- a/Engine/Source/Programs/Horde/HordeDashboard/src/components/JobOperationsContextMenu.tsx
+++ b/Engine/Source/Programs/Horde/HordeDashboard/src/components/JobOperationsContextMenu.tsx
@@ -2,178 +2,123 @@
 
 import { ContextualMenu, IContextualMenuItem } from '@fluentui/react';
 import React, { useState } from 'react';
-import { GetJobResponse } from '../backend/Api';
+import { GetJobResponse, JobState } from '../backend/Api';
 import dashboard from '../backend/Dashboard';
-import { JobDetails } from '../backend/JobDetails';
 import { projectStore } from '../backend/ProjectStore';
-// import { AbortJobModal } from './AbortJobModal';
-
-enum ParameterState {
-    Hidden,
-    Parameters,
-    Clone
-}
+import { AbortJobModal } from './jobDetailsV2/AbortJobModal';
+// import { EditJobModal } from './EditJobModal';
 
 export const JobOperationsContextMenu: React.FC<{ job: GetJobResponse, children?: any }> = ({ job, children }) => {
 
-    const [state, setState] = useState<{ pos?: { x: number, y: number } }>({});
-    const [details, setDetails] = useState<JobDetails | undefined>(undefined);
-
-
-    // const [abortShown, setAbortShown] = useState(false);
-    //const [editShown, setEditShown] = useState(false);
-    const [, setParametersState] = useState(ParameterState.Hidden);
-
-
-    const copyToClipboard = (value: string | undefined) => {
-
-        if (!value) {
-            return;
-        }
-
-        const el = document.createElement('textarea');
-        el.value = value;
-        document.body.appendChild(el);
-        el.select();
-        document.execCommand('copy');
-        document.body.removeChild(el);
-    }
-
-    const pinned = dashboard.jobPinned(job.id);
-
-    let change = job.change!.toString();
-    if (job.preflightChange) {
-        change = job.preflightChange.toString();
-    }
-
-    let menuItems: IContextualMenuItem[] = [
-        {
-            key: 'open_in_new_window',
-            text: 'Open in New Tab',
-            onClick: () => window.open(`/job/${job.id}`),
-        },
-        {
-            key: 'pin_job',
-            text: pinned ? 'Unpin from Home' : "Pin to Home",
-            onClick: () => { pinned ? dashboard.unpinJob(job.id) : dashboard.pinJob(job.id) },
-        },
-        {
-            key: 'copy_to_clipboard',
-            text: 'Copy Job to Clipboard',
-            onClick: () => copyToClipboard(window.location.protocol + "//" + window.location.hostname + `/job/${job.id}`)
-        },
-        {
-            key: 'copy_cl_to_clipboard',
-            text: 'Copy CL to Clipboard',
-            onClick: () => copyToClipboard(change),
-        },
-        {
-            key: 'open_in_swarm', text: 'Open CL in Swarm', onClick: () => {
-                window.open(`${dashboard.swarmUrl}/changes/${change}`);
-            }
-        },
-        {        
-            key: 'open_in_swarm_range', text: 'Open CL Range in Swarm', onClick: () => {
-
-                const stream = projectStore.streamById(job.streamId)!;
-                const project = projectStore.byId(stream!.projectId)!;
-                const name = project.name === "Engine" ? "UE4" : project.name;
-                const rangeUrl = `${dashboard.swarmUrl}/files/${name}/${stream.name}?range=@${change}#commits`;
-    
-                window.open(rangeUrl)
-            }
-        }
-
-    ];
-
-    // const abortDisabled = details?.jobdata?.state === JobState.Complete;
-    // const editDisabled = details?.jobdata?.state === JobState.Complete;
-
-    if (details) {
-
-        menuItems.push({
-            key: 'view_paremeters',
-            text: 'View Parameters',
-            onClick: () => setParametersState(ParameterState.Parameters),
-        });
-
-       /*
-        if (!editDisabled) {
-            menuItems.push({
-                key: 'edit_job',
-                text: 'Edit',
-                onClick: () => setEditShown(true),
-            });
-        } */
-
-       /*
-        if (!abortDisabled) {
-            menuItems.push({
-                key: 'abort_job',
-                text: 'Abort Job',
-                onClick: () => setAbortShown(true),
-            });
-        } */
-
-        menuItems.push({
-            key: 'run_again',
-            text: 'Run Again',
-            onClick: () => setParametersState(ParameterState.Clone),
-        });
-
-    }
-
-    if (state.pos && !details) {
-        menuItems = menuItems.slice(0, 6);
-        menuItems.push({
-            key: 'loading_job',
-            text: 'Loading...',
-        });
+   const [state, setState] = useState<{ pos?: { x: number, y: number } }>({});
+
+
+   const [abortShown, setAbortShown] = useState(false);
+   // const [editShown, setEditShown] = useState(false);
+
+   const copyToClipboard = (value: string | undefined) => {
+
+      if (!value) {
+         return;
+      }
+
+      const el = document.createElement('textarea');
+      el.value = value;
+      document.body.appendChild(el);
+      el.select();
+      document.execCommand('copy');
+      document.body.removeChild(el);
    }
-   
-   // removed due to JobDetails
-   /*
-           { !!details && <NewBuild streamId={details.stream!.id} jobDetails={details} readOnly={parametersState === ParameterState.Parameters}
-            show={parametersState === ParameterState.Parameters || parametersState === ParameterState.Clone}
-            onClose={() => { setParametersState(ParameterState.Hidden); }} />}
 
-             { !abortDisabled && !!details && <AbortJobModal jobDetails={details} show={abortShown} onClose={() => { setAbortShown(false); }} />}
-   */
-   
-   // Also, removed due to job details refactor
+   const pinned = dashboard.jobPinned(job.id);
+
+   let change = job.change!.toString();
+   if (job.preflightChange) {
+      change = job.preflightChange.toString();
+   }
+
+   let menuItems: IContextualMenuItem[] = [
+      {
+         key: 'open_in_new_window',
+         text: 'Open in New Tab',
+         onClick: () => window.open(`/job/${job.id}`),
+      },
+      {
+         key: 'pin_job',
+         text: pinned ? 'Unpin from Home' : "Pin to Home",
+         onClick: () => { pinned ? dashboard.unpinJob(job.id) : dashboard.pinJob(job.id) },
+      },
+      {
+         key: 'copy_to_clipboard',
+         text: 'Copy Job to Clipboard',
+         onClick: () => copyToClipboard(window.location.protocol + "//" + window.location.hostname + `/job/${job.id}`)
+      },
+      {
+         key: 'copy_cl_to_clipboard',
+         text: 'Copy CL to Clipboard',
+         onClick: () => copyToClipboard(change),
+      },
+      {
+         key: 'open_in_swarm', text: 'Open CL in Swarm', onClick: () => {
+            window.open(`${dashboard.swarmUrl}/changes/${change}`);
+         }
+      },
+      {
+         key: 'open_in_swarm_range', text: 'Open CL Range in Swarm', onClick: () => {
+
+            const stream = projectStore.streamById(job.streamId)!;
+            const project = projectStore.byId(stream!.projectId)!;
+            const name = project.name === "Engine" ? "UE4" : project.name;
+            const rangeUrl = `${dashboard.swarmUrl}/files/${name}/${stream.name}?range=@${change}#commits`;
+
+            window.open(rangeUrl)
+         }
+      }
+
+   ];
+
+   const abortDisabled = job.state === JobState.Complete;
+
    /*
-      { !editDisabled && !!details && <EditJobModal jobDetails={details} show={editShown} onClose={() => { setEditShown(false); }} />}
+   const editDisabled = job.state === JobState.Complete;
+   {!editDisabled && <EditJobModal jobData={job} show={editShown} onClose={() => { setEditShown(false); }} />}
+   if (!editDisabled) {
+      menuItems.push({
+         key: 'edit_job',
+         text: 'Edit',
+         onClick: () => setEditShown(true),
+      });
+   }
    */
 
-    return (<div onContextMenu={(ev) => {
-        ev.preventDefault();
-
-        // chrome nulls out the event properties so capture them
-        const pos = { x: ev.clientX, y: ev.clientY };
+   if (!abortDisabled) {
+      menuItems.push({
+         key: 'abort_job',
+         text: 'Abort Job',
+         onClick: () => setAbortShown(true),
+      });
+   }
 
-        if (!details) {
-            const details = new JobDetails();
 
-            details.set(job.id, undefined, undefined, undefined, undefined, () => {
-                details.stopPolling();
+   return (<div onContextMenu={(ev) => {
+      ev.preventDefault();
 
-                setDetails(details);
+      // chrome nulls out the event properties so capture them
+      const pos = { x: ev.clientX, y: ev.clientY };
 
-            });
-        }
+      setState({ pos: pos });
 
-        setState({ pos: pos });
+      return false;
+   }}>
+      {!abortDisabled && <AbortJobModal jobDataIn={job} show={abortShown} onClose={() => { setAbortShown(false); }} />}
 
-        return false;
-    }}>
-               
-        <ContextualMenu
-            styles={{ list: { selectors: { '.ms-ContextualMenu-itemText': { fontSize: "10px" } } } }}
-            items={menuItems}
-            hidden={(!state.pos)}
-            target={state.pos!}
-            onItemClick={(ev) => { setState({}); ev?.preventDefault(); }}
-            onDismiss={(ev) => { setState({}); ev?.preventDefault(); }} />
-        {children}
-    </div>)
+      <ContextualMenu
+         styles={{ list: { selectors: { '.ms-ContextualMenu-itemText': { fontSize: "10px" } } } }}
+         items={menuItems}
+         hidden={(!state.pos)}
+         target={state.pos!}
+         onItemClick={(ev) => { setState({}); ev?.preventDefault(); }}
+         onDismiss={(ev) => { setState({}); ev?.preventDefault(); }} />
+      {children}
+   </div>)
 };
\ No newline at end of file
diff --git a/Engine/Source/Programs/Horde/HordeDashboard/src/components/jobDetailsV2/AbortJobModal.tsx b/Engine/Source/Programs/Horde/HordeDashboard/src/components/jobDetailsV2/AbortJobModal.tsx
--- a/Engine/Source/Programs/Horde/HordeDashboard/src/components/jobDetailsV2/AbortJobModal.tsx
+++ b/Engine/Source/Programs/Horde/HordeDashboard/src/components/jobDetailsV2/AbortJobModal.tsx
@@ -2,84 +2,85 @@
 import { DefaultButton, IconButton, Modal, PrimaryButton, Stack, Text } from '@fluentui/react';
 import React from 'react';
 import backend from '../../backend';
+import { GetJobResponse } from '../../backend/Api';
 import { hordeClasses } from '../../styles/Styles';
 import { JobDetailsV2 } from './JobDetailsViewCommon';
 
 
-export const AbortJobModal: React.FC<{ jobDetails: JobDetailsV2; stepId?: string, show: boolean; onClose: () => void }> = ({ jobDetails, stepId, show, onClose }) => {
+export const AbortJobModal: React.FC<{ jobDetails?: JobDetailsV2; jobDataIn?: GetJobResponse; stepId?: string, show: boolean; onClose: () => void }> = ({ jobDetails, jobDataIn, stepId, show, onClose }) => {
 
-   const jobData = jobDetails.jobData;
+   const jobData = jobDetails?.jobData ?? jobDataIn;
    if (!jobData) {
       console.error("Attempting to display abort job modal with no job data")
       return null;
    }
 
-    let headerText = `Abort ${jobData.name} Job?`;
-
-    if (stepId) {
-        headerText = `Abort ${jobDetails.getStepName(stepId)} Step?`;
-    }
-
-    const close = () => {
-        onClose();
-    };
-
-    const onAbort = async () => {
-
-        if (!stepId) {
-            await backend.updateJob(jobData.id, { aborted: true }).then((response) => {
-                console.log("Job aborted", response);
-            }).catch((reason) => {
-                // @todo: error ui
-                console.error(reason);
-            }).finally(() => {
-                close();
-            });
-        } else {
-
-            const batch = jobDetails.batchByStepId(stepId);
-
-            if (!batch) {
-                console.error("Unable to get batch id for step abort");
-                return;
-            }
-
-            await backend.updateJobStep(jobData.id, batch.id, stepId, { abortRequested: true }).then((response) => {
-                console.log("Job step aborted", response);
-            }).catch((reason) => {
-                // @todo: error ui
-                console.error(reason);
-            }).finally(() => {
-                close();
-            });
-
-        }
-
-    };
-
-    const height = 140;
-
-    return <Modal isOpen={show} className={hordeClasses.modal} styles={{ main: { padding: 8, width: 540, height: height, minHeight: height } }} onDismiss={() => { if (show) close() }}>
-        <Stack horizontal styles={{ root: { padding: 8 } }}>
-            <Stack.Item grow={2}>
-                <Text variant="mediumPlus">{headerText}</Text>
-            </Stack.Item>
-            <Stack.Item grow={0}>
-                <IconButton
-                    iconProps={{ iconName: 'Cancel' }}
-                    ariaLabel="Close popup modal"
-                    onClick={() => { close(); }}
-                />
-            </Stack.Item>
-        </Stack>
-
-        <Stack styles={{ root: { padding: 8 } }}>
-            <Stack horizontal tokens={{ childrenGap: 16 }} styles={{ root: { paddingTop: 12, paddingLeft: 8, paddingBottom: 8 } }}>
-                <Stack grow />
-                <PrimaryButton text="Abort" disabled={false} onClick={() => { onAbort(); }} />
-                <DefaultButton text="Cancel" disabled={false} onClick={() => { close(); }} />
-            </Stack>
-        </Stack>
-    </Modal>;
+   let headerText = `Abort ${jobData.name} Job?`;
+
+   if (jobDetails && stepId) {
+      headerText = `Abort ${jobDetails.getStepName(stepId)} Step?`;
+   }
+
+   const close = () => {
+      onClose();
+   };
+
+   const onAbort = async () => {
+
+      if (!stepId) {
+         await backend.updateJob(jobData.id, { aborted: true }).then((response) => {
+            console.log("Job aborted", response);
+         }).catch((reason) => {
+            // @todo: error ui
+            console.error(reason);
+         }).finally(() => {
+            close();
+         });
+      } else {
+
+         const batch = jobDetails?.batchByStepId(stepId);
+
+         if (!batch) {
+            console.error("Unable to get batch id for step abort");
+            return;
+         }
+
+         await backend.updateJobStep(jobData.id, batch.id, stepId, { abortRequested: true }).then((response) => {
+            console.log("Job step aborted", response);
+         }).catch((reason) => {
+            // @todo: error ui
+            console.error(reason);
+         }).finally(() => {
+            close();
+         });
+
+      }
+
+   };
+
+   const height = 140;
+
+   return <Modal isOpen={show} className={hordeClasses.modal} styles={{ main: { padding: 8, width: 540, height: height, minHeight: height } }} onDismiss={() => { if (show) close() }}>
+      <Stack horizontal styles={{ root: { padding: 8 } }}>
+         <Stack.Item grow={2}>
+            <Text variant="mediumPlus">{headerText}</Text>
+         </Stack.Item>
+         <Stack.Item grow={0}>
+            <IconButton
+               iconProps={{ iconName: 'Cancel' }}
+               ariaLabel="Close popup modal"
+               onClick={() => { close(); }}
+            />
+         </Stack.Item>
+      </Stack>
+
+      <Stack styles={{ root: { padding: 8 } }}>
+         <Stack horizontal tokens={{ childrenGap: 16 }} styles={{ root: { paddingTop: 12, paddingLeft: 8, paddingBottom: 8 } }}>
+            <Stack grow />
+            <PrimaryButton text="Abort" disabled={false} onClick={() => { onAbort(); }} />
+            <DefaultButton text="Cancel" disabled={false} onClick={() => { close(); }} />
+         </Stack>
+      </Stack>
+   </Modal>;
 
 };
\ No newline at end of file
