Horde: Return the stream config within the stream responses. The dashboard to pass in a config revision string, and the response will excude it if the revision matches.

diff --git a/Engine/Source/Programs/Horde/Horde.Build/Streams/StreamMessages.cs b/Engine/Source/Programs/Horde/Horde.Build/Streams/StreamMessages.cs
--- a/Engine/Source/Programs/Horde/Horde.Build/Streams/StreamMessages.cs
+++ b/Engine/Source/Programs/Horde/Horde.Build/Streams/StreamMessages.cs
@@ -14,6 +14,7 @@ using Horde.Build.Perforce;
 using Horde.Build.Projects;
 using Horde.Build.Jobs;
 using Microsoft.AspNetCore.DataProtection;
+using System.Linq;
 
 namespace Horde.Build.Streams
 {
@@ -438,25 +439,6 @@ namespace Horde.Build.Streams
 		}
 	}
 
-	/// <summary>
-	/// Returns the configuration data for a stream
-	/// </summary>
-	public class GetStreamConfigResponseV2
-	{
-		readonly IStream _stream;
-
-		/// <inheritdoc cref="IStream.ConfigRevision"/>
-		public string Revision => _stream.ConfigRevision;
-
-		/// <inheritdoc cref="IStream.Config"/>
-		public StreamConfig Config => _stream.Config;
-
-		internal GetStreamConfigResponseV2(IStream stream)
-		{
-			_stream = stream;
-		}
-	}
-
 	/// <summary>
 	/// Response describing a stream
 	/// </summary>
@@ -473,9 +455,12 @@ namespace Horde.Build.Streams
 		/// <inheritdoc cref="IStream.Name"/>
 		public string Name => _stream.Name;
 
-		/// <inheritdoc cref="IStream.Name"/>
+		/// <inheritdoc cref="IStream.ConfigRevision"/>
 		public string ConfigRevision => _stream.ConfigRevision;
 
+		/// <inheritdoc cref="IStream.Config"/>
+		public StreamConfig? Config { get; }
+
 		/// <inheritdoc cref="IStream.Deleted"/>
 		public bool Deleted => _stream.Deleted;
 
@@ -488,10 +473,16 @@ namespace Horde.Build.Streams
 		/// <inheritdoc cref="IStream.PauseComment"/>
 		public string? PauseComment => _stream.PauseComment;
 
-		internal GetStreamResponseV2(IStream stream, List<GetTemplateRefResponseV2> templates)
+		internal GetStreamResponseV2(IStream stream, bool includeConfig)
 		{
 			_stream = stream;
-			Templates = templates;
+
+			Templates = stream.Templates.Values.Select(x => new GetTemplateRefResponseV2(x)).ToList();
+
+			if (includeConfig)
+			{
+				Config = stream.Config;
+			}
 		}
 	}
 
@@ -508,22 +499,9 @@ namespace Horde.Build.Streams
 		/// <inheritdoc cref="ITemplateRef.Hash"/>
 		public ContentHash Hash => _templateRef.Hash;
 
-		/// <inheritdoc cref="ITemplateRef.Schedule"/>
-		public GetTemplateScheduleResponseV2? Schedule { get; }
-
-		/// <inheritdoc cref="ITemplateRef.StepStates"/>
-		public List<GetTemplateStepResponseV2>? StepStates { get; }
-
-		internal GetTemplateRefResponseV2(ITemplateRef templateRef, List<GetTemplateStepResponseV2>? stepStates)
+		internal GetTemplateRefResponseV2(ITemplateRef templateRef)
 		{
 			_templateRef = templateRef;
-
-			StepStates = stepStates;
-
-			if (templateRef.Schedule != null)
-			{
-				Schedule = new GetTemplateScheduleResponseV2(templateRef.Schedule);
-			}
 		}
 	}
 
diff --git a/Engine/Source/Programs/Horde/Horde.Build/Streams/StreamsController.cs b/Engine/Source/Programs/Horde/Horde.Build/Streams/StreamsController.cs
--- a/Engine/Source/Programs/Horde/Horde.Build/Streams/StreamsController.cs
+++ b/Engine/Source/Programs/Horde/Horde.Build/Streams/StreamsController.cs
@@ -100,7 +100,7 @@ namespace Horde.Build.Streams
 			{
 				if (await _streamService.AuthorizeAsync(stream, AclAction.ViewStream, User, permissionsCache))
 				{
-					GetStreamResponseV2 response = await CreateGetStreamResponseAsyncV2(stream, permissionsCache);
+					GetStreamResponseV2 response = new GetStreamResponseV2(stream, false);
 					responses.Add(response);
 				}
 			}
@@ -137,12 +137,13 @@ namespace Horde.Build.Streams
 		/// Retrieve information about a specific stream.
 		/// </summary>
 		/// <param name="streamId">Id of the stream to get information about</param>
+		/// <param name="config">The client's cached config revision. If this matches value matches the current config version of the stream, the config object will be omitted from the response.</param>
 		/// <param name="filter">Filter for the properties to return</param>
 		/// <returns>Information about the requested project</returns>
 		[HttpGet]
 		[Route("/api/v2/streams/{streamId}")]
 		[ProducesResponseType(typeof(GetStreamResponseV2), 200)]
-		public async Task<ActionResult<object>> GetStreamAsyncV2(StreamId streamId, [FromQuery] PropertyFilter? filter = null)
+		public async Task<ActionResult<object>> GetStreamAsyncV2(StreamId streamId, [FromQuery] string? config = null, [FromQuery] PropertyFilter? filter = null)
 		{
 			IStream? stream = await _streamService.GetStreamAsync(streamId);
 			if (stream == null)
@@ -156,33 +157,8 @@ namespace Horde.Build.Streams
 				return Forbid(AclAction.ViewStream, streamId);
 			}
 
-			return PropertyFilter.Apply(await CreateGetStreamResponseAsyncV2(stream, permissionsCache), filter);
-		}
-
-		/// <summary>
-		/// Retrieve config information for a specific stream.
-		/// </summary>
-		/// <param name="streamId">Id of the stream to get information about</param>
-		/// <param name="filter">Filter for the properties to return</param>
-		/// <returns>Information about the requested project</returns>
-		[HttpGet]
-		[Route("/api/v2/streams/{streamId}/config")]
-		[ProducesResponseType(typeof(GetStreamConfigResponseV2), 200)]
-		public async Task<ActionResult<object>> GetStreamConfigAsyncV2(StreamId streamId, [FromQuery] PropertyFilter? filter = null)
-		{
-			IStream? stream = await _streamService.GetStreamAsync(streamId);
-			if (stream == null)
-			{
-				return NotFound(streamId);
-			}
-
-			ProjectPermissionsCache permissionsCache = new ProjectPermissionsCache();
-			if (!await _streamService.AuthorizeAsync(stream, AclAction.ViewStream, User, permissionsCache))
-			{
-				return Forbid(AclAction.ViewStream, streamId);
-			}
-
-			return PropertyFilter.Apply(new GetStreamConfigResponseV2(stream), filter);
+			bool includeConfig = !String.Equals(config, stream.ConfigRevision, StringComparison.Ordinal); 
+			return PropertyFilter.Apply(new GetStreamResponseV2(stream, includeConfig), filter);
 		}
 
 		/// <summary>
@@ -235,38 +211,6 @@ namespace Horde.Build.Streams
 			return stream.ToApiResponse(bIncludeAcl, apiTemplateRefs);
 		}
 
-		/// <summary>
-		/// Create a stream response object, including all the templates
-		/// </summary>
-		/// <param name="stream">Stream to create response for</param>
-		/// <param name="cache">Permissions cache</param>
-		/// <returns>Response object</returns>
-		async Task<GetStreamResponseV2> CreateGetStreamResponseAsyncV2(IStream stream, ProjectPermissionsCache cache)
-		{
-			using IScope scope = GlobalTracer.Instance.BuildSpan("CreateGetStreamResponseAsyncV2").StartActive();
-			scope.Span.SetTag("streamId", stream.Id);
-
-			List<GetTemplateRefResponseV2> templateRefResponses = new List<GetTemplateRefResponseV2>();
-			foreach (ITemplateRef templateRef in stream.Templates.Values)
-			{
-				List<GetTemplateStepResponseV2>? templateStepResponses = null;
-				if (templateRef.StepStates != null && templateRef.StepStates.Count > 0)
-				{
-					templateStepResponses = new List<GetTemplateStepResponseV2>();
-					foreach (ITemplateStep templateStep in templateRef.StepStates)
-					{
-						GetThinUserInfoResponse? pausedByUserInfo = new GetThinUserInfoResponse(await _userCollection.GetCachedUserAsync(templateStep.PausedByUserId));
-						templateStepResponses.Add(new GetTemplateStepResponseV2(templateStep, pausedByUserInfo));
-					}
-				}
-
-				GetTemplateRefResponseV2 templateRefResponse = new GetTemplateRefResponseV2(templateRef, templateStepResponses);
-				templateRefResponses.Add(templateRefResponse);
-			}
-
-			return new GetStreamResponseV2(stream, templateRefResponses);
-		}
-
 		/// <summary>
 		/// Gets a list of changes for a stream
 		/// </summary>
