Horde Storage - Removed unused settings

diff --git a/Engine/Source/Programs/Horde/HordeStorage/Horde.Storage/appsettings.json b/Engine/Source/Programs/Horde/HordeStorage/Horde.Storage/appsettings.json
--- a/Engine/Source/Programs/Horde/HordeStorage/Horde.Storage/appsettings.json
+++ b/Engine/Source/Programs/Horde/HordeStorage/Horde.Storage/appsettings.json
@@ -1,18 +1,3 @@
 {
-  "AllowedHosts": "*",
-
-  "Cache.Blob": {
-    // limit cache to 1 GB
-    "SizeLimit": "1073741824"
-  },
-  "Cache.Tree": {
-    // limit cache to 1 GB
-    "SizeLimit": "1073741824"
-  },
-
-  "GRPC": {
-    // bump send and receive size to 1gb
-    "MaxReceiveMessageSize": "1073741824",
-    "MaxSendMessageSize": "1073741824"
-  }
+  "AllowedHosts": "*"
 }
