Horde Storage - Improved error message when latency check fails to resolve url

diff --git a/Engine/Source/Programs/Horde/HordeStorage/Horde.Storage/Implementation/PeerStatusService.cs b/Engine/Source/Programs/Horde/HordeStorage/Horde.Storage/Implementation/PeerStatusService.cs
--- a/Engine/Source/Programs/Horde/HordeStorage/Horde.Storage/Implementation/PeerStatusService.cs
+++ b/Engine/Source/Programs/Horde/HordeStorage/Horde.Storage/Implementation/PeerStatusService.cs
@@ -164,7 +164,7 @@ namespace Horde.Storage.Implementation
                 // ignore error responses as they may not have reached the actual instance
                 if (!result.IsSuccessStatusCode)
                 {
-                    _logger.Warning("Non-success status code ({StatusCode}) when attempting to measure latency to {Endpoint}", result.StatusCode, endpoint.Url);
+                    _logger.Warning("Non-success status code ({StatusCode}) when attempting to measure latency to {Endpoint}", result.StatusCode, uri);
                     return int.MaxValue;
                 }
 
@@ -172,7 +172,7 @@ namespace Horde.Storage.Implementation
             }
             catch (Exception e)
             {
-                _logger.Warning(e, "Exception when attempting to measure latency from {Endpoint}", endpoint.Url);
+                _logger.Warning(e, "Exception when attempting to measure latency from {Endpoint}", uri);
                 // error reaching the endpoint is just considered to max latency
                 return int.MaxValue;
             }
