Horde Storage - Fixed exception when starting horde storage that some users experienced.

diff --git a/Engine/Source/Programs/Horde/HordeStorage/Jupiter.Common/BaseStartup.cs b/Engine/Source/Programs/Horde/HordeStorage/Jupiter.Common/BaseStartup.cs
--- a/Engine/Source/Programs/Horde/HordeStorage/Jupiter.Common/BaseStartup.cs
+++ b/Engine/Source/Programs/Horde/HordeStorage/Jupiter.Common/BaseStartup.cs
@@ -59,7 +59,6 @@ namespace Jupiter
             CbConvertersAspNet.AddAspnetConverters();
 
             services.AddServerTiming();
-            services.AddSuppressExceptionsMiddleware();
 
             // aws specific settings
             services.AddOptions<AWSCredentialsSettings>().Bind(Configuration.GetSection("AWSCredentials")).ValidateDataAnnotations();
@@ -355,7 +354,7 @@ namespace Jupiter
             app.UseAuthorization();
 
             app.UseMiddleware<SuppressExceptionMiddleware>();
-			app.UseMiddleware<ServerTimingMiddleware>();
+            app.UseMiddleware<ServerTimingMiddleware>();
 
             //app.UseMiddleware<DatadogTraceMiddleware>("Endpoints");
             app.UseEndpoints(endpoints =>
diff --git a/Engine/Source/Programs/Horde/HordeStorage/Jupiter.Common/SuppressExceptionMiddleware.cs b/Engine/Source/Programs/Horde/HordeStorage/Jupiter.Common/SuppressExceptionMiddleware.cs
--- a/Engine/Source/Programs/Horde/HordeStorage/Jupiter.Common/SuppressExceptionMiddleware.cs
+++ b/Engine/Source/Programs/Horde/HordeStorage/Jupiter.Common/SuppressExceptionMiddleware.cs
@@ -4,15 +4,14 @@ using System;
 using System.Threading.Tasks;
 using EpicGames.AspNet;
 using Microsoft.AspNetCore.Http;
-using Microsoft.Extensions.DependencyInjection;
 using Serilog;
 
 namespace Jupiter.Common
 {
     public class SuppressExceptionMiddleware
-	{
+    {
         private readonly ILogger _logger = Log.ForContext<SuppressExceptionMiddleware>();
-		private readonly RequestDelegate _next;
+        private readonly RequestDelegate _next;
 
         public SuppressExceptionMiddleware(RequestDelegate next)
         {
@@ -48,14 +47,4 @@ namespace Jupiter.Common
             }
         }
     }
-
-	public static class SuppressExceptionMiddlewareCollectionExtensions
-	{
-        public static IServiceCollection AddSuppressExceptionsMiddleware(this IServiceCollection services)
-        {
-            services.AddSingleton<SuppressExceptionMiddleware>();
-
-            return services;
-        }
-    }
 }
