Horde Storage - re-added fallback for the internal jupiter namespace so that we do not have to configure it. It can still be explicitly overriden should we want to.

diff --git a/Engine/Source/Programs/Horde/HordeStorage/Jupiter.Common/Implementation/NamespacePolicyResolver.cs b/Engine/Source/Programs/Horde/HordeStorage/Jupiter.Common/Implementation/NamespacePolicyResolver.cs
--- a/Engine/Source/Programs/Horde/HordeStorage/Jupiter.Common/Implementation/NamespacePolicyResolver.cs
+++ b/Engine/Source/Programs/Horde/HordeStorage/Jupiter.Common/Implementation/NamespacePolicyResolver.cs
@@ -15,10 +15,15 @@ namespace Jupiter.Common
     public class NamespacePolicyResolver : INamespacePolicyResolver
     {
         private readonly IOptionsMonitor<NamespaceSettings> _namespaceSettings;
+        private readonly NamespacePolicy _internalNamespacePolicy;
 
         public NamespacePolicyResolver(IOptionsMonitor<NamespaceSettings> namespaceSettings)
         {
             _namespaceSettings = namespaceSettings;
+            _internalNamespacePolicy = new NamespacePolicy()
+            {
+                StoragePool = "persistent"
+            };
         }
 
         public NamespacePolicy GetPoliciesForNs(NamespaceId ns)
@@ -28,7 +33,7 @@ namespace Jupiter.Common
             {
                 return settings;
             }
-
+            
             // attempt to find the default mapping
             if (_namespaceSettings.CurrentValue.Policies.TryGetValue("*",
                     out NamespacePolicy? defaultSettings))
@@ -36,6 +41,12 @@ namespace Jupiter.Common
                 return defaultSettings;
             }
 
+            // if no override has been specified for the internal jupiter namespace then we default it
+            if (ns == INamespacePolicyResolver.JupiterInternalNamespace)
+            {
+                return _internalNamespacePolicy;
+            }
+
             throw new UnknownNamespaceException($"Unable to find a valid policy for namespace {ns}");
         }
     }
