Horde Storage - Changed order we delete records from the blob index vs blob store to resolve inconsistency between blob index state and blob store when we have a timeout from the db.

diff --git a/Engine/Source/Programs/Horde/HordeStorage/Horde.Storage/Implementation/Blob/IBlobService.cs b/Engine/Source/Programs/Horde/HordeStorage/Horde.Storage/Implementation/Blob/IBlobService.cs
--- a/Engine/Source/Programs/Horde/HordeStorage/Horde.Storage/Implementation/Blob/IBlobService.cs
+++ b/Engine/Source/Programs/Horde/HordeStorage/Horde.Storage/Implementation/Blob/IBlobService.cs
@@ -425,6 +425,11 @@ public class BlobService : IBlobService
         bool blobNotFound = false;
         bool deletedAtLeastOnce = false;
 
+        // remove the object from the tracking first, if this times out we do not want to end up with a inconsistent blob index
+        // if the blob store delete fails on the other hand we will still run a delete again during GC (as the blob is still orphaned at that point)
+        // this assumes that blob gc is based on scanning the root blob store
+        await _blobIndex.RemoveBlobFromRegion(ns, blob);
+
         foreach (IBlobStore store in _blobStores)
         {
             try
@@ -445,8 +450,6 @@ public class BlobService : IBlobService
             }
         }
 
-        await _blobIndex.RemoveBlobFromRegion(ns, blob);
-
         if (deletedAtLeastOnce)
         {
             return;
