Horde: Fix route for v2 streams config endpoint
#rnx

diff --git a/Engine/Source/Programs/Horde/Horde.Build/Streams/StreamsController.cs b/Engine/Source/Programs/Horde/Horde.Build/Streams/StreamsController.cs
--- a/Engine/Source/Programs/Horde/Horde.Build/Streams/StreamsController.cs
+++ b/Engine/Source/Programs/Horde/Horde.Build/Streams/StreamsController.cs
@@ -141,7 +141,7 @@ namespace Horde.Build.Streams
 		/// <param name="filter">Filter for the properties to return</param>
 		/// <returns>Information about the requested project</returns>
 		[HttpGet]
-		[Route("/api/v2/streams/{streamId}")]
+		[Route("/api/v2/streams/{streamId}/config")]
 		[ProducesResponseType(typeof(GetStreamResponseV2), 200)]
 		public async Task<ActionResult<object>> GetStreamAsyncV2(StreamId streamId, [FromQuery] string? config = null, [FromQuery] PropertyFilter? filter = null)
 		{
