Horde: Fix bug with issue updates always being disabled
#rnx

diff --git a/Engine/Source/Programs/Horde/Horde.Build/Jobs/JobCollection.cs b/Engine/Source/Programs/Horde/Horde.Build/Jobs/JobCollection.cs
--- a/Engine/Source/Programs/Horde/Horde.Build/Jobs/JobCollection.cs
+++ b/Engine/Source/Programs/Horde/Horde.Build/Jobs/JobCollection.cs
@@ -319,7 +319,7 @@ namespace Horde.Build.Jobs
 				StartedByUserId = options.StartedByUserId;
 				Priority = options.Priority ?? HordeCommon.Priority.Normal;
 				AutoSubmit = options.AutoSubmit ?? false;
-				UpdateIssues = options.UpdateIssues ?? (options.StartedByUserId == null && options.PreflightChange == 0);
+				UpdateIssues = options.UpdateIssues ?? (options.StartedByUserId == null && ( options.PreflightChange == 0 || options.PreflightChange == null));
 				PromoteIssuesByDefault = options.PromoteIssuesByDefault ?? false;
 				CreateTimeUtc = createTimeUtc;
 				ChainedJobs.AddRange(options.JobTriggers.Select(x => new ChainedJobDocument(x)));
