Horde: Don't log unknown device platform errors, surfaced in logs themselves
#rnx

diff --git a/Engine/Source/Programs/Horde/Horde.Build/Devices/DevicesController.cs b/Engine/Source/Programs/Horde/Horde.Build/Devices/DevicesController.cs
--- a/Engine/Source/Programs/Horde/Horde.Build/Devices/DevicesController.cs
+++ b/Engine/Source/Programs/Horde/Horde.Build/Devices/DevicesController.cs
@@ -884,8 +884,7 @@ namespace Horde.Build.Devices
 
 				if (platform == null)
 				{
-					string message = $"Unknown platform {platformId}" + details;
-					_logger.LogError("Unknown platform {PlatformName} {PlatformId} {Details}", platformName, platformId, details);
+					string message = $"Unknown platform {platformId}" + details;					
 					await _deviceService.NotifyDeviceServiceAsync(message, null, request.JobId, request.StepId);
 					return BadRequest(message);
 				}
