Horde: Use dependency-injected EC2 client in fleet manager

Custom setup of client inside class makes testing more difficult and makes for more places of instantiation/config.

diff --git a/Engine/Source/Programs/Horde/Horde.Build.Tests/TestSetup.cs b/Engine/Source/Programs/Horde/Horde.Build.Tests/TestSetup.cs
--- a/Engine/Source/Programs/Horde/Horde.Build.Tests/TestSetup.cs
+++ b/Engine/Source/Programs/Horde/Horde.Build.Tests/TestSetup.cs
@@ -6,6 +6,7 @@ using System.Reflection;
 using System.Security.Claims;
 using System.Text.Json;
 using System.Threading.Tasks;
+using Amazon.EC2;
 using Datadog.Trace;
 using EpicGames.Core;
 using EpicGames.Horde.Storage;
@@ -59,6 +60,7 @@ using Horde.Build.Auditing;
 using Horde.Build.Storage.Backends;
 using Horde.Build.Compute;
 using Horde.Build.Devices;
+using Moq;
 
 namespace Horde.Build.Tests
 {
@@ -198,6 +200,9 @@ namespace Horde.Build.Tests
 			services.AddSingleton<FakeClock>();
 			services.AddSingleton<IClock>(sp => sp.GetRequiredService<FakeClock>());
 			services.AddSingleton<IHostApplicationLifetime, AppLifetimeStub>();
+			
+			// Empty mocked object to satisfy basic test runs
+			services.AddSingleton<IAmazonEC2>(sp => new Mock<IAmazonEC2>().Object);
 
 			services.AddSingleton<AclService>();
 			services.AddSingleton<AgentService>();
diff --git a/Engine/Source/Programs/Horde/Horde.Build/Agents/Fleet/Providers/AwsReuseFleetManager.cs b/Engine/Source/Programs/Horde/Horde.Build/Agents/Fleet/Providers/AwsReuseFleetManager.cs
--- a/Engine/Source/Programs/Horde/Horde.Build/Agents/Fleet/Providers/AwsReuseFleetManager.cs
+++ b/Engine/Source/Programs/Horde/Horde.Build/Agents/Fleet/Providers/AwsReuseFleetManager.cs
@@ -5,11 +5,9 @@ using System.Collections.Generic;
 using System.Linq;
 using System.Threading;
 using System.Threading.Tasks;
-using Amazon;
 using Amazon.EC2;
 using Amazon.EC2.Model;
 using Horde.Build.Agents.Pools;
-using Horde.Build.Auditing;
 using Microsoft.Extensions.Logging;
 using OpenTracing;
 using OpenTracing.Util;
@@ -21,32 +19,20 @@ namespace Horde.Build.Agents.Fleet.Providers
 	/// Will start already existing but stopped instances to reuse existing EBS disks.
 	/// See <see cref="AwsFleetManager" /> for creating and terminating instances from scratch.
 	/// </summary>
-	public sealed class AwsReuseFleetManager : IFleetManager, IDisposable
+	public sealed class AwsReuseFleetManager : IFleetManager
 	{
-		readonly AmazonEC2Client _client;
-		readonly IAgentCollection _agentCollection;
-		readonly ILogger _logger;
+		private readonly IAmazonEC2 _ec2;
+		private readonly IAgentCollection _agentCollection;
+		private readonly ILogger _logger;
 
 		/// <summary>
 		/// Constructor
 		/// </summary>
-		public AwsReuseFleetManager(IAgentCollection agentCollection, ILogger<AwsReuseFleetManager> logger)
+		public AwsReuseFleetManager(IAmazonEC2 ec2, IAgentCollection agentCollection, ILogger<AwsReuseFleetManager> logger)
 		{
+			_ec2 = ec2;
 			_agentCollection = agentCollection;
 			_logger = logger;
-
-			AmazonEC2Config config = new AmazonEC2Config();
-			config.RegionEndpoint = RegionEndpoint.USEast1;
-
-			logger.LogInformation("Initializing AWS fleet manager for region {Region}", config.RegionEndpoint);
-
-			_client = new AmazonEC2Client(config);
-		}
-
-		/// <inheritdoc/>
-		public void Dispose()
-		{
-			_client.Dispose();
 		}
 
 		/// <inheritdoc/>
@@ -65,7 +51,7 @@ namespace Horde.Build.Agents.Fleet.Providers
 				describeRequest.Filters = new List<Filter>();
 				describeRequest.Filters.Add(new Filter("instance-state-name", new List<string> { InstanceStateName.Stopped.Value }));
 				describeRequest.Filters.Add(new Filter("tag:" + AwsFleetManager.PoolTagName, new List<string> { pool.Name }));
-				describeResponse = await _client.DescribeInstancesAsync(describeRequest, cancellationToken);
+				describeResponse = await _ec2.DescribeInstancesAsync(describeRequest, cancellationToken);
 				describeScope.Span.SetTag("res.statusCode", (int)describeResponse.HttpStatusCode);
 				describeScope.Span.SetTag("res.numReservations", describeResponse.Reservations.Count);
 			}
@@ -79,7 +65,7 @@ namespace Horde.Build.Agents.Fleet.Providers
 				startScope.Span.SetTag("req.instanceIds", String.Join(",", startRequest.InstanceIds));
 				if (startRequest.InstanceIds.Count > 0)
 				{
-					StartInstancesResponse startResponse = await _client.StartInstancesAsync(startRequest, cancellationToken);
+					StartInstancesResponse startResponse = await _ec2.StartInstancesAsync(startRequest, cancellationToken);
 					startScope.Span.SetTag("res.statusCode", (int)startResponse.HttpStatusCode);
 					startScope.Span.SetTag("res.numInstances", startResponse.StartingInstances.Count);
 					if ((int)startResponse.HttpStatusCode >= 200 && (int)startResponse.HttpStatusCode <= 299)
@@ -114,7 +100,7 @@ namespace Horde.Build.Agents.Fleet.Providers
 			describeRequest.Filters.Add(new Filter("instance-state-name", new List<string> { InstanceStateName.Stopped.Value }));
 			describeRequest.Filters.Add(new Filter("tag:" + AwsFleetManager.PoolTagName, new List<string> { pool.Name }));
 
-			DescribeInstancesResponse describeResponse = await _client.DescribeInstancesAsync(describeRequest, cancellationToken);
+			DescribeInstancesResponse describeResponse = await _ec2.DescribeInstancesAsync(describeRequest, cancellationToken);
 			return describeResponse.Reservations.SelectMany(x => x.Instances).Select(x => x.InstanceId).Distinct().Count();
 		}
 	}
diff --git a/Engine/Source/Programs/Horde/Horde.Build/Startup.cs b/Engine/Source/Programs/Horde/Horde.Build/Startup.cs
--- a/Engine/Source/Programs/Horde/Horde.Build/Startup.cs
+++ b/Engine/Source/Programs/Horde/Horde.Build/Startup.cs
@@ -14,6 +14,7 @@ using System.Threading;
 using System.Threading.Tasks;
 using Amazon;
 using Amazon.CloudWatch;
+using Amazon.EC2;
 using Amazon.Extensions.NETCore.Setup;
 using EpicGames.AspNet;
 using EpicGames.Core;
@@ -474,6 +475,7 @@ namespace Horde.Build
 			}
 			
 			services.AddAWSService<IAmazonCloudWatch>();
+			services.AddAWSService<IAmazonEC2>();
 
 			ConfigureLogStorage(services);
 
