Horde: Fix issue reports not correctly linking to triage threads when output to a different channel.

diff --git a/Engine/Source/Programs/Horde/Horde.Build/Notifications/Sinks/SlackNotificationSink.cs b/Engine/Source/Programs/Horde/Horde.Build/Notifications/Sinks/SlackNotificationSink.cs
--- a/Engine/Source/Programs/Horde/Horde.Build/Notifications/Sinks/SlackNotificationSink.cs
+++ b/Engine/Source/Programs/Horde/Horde.Build/Notifications/Sinks/SlackNotificationSink.cs
@@ -1492,7 +1492,7 @@ namespace Horde.Build.Notifications.Sinks
 				for (int idx = 0; idx < state.Blocks.Count; idx++)
 				{
 					string blockEventId = GetReportBlockEventId(messageId.Ts, idx);
-					await UpdateReportBlockAsync(channel, blockEventId, time, report.Stream, state.Blocks[idx].TemplateId, issuesByBlock[idx], state.Blocks[idx].TemplateHeader);
+					await UpdateReportBlockAsync(channel, blockEventId, time, report.Stream, state.Blocks[idx].TemplateId, issuesByBlock[idx], report.TriageChannel, state.Blocks[idx].TemplateHeader);
 				}
 
 				if (report.WorkflowStats.NumSteps > 0)
@@ -1598,13 +1598,13 @@ namespace Horde.Build.Notifications.Sinks
 							}
 						}
 
-						await UpdateReportBlockAsync(workflow.ReportChannel, blockEventId, state.Time, stream, block.TemplateId, issues, block.TemplateHeader);
+						await UpdateReportBlockAsync(workflow.ReportChannel, blockEventId, state.Time, stream, block.TemplateId, issues, workflow.TriageChannel, block.TemplateHeader);
 					}
 				}
 			}
 		}
 
-		async Task UpdateReportBlockAsync(string channel, string eventId, DateTime reportTime, IStream stream, TemplateId templateId, List<(IIssue, IIssueSpan?, bool)> issues, bool templateHeader)
+		async Task UpdateReportBlockAsync(string channel, string eventId, DateTime reportTime, IStream stream, TemplateId templateId, List<(IIssue, IIssueSpan?, bool)> issues, string? triageChannel, bool templateHeader)
 		{
 			StringBuilder body = new StringBuilder();
 
@@ -1629,7 +1629,7 @@ namespace Horde.Build.Notifications.Sinks
 					body.Append('\n');
 				}
 
-				string text = await FormatIssueAsync(issue, span, channel, reportTime, open);
+				string text = await FormatIssueAsync(issue, span, triageChannel, reportTime, open);
 				body.Append(text);
 			}
 
