Horde: Disable DMs for issues that already have a triage thread.

diff --git a/Engine/Source/Programs/Horde/Horde.Build/Notifications/Sinks/SlackNotificationSink.cs b/Engine/Source/Programs/Horde/Horde.Build/Notifications/Sinks/SlackNotificationSink.cs
--- a/Engine/Source/Programs/Horde/Horde.Build/Notifications/Sinks/SlackNotificationSink.cs
+++ b/Engine/Source/Programs/Horde/Horde.Build/Notifications/Sinks/SlackNotificationSink.cs
@@ -660,6 +660,9 @@ namespace Horde.Build.Notifications.Sinks
 
 			IIssueDetails details = await _issueService.GetIssueDetailsAsync(issue);
 
+			bool notifyOwner = true;
+			bool notifySuspects = issue.Promoted || details.Spans.Any(x => x.LastFailure.Annotations.NotifySubmitters ?? false);
+
 			WorkflowConfig? workflow = null;
 			if (details.Spans.Count > 0)
 			{
@@ -669,16 +672,14 @@ namespace Horde.Build.Notifications.Sinks
 				if (workflowId != null)
 				{
 					IStream? stream = await _streamService.GetStreamAsync(span.StreamId);
-					if (stream != null && stream.Config.TryGetWorkflow(workflowId.Value, out workflow))
+					if (stream != null && stream.Config.TryGetWorkflow(workflowId.Value, out workflow) && workflow.TriageChannel != null)
 					{
-						await CreateOrUpdateWorkflowThreadAsync(issue, span, details.Spans, workflow);
+						await CreateOrUpdateWorkflowThreadAsync(workflow.TriageChannel, issue, span, details.Spans, workflow);
+						notifyOwner = notifySuspects = false;
 					}
 				}
 			}
 
-			bool notifyOwner = workflow?.TriageChannel == null;
-			bool notifySuspects = issue.Promoted || details.Spans.Any(x => x.LastFailure.Annotations.NotifySubmitters ?? false);
-
 			HashSet<UserId> userIds = new HashSet<UserId>();
 			if (notifySuspects)
 			{
@@ -788,292 +789,288 @@ namespace Horde.Build.Notifications.Sinks
 
 		static string GetTriageThreadEventId(int issueId) => $"issue_triage_{issueId}";
 
-		async Task CreateOrUpdateWorkflowThreadAsync(IIssue issue, IIssueSpan span, IReadOnlyList<IIssueSpan> spans, WorkflowConfig workflow)
+		async Task CreateOrUpdateWorkflowThreadAsync(string triageChannel, IIssue issue, IIssueSpan span, IReadOnlyList<IIssueSpan> spans, WorkflowConfig workflow)
 		{
-			string? triageChannel = workflow.TriageChannel;
-			if (triageChannel != null)
-			{
-				Uri issueUrl = GetIssueUrl(issue, span.FirstFailure);
+			Uri issueUrl = GetIssueUrl(issue, span.FirstFailure);
 
-				string eventId = GetTriageThreadEventId(issue.Id);
+			string eventId = GetTriageThreadEventId(issue.Id);
 
-				string prefix = workflow.TriagePrefix ?? String.Empty;
-				if (issue.Severity == IssueSeverity.Error && (span.LastFailure.Annotations.BuildBlocker ?? false))
-				{
-					prefix = $"{prefix}*[BUILD BLOCKER]* ";
-				}
+			string prefix = workflow.TriagePrefix ?? String.Empty;
+			if (issue.Severity == IssueSeverity.Error && (span.LastFailure.Annotations.BuildBlocker ?? false))
+			{
+				prefix = $"{prefix}*[BUILD BLOCKER]* ";
+			}
 
-				string issueSummary = issue.UserSummary ?? issue.Summary;
-				string text = $"{workflow.TriagePrefix}{GetSeverityPrefix(issue.Severity)}Issue <{issueUrl}|{issue.Id}>: {issueSummary}{workflow.TriageSuffix}";
-				bool closed = !spans.Any(x => x.NextSuccess == null && x.LastFailure.Annotations.WorkflowId != null);
-				if (closed) // Thread may be shared by multiple workflows
-				{
-					text = $"~{text}~";
-				}
+			string issueSummary = issue.UserSummary ?? issue.Summary;
+			string text = $"{workflow.TriagePrefix}{GetSeverityPrefix(issue.Severity)}Issue <{issueUrl}|{issue.Id}>: {issueSummary}{workflow.TriageSuffix}";
+			bool closed = !spans.Any(x => x.NextSuccess == null && x.LastFailure.Annotations.WorkflowId != null);
+			if (closed) // Thread may be shared by multiple workflows
+			{
+				text = $"~{text}~";
+			}
 
-				// Get the suspects for the issue
-				List<IIssueSuspect> suspects = await _issueService.Collection.FindSuspectsAsync(issue);
+			// Get the suspects for the issue
+			List<IIssueSuspect> suspects = await _issueService.Collection.FindSuspectsAsync(issue);
 
-				(MessageStateDocument state, bool isNew) = await SendOrUpdateMessageAsync(triageChannel, eventId, null, text);
-				SlackMessageId threadId = state.MessageId;
+			(MessageStateDocument state, bool isNew) = await SendOrUpdateMessageAsync(triageChannel, eventId, null, text);
+			SlackMessageId threadId = state.MessageId;
 
-				if (isNew)
-				{
-					// Create the summary text
-					List<ILogEvent> events = new List<ILogEvent>();
-					List<ILogEventData> eventDataItems = new List<ILogEventData>();
+			if (isNew)
+			{
+				// Create the summary text
+				List<ILogEvent> events = new List<ILogEvent>();
+				List<ILogEventData> eventDataItems = new List<ILogEventData>();
 
-					if (span.FirstFailure.LogId != null)
+				if (span.FirstFailure.LogId != null)
+				{
+					LogId logId = span.FirstFailure.LogId.Value;
+					ILogFile? logFile = await _logFileService.GetLogFileAsync(logId);
+					if (logFile != null)
 					{
-						LogId logId = span.FirstFailure.LogId.Value;
-						ILogFile? logFile = await _logFileService.GetLogFileAsync(logId);
-						if (logFile != null)
+						events = await _logFileService.FindEventsAsync(logFile, span.Id, 0, 50);
+						if (events.Any(x => x.Severity == EventSeverity.Error))
 						{
-							events = await _logFileService.FindEventsAsync(logFile, span.Id, 0, 50);
-							if (events.Any(x => x.Severity == EventSeverity.Error))
-							{
-								events.RemoveAll(x => x.Severity == EventSeverity.Warning);
-							}
-
-							List<string> eventStrings = new List<string>();
-							for (int idx = 0; idx < Math.Min(events.Count, 3); idx++)
-							{
-								ILogEventData data = await _logFileService.GetEventDataAsync(logFile, events[idx].LineIndex, events[idx].LineCount);
-								eventDataItems.Add(data);
-							}
+							events.RemoveAll(x => x.Severity == EventSeverity.Warning);
 						}
-					}
-
-					SlackMessage message = new SlackMessage();
 
-					message.Blocks.Add(new SectionBlock(new TextObject($"From {FormatJobStep(span.FirstFailure, span.NodeName)}:")));
-					foreach (ILogEventData eventDataItem in eventDataItems)
-					{
-						message.Blocks.Add(new SectionBlock(new TextObject(QuoteText(eventDataItem.Message, MaxLineLength))));
-					}
-					if (events.Count > eventDataItems.Count)
-					{
-						message.Blocks.Add(new SectionBlock(new TextObject("```...```")));
+						List<string> eventStrings = new List<string>();
+						for (int idx = 0; idx < Math.Min(events.Count, 3); idx++)
+						{
+							ILogEventData data = await _logFileService.GetEventDataAsync(logFile, events[idx].LineIndex, events[idx].LineCount);
+							eventDataItems.Add(data);
+						}
 					}
+				}
 
-					SlackMessageId summaryId = await _slackClient.PostMessageToThreadAsync(threadId, message);
+				SlackMessage message = new SlackMessage();
 
-					// Permalink to the summary text so we link inside the thread rather than just to the original message
-					string? permalink = await _slackClient.GetPermalinkAsync(summaryId);
-					await UpdateMessageStateAsync(state.Id, state.MessageId, permalink);
+				message.Blocks.Add(new SectionBlock(new TextObject($"From {FormatJobStep(span.FirstFailure, span.NodeName)}:")));
+				foreach (ILogEventData eventDataItem in eventDataItems)
+				{
+					message.Blocks.Add(new SectionBlock(new TextObject(QuoteText(eventDataItem.Message, MaxLineLength))));
 				}
-
-				// Post a message containing the controls and status
+				if (events.Count > eventDataItems.Count)
 				{
-					SlackMessage message = new SlackMessage();
-
-					if (workflow.TriageInstructions != null)
-					{
-						message.AddSection(workflow.TriageInstructions);
-					}
-
-					if (!closed)
-					{
-						ActionsBlock actions = message.AddActions();
-						actions.AddButton("Assign to Me", value: $"issue_{issue.Id}_ack", style: ButtonStyle.Primary);
-						actions.AddButton("Not Me", value: $"issue_{issue.Id}_decline", style: ButtonStyle.Danger);
-						actions.AddButton("Mark Fixed", value: $"issue_{issue.Id}_markfixed");
-
-						string? context = null;
-						if (issue.OwnerId != null)
-						{
-							string user = await FormatNameAsync(issue.OwnerId.Value);
-							if (issue.AcknowledgedAt == null)
-							{
-								context = $"Assigned to {user} (unacknowledged).";
-							}
-							else
-							{
-								context = $"Acknowledged by {user}.";
-							}
-						}
-						else if (suspects.Any(x => x.DeclinedAt != null))
-						{
-							HashSet<UserId> userIds = new HashSet<UserId>();
-							foreach (IIssueSuspect suspect in suspects)
-							{
-								if (suspect.DeclinedAt != null)
-								{
-									userIds.Add(suspect.AuthorId);
-								}
-							}
+					message.Blocks.Add(new SectionBlock(new TextObject("```...```")));
+				}
 
-							List<string> users = new List<string>();
-							foreach (UserId userId in userIds)
-							{
-								users.Add(await FormatNameAsync(userId));
-							}
-							users.Sort(StringComparer.OrdinalIgnoreCase);
+				SlackMessageId summaryId = await _slackClient.PostMessageToThreadAsync(threadId, message);
 
-							context = $"Declined by {StringUtils.FormatList(users)}.";
-						}
+				// Permalink to the summary text so we link inside the thread rather than just to the original message
+				string? permalink = await _slackClient.GetPermalinkAsync(summaryId);
+				await UpdateMessageStateAsync(state.Id, state.MessageId, permalink);
+			}
 
-						if (context != null)
-						{
-							message.AddContext(context);
-						}
-					}
+			// Post a message containing the controls and status
+			{
+				SlackMessage message = new SlackMessage();
 
-					await SendOrUpdateMessageToThreadAsync(triageChannel, $"{eventId}_buttons", null, threadId, message);
+				if (workflow.TriageInstructions != null)
+				{
+					message.AddSection(workflow.TriageInstructions);
 				}
 
-				bool notifyTriageAlias = false;
-				if (isNew)
+				if (!closed)
 				{
-					// If it has an owner, show that
-					HashSet<UserId> inviteUserIds = new HashSet<UserId>();
+					ActionsBlock actions = message.AddActions();
+					actions.AddButton("Assign to Me", value: $"issue_{issue.Id}_ack", style: ButtonStyle.Primary);
+					actions.AddButton("Not Me", value: $"issue_{issue.Id}_decline", style: ButtonStyle.Danger);
+					actions.AddButton("Mark Fixed", value: $"issue_{issue.Id}_markfixed");
+
+					string? context = null;
 					if (issue.OwnerId != null)
 					{
-						string mention = await FormatMentionAsync(issue.OwnerId.Value, workflow.AllowMentions);
-
-						string changes = String.Join(", ", suspects.Where(x => x.AuthorId == issue.OwnerId).Select(x => FormatChange(x.Change)));
-						if (changes.Length > 0)
+						string user = await FormatNameAsync(issue.OwnerId.Value);
+						if (issue.AcknowledgedAt == null)
 						{
-							mention += $" ({changes})";
+							context = $"Assigned to {user} (unacknowledged).";
+						}
+						else
+						{
+							context = $"Acknowledged by {user}.";
 						}
-
-						await _slackClient.PostMessageToThreadAsync(threadId, $"Assigned to {mention}");
-						inviteUserIds.Add(issue.OwnerId.Value);
 					}
-					else
+					else if (suspects.Any(x => x.DeclinedAt != null))
 					{
-						IGrouping<UserId, IIssueSuspect>[] suspectGroups = suspects.GroupBy(x => x.AuthorId).ToArray();
-						if (suspectGroups.Length > 0 && suspectGroups.Length <= workflow.MaxMentions)
+						HashSet<UserId> userIds = new HashSet<UserId>();
+						foreach (IIssueSuspect suspect in suspects)
 						{
-							List<string> suspectList = new List<string>();
-							foreach (IGrouping<UserId, IIssueSuspect> suspectGroup in suspectGroups)
+							if (suspect.DeclinedAt != null)
 							{
-								string mention = await FormatMentionAsync(suspectGroup.Key, workflow.AllowMentions);
-								string changes = String.Join(", ", suspectGroup.Select(x => FormatChange(x.Change)));
-								suspectList.Add($"{mention} ({changes})");
-								inviteUserIds.Add(suspectGroup.Key);
+								userIds.Add(suspect.AuthorId);
 							}
-
-							string suspectMessage = $"Possibly {StringUtils.FormatList(suspectList, "or")}.";
-							await _slackClient.PostMessageToThreadAsync(threadId, suspectMessage);
 						}
-						else
+
+						List<string> users = new List<string>();
+						foreach (UserId userId in userIds)
 						{
-							notifyTriageAlias = true;
+							users.Add(await FormatNameAsync(userId));
 						}
+						users.Sort(StringComparer.OrdinalIgnoreCase);
+
+						context = $"Declined by {StringUtils.FormatList(users)}.";
 					}
 
-					if (_environment.IsProduction() && workflow.AllowMentions)
+					if (context != null)
 					{
-						await InviteUsersAsync(state.Channel, inviteUserIds);
+						message.AddContext(context);
 					}
 				}
 
-				if (workflow.EscalateAlias != null && workflow.EscalateTimes.Count > 0)
+				await SendOrUpdateMessageToThreadAsync(triageChannel, $"{eventId}_buttons", null, threadId, message);
+			}
+
+			bool notifyTriageAlias = false;
+			if (isNew)
+			{
+				// If it has an owner, show that
+				HashSet<UserId> inviteUserIds = new HashSet<UserId>();
+				if (issue.OwnerId != null)
 				{
-					DateTime escalateTime = _clock.UtcNow + TimeSpan.FromMinutes(workflow.EscalateTimes[0]);
-					if (await _escalateIssues.AddAsync(issue.Id, (escalateTime - DateTime.UnixEpoch).TotalSeconds, StackExchange.Redis.When.NotExists))
+					string mention = await FormatMentionAsync(issue.OwnerId.Value, workflow.AllowMentions);
+
+					string changes = String.Join(", ", suspects.Where(x => x.AuthorId == issue.OwnerId).Select(x => FormatChange(x.Change)));
+					if (changes.Length > 0)
 					{
-						_logger.LogInformation("First escalation time for issue {IssueId} is {Time}", issue.Id, escalateTime);
+						mention += $" ({changes})";
 					}
-				}
 
-				if ((workflow.TriageAlias != null || workflow.TriageTypeAliases != null) && issue.OwnerId == null && (suspects.All(x => x.DeclinedAt != null) || notifyTriageAlias) && !closed)
+					await _slackClient.PostMessageToThreadAsync(threadId, $"Assigned to {mention}");
+					inviteUserIds.Add(issue.OwnerId.Value);
+				}
+				else
 				{
-					string? triageAlias;
-					
-					if (workflow.TriageTypeAliases == null || issue.Fingerprints.Count == 0 || !workflow.TriageTypeAliases.TryGetValue(issue.Fingerprints[0].Type, out triageAlias))
+					IGrouping<UserId, IIssueSuspect>[] suspectGroups = suspects.GroupBy(x => x.AuthorId).ToArray();
+					if (suspectGroups.Length > 0 && suspectGroups.Length <= workflow.MaxMentions)
 					{
-						triageAlias = workflow.TriageAlias;
-					}
+						List<string> suspectList = new List<string>();
+						foreach (IGrouping<UserId, IIssueSuspect> suspectGroup in suspectGroups)
+						{
+							string mention = await FormatMentionAsync(suspectGroup.Key, workflow.AllowMentions);
+							string changes = String.Join(", ", suspectGroup.Select(x => FormatChange(x.Change)));
+							suspectList.Add($"{mention} ({changes})");
+							inviteUserIds.Add(suspectGroup.Key);
+						}
 
-					if (triageAlias != null)
+						string suspectMessage = $"Possibly {StringUtils.FormatList(suspectList, "or")}.";
+						await _slackClient.PostMessageToThreadAsync(threadId, suspectMessage);
+					}
+					else
 					{
-						string triageMessage = $"(cc {FormatUserOrGroupMention(triageAlias)} for triage).";
-						await SendOrUpdateMessageToThreadAsync(triageChannel, eventId + "_triage", null, threadId, triageMessage);
+						notifyTriageAlias = true;
 					}
 				}
 
-				if (issue.AcknowledgedAt != null)
+				if (_environment.IsProduction() && workflow.AllowMentions)
 				{
-					await _slackClient.AddReactionAsync(threadId, "eyes");
+					await InviteUsersAsync(state.Channel, inviteUserIds);
 				}
-				else
+			}
+
+			if (workflow.EscalateAlias != null && workflow.EscalateTimes.Count > 0)
+			{
+				DateTime escalateTime = _clock.UtcNow + TimeSpan.FromMinutes(workflow.EscalateTimes[0]);
+				if (await _escalateIssues.AddAsync(issue.Id, (escalateTime - DateTime.UnixEpoch).TotalSeconds, StackExchange.Redis.When.NotExists))
 				{
-					await _slackClient.RemoveReactionAsync(threadId, "eyes");
+					_logger.LogInformation("First escalation time for issue {IssueId} is {Time}", issue.Id, escalateTime);
 				}
+			}
 
-				if (issue.QuarantinedByUserId != null)
+			if ((workflow.TriageAlias != null || workflow.TriageTypeAliases != null) && issue.OwnerId == null && (suspects.All(x => x.DeclinedAt != null) || notifyTriageAlias) && !closed)
+			{
+				string? triageAlias;
+					
+				if (workflow.TriageTypeAliases == null || issue.Fingerprints.Count == 0 || !workflow.TriageTypeAliases.TryGetValue(issue.Fingerprints[0].Type, out triageAlias))
 				{
-					await _slackClient.AddReactionAsync(threadId, "mask");
+					triageAlias = workflow.TriageAlias;
 				}
-				else
+
+				if (triageAlias != null)
 				{
-					await _slackClient.RemoveReactionAsync(threadId, "mask");
+					string triageMessage = $"(cc {FormatUserOrGroupMention(triageAlias)} for triage).";
+					await SendOrUpdateMessageToThreadAsync(triageChannel, eventId + "_triage", null, threadId, triageMessage);
 				}
+			}
 
-				IIssueSpan? fixFailedSpan = null;
-				if (issue.FixChange != null)
-				{
-					fixFailedSpan = GetFixFailedSpan(issue, spans);
+			if (issue.AcknowledgedAt != null)
+			{
+				await _slackClient.AddReactionAsync(threadId, "eyes");
+			}
+			else
+			{
+				await _slackClient.RemoveReactionAsync(threadId, "eyes");
+			}
 
-					if (fixFailedSpan == null)
-					{
-						string fixedEventId = $"issue_{issue.Id}_fixed_{issue.FixChange}";
-						string fixedMessage = $"Marked as fixed in {FormatChange(issue.FixChange.Value)}";
-						await PostSingleMessageToThreadAsync(triageChannel, fixedEventId, threadId, fixedMessage);
-					}
-					else
-					{
-						string fixFailedEventId = $"issue_{issue.Id}_fixfailed_{issue.FixChange}";
-						string fixFailedMessage = $"Issue not fixed by {FormatChange(issue.FixChange.Value)}; see {FormatJobStep(fixFailedSpan.LastFailure, fixFailedSpan.NodeName)} at CL {fixFailedSpan.LastFailure.Change} in {fixFailedSpan.StreamName}.";
-						if (issue.OwnerId.HasValue)
-						{
-							string mention = await FormatMentionAsync(issue.OwnerId.Value, workflow.AllowMentions);
-							fixFailedMessage += $" ({mention})";
-						}
-						await PostSingleMessageToThreadAsync(triageChannel, fixFailedEventId, threadId, fixFailedMessage);
-					}
+			if (issue.QuarantinedByUserId != null)
+			{
+				await _slackClient.AddReactionAsync(threadId, "mask");
+			}
+			else
+			{
+				await _slackClient.RemoveReactionAsync(threadId, "mask");
+			}
 
-					if (fixFailedSpan == null)
-					{
-						foreach (IIssueStream stream in issue.Streams)
-						{
-							if ((stream.MergeOrigin ?? false) && !(stream.ContainsFix ?? false))
-							{
-								string streamName = spans.FirstOrDefault(x => x.StreamId == stream.StreamId)?.StreamName ?? stream.StreamId.ToString();
-								string missingEventId = $"issue_{issue.Id}_fixmissing_{issue.FixChange}_{stream.StreamId}";
-								string missingMessage = $"Note: Fix may need manually merging to {streamName}";
-								await PostSingleMessageToThreadAsync(triageChannel, missingEventId, threadId, missingMessage);
-							}
-						}
-					}
-				}
+			IIssueSpan? fixFailedSpan = null;
+			if (issue.FixChange != null)
+			{
+				fixFailedSpan = GetFixFailedSpan(issue, spans);
 
-				if (fixFailedSpan != null)
+				if (fixFailedSpan == null)
 				{
-					await _slackClient.AddReactionAsync(threadId, "x");
+					string fixedEventId = $"issue_{issue.Id}_fixed_{issue.FixChange}";
+					string fixedMessage = $"Marked as fixed in {FormatChange(issue.FixChange.Value)}";
+					await PostSingleMessageToThreadAsync(triageChannel, fixedEventId, threadId, fixedMessage);
 				}
 				else
 				{
-					await _slackClient.RemoveReactionAsync(threadId, "x");
+					string fixFailedEventId = $"issue_{issue.Id}_fixfailed_{issue.FixChange}";
+					string fixFailedMessage = $"Issue not fixed by {FormatChange(issue.FixChange.Value)}; see {FormatJobStep(fixFailedSpan.LastFailure, fixFailedSpan.NodeName)} at CL {fixFailedSpan.LastFailure.Change} in {fixFailedSpan.StreamName}.";
+					if (issue.OwnerId.HasValue)
+					{
+						string mention = await FormatMentionAsync(issue.OwnerId.Value, workflow.AllowMentions);
+						fixFailedMessage += $" ({mention})";
+					}
+					await PostSingleMessageToThreadAsync(triageChannel, fixFailedEventId, threadId, fixFailedMessage);
 				}
 
-				if (issue.ResolvedAt != null && fixFailedSpan == null)
+				if (fixFailedSpan == null)
 				{
-					await _slackClient.AddReactionAsync(threadId, "tick");
-				}
-				else
-				{
-					await _slackClient.RemoveReactionAsync(threadId, "tick");
+					foreach (IIssueStream stream in issue.Streams)
+					{
+						if ((stream.MergeOrigin ?? false) && !(stream.ContainsFix ?? false))
+						{
+							string streamName = spans.FirstOrDefault(x => x.StreamId == stream.StreamId)?.StreamName ?? stream.StreamId.ToString();
+							string missingEventId = $"issue_{issue.Id}_fixmissing_{issue.FixChange}_{stream.StreamId}";
+							string missingMessage = $"Note: Fix may need manually merging to {streamName}";
+							await PostSingleMessageToThreadAsync(triageChannel, missingEventId, threadId, missingMessage);
+						}
+					}
 				}
+			}
 
-				if (issue.ExternalIssueKey != null)
-				{
-					string extIssueEventId = $"issue_{issue.Id}_ext_{issue.ExternalIssueKey}";
-					string extIssueMessage = $"Linked to issue {FormatExternalIssue(issue.ExternalIssueKey)}";
-					await PostSingleMessageToThreadAsync(triageChannel, extIssueEventId, threadId, extIssueMessage);
-				}
+			if (fixFailedSpan != null)
+			{
+				await _slackClient.AddReactionAsync(threadId, "x");
+			}
+			else
+			{
+				await _slackClient.RemoveReactionAsync(threadId, "x");
+			}
+
+			if (issue.ResolvedAt != null && fixFailedSpan == null)
+			{
+				await _slackClient.AddReactionAsync(threadId, "tick");
+			}
+			else
+			{
+				await _slackClient.RemoveReactionAsync(threadId, "tick");
+			}
+
+			if (issue.ExternalIssueKey != null)
+			{
+				string extIssueEventId = $"issue_{issue.Id}_ext_{issue.ExternalIssueKey}";
+				string extIssueMessage = $"Linked to issue {FormatExternalIssue(issue.ExternalIssueKey)}";
+				await PostSingleMessageToThreadAsync(triageChannel, extIssueEventId, threadId, extIssueMessage);
 			}
 		}
 
