Horde: Add latest problem steps to device telemetry summary and sort these by date
#rnx
#skipci

diff --git a/Engine/Source/Programs/Horde/HordeDashboard/src/components/DevicePoolTelemetry.tsx b/Engine/Source/Programs/Horde/HordeDashboard/src/components/DevicePoolTelemetry.tsx
--- a/Engine/Source/Programs/Horde/HordeDashboard/src/components/DevicePoolTelemetry.tsx
+++ b/Engine/Source/Programs/Horde/HordeDashboard/src/components/DevicePoolTelemetry.tsx
@@ -1,9 +1,10 @@
-import { DefaultButton, DetailsList, DetailsListLayoutMode, Dropdown, FocusZone, FocusZoneDirection, IconButton, IDropdownOption, Modal, ScrollablePane, ScrollbarVisibility, SelectionMode, Spinner, SpinnerSize, Stack, Text } from "@fluentui/react";
+import { DefaultButton, DetailsList, DetailsListLayoutMode, Dropdown, FocusZone, FocusZoneDirection, IColumn, IconButton, IDropdownOption, Modal, ScrollablePane, ScrollbarVisibility, SelectionMode, Spinner, SpinnerSize, Stack, Text } from "@fluentui/react";
 import * as d3 from "d3";
 import { action, observable } from "mobx";
 import { observer } from "mobx-react-lite";
 import moment from "moment";
 import { useState } from "react";
+import { Link } from "react-router-dom";
 import backend from "../backend";
 import { DevicePoolTelemetryQuery, DevicePoolType, DeviceTelemetryQuery, GetDevicePlatformResponse, GetDevicePlatformTelemetryResponse, GetDevicePoolResponse, GetDevicePoolTelemetryResponse, GetDeviceResponse, GetDeviceTelemetryResponse, GetTelemetryInfoResponse } from "../backend/Api";
 import dashboard, { StatusColor } from "../backend/Dashboard";
@@ -68,6 +69,8 @@ type ProblemStep = {
    stepName?: string;
 }
 
+type DeviceProblem = { deviceId: string, deviceName: string, problems: number, problemsDesc: string, latestProblem?: DeviceTelemetry };
+
 class PoolTelemetryHandler {
 
    constructor() {
@@ -172,7 +175,7 @@ class PoolTelemetryHandler {
 
    getProblemDevices() {
 
-      const platforms: Map<string, { deviceId: string, deviceName: string, problems: number, problemsDesc: string, latestProblem?: string }[]> = new Map();
+      const platforms: Map<string, DeviceProblem[]> = new Map();
 
       this.deviceTelemetry.forEach((telemetry, deviceId) => {
 
@@ -187,8 +190,18 @@ class PoolTelemetryHandler {
          }
 
          let count = 0;
-         let latest: Date | undefined;
-         telemetry.forEach(t => { if (t.problemTimeUtc) count++; if (!latest || latest < t.problemTime!) latest = t.problemTime });
+         let latest: DeviceTelemetry | undefined;
+         telemetry.filter(t => t.problemTime).forEach(t => {
+            count++;
+            if (!latest || latest.problemTime! < t.problemTime!) {
+               latest = t
+            }
+
+            if (t.jobName?.indexOf("- Kicked By") !== -1) {
+               t.jobName = t.jobName?.split("- Kicked By")[0];
+            }
+
+         });
 
          if (!count) {
             return;
@@ -198,7 +211,7 @@ class PoolTelemetryHandler {
             platforms.set(platform.id, []);
          }
 
-         platforms.get(platform.id)!.push({ deviceId: deviceId, deviceName: `${device.name} / ${device.address}`, problems: count, problemsDesc: `${count} / ${telemetry.length}`, latestProblem: latest ? getHumanTime(latest) : "" });
+         platforms.get(platform.id)!.push({ deviceId: deviceId, deviceName: `${device.name} / ${device.address}`, problems: count, problemsDesc: `${count} / ${telemetry.length}`, latestProblem: latest });
 
       });
 
@@ -996,9 +1009,17 @@ class PoolTelemetryGraph {
    extent?: any;
 }
 
-const SummaryModal: React.FC<{ handler: PoolTelemetryHandler, onClose: () => void }> = ({ handler, onClose }) => {
+const SummaryModal: React.FC<{ handler: PoolTelemetryHandler, selectedPlatform?: GetDevicePlatformResponse, onClose: () => void }> = ({ handler, selectedPlatform, onClose }) => {
 
-   const platforms = Array.from(handler.platforms.values()).sort((a, b) => a.name.localeCompare(b.name));
+   let platforms = Array.from(handler.platforms.values()).sort((a, b) => {
+      if (a.id === selectedPlatform?.id) {
+         return -1;
+      }
+      if (b.id === selectedPlatform?.id) {
+         return 1;
+      }
+      return a.name.localeCompare(b.name)
+   });
 
    const problems = handler.getProblemDevices();
 
@@ -1006,14 +1027,39 @@ const SummaryModal: React.FC<{ handler: PoolTelemetryHandler, onClose: () => voi
 
       const psteps = handler.getSteps(platform.id);
       let platformProblems = problems.get(platform.id) ?? [];
-      platformProblems = platformProblems.filter(p => p.problems > 8);
+      platformProblems = platformProblems.filter(p => p.problems > 8).sort((a, b) => b.latestProblem!.problemTime!.getTime() - a.latestProblem!.problemTime!.getTime());
 
-      const problemColumns = [
+      const problemColumns: IColumn[] = [
          { key: 'column1', name: 'Problem Devices', fieldName: 'deviceName', minWidth: 320, maxWidth: 320 },
-         { key: 'column2', name: 'Issues vs Reservations', fieldName: 'problemsDesc', minWidth: 180, maxWidth: 180 },
-         { key: 'column3', name: 'Latest Issue', fieldName: 'latestProblem', minWidth: 120, maxWidth: 120 },
+         { key: 'column2', name: 'Issues/Reservations', fieldName: 'problemsDesc', minWidth: 150, maxWidth: 150 },
+         { key: 'column3', name: 'Latest Issue', minWidth: 640, maxWidth: 640 },
       ];
 
+      let column = problemColumns.find(c => c.name === "Latest Issue")!;
+
+      column.onRender = ((item: DeviceProblem) => {
+
+         const latest = item.latestProblem;
+
+         if (!latest) {
+            return null;
+         }
+
+         if (!latest.jobId || !latest.stepId) {
+            return <Stack verticalFill={true} verticalAlign="center">
+               <Text>{getHumanTime(latest.problemTime)}</Text>
+            </Stack>
+         }
+
+         const url = `/job/${latest.jobId}?step=${latest.stepId}`;
+
+         return <Stack verticalFill={true} verticalAlign="center">
+            <Link to={url} target="_blank">
+               <Text variant="small">{`${getHumanTime(latest.problemTime!)} - ${latest.jobName ?? "Unknown Job Name"} - ${latest.stepName ?? "Unknown Step Name"}`}</Text>
+            </Link>
+         </Stack>
+      });
+
       const steps: Map<string, { stepName: string, reservations: number, problems: number, durationMS: number }> = new Map();
 
       psteps.forEach(s => {
@@ -1045,7 +1091,7 @@ const SummaryModal: React.FC<{ handler: PoolTelemetryHandler, onClose: () => voi
             </Stack>
             <Stack tokens={{ childrenGap: 12 }}>
                {!!platformProblems?.length &&
-                  <Stack style={{ paddingLeft: 12, paddingTop: 12, width: 800 }}>
+                  <Stack style={{ paddingLeft: 12, paddingTop: 12, width: 1220 }}>
                      <DetailsList
                         isHeaderVisible={true}
                         items={platformProblems}
@@ -1056,7 +1102,7 @@ const SummaryModal: React.FC<{ handler: PoolTelemetryHandler, onClose: () => voi
                      />
                   </Stack>}
                {!!stepItems?.length &&
-                  <Stack style={{ paddingLeft: 12, paddingTop: 12, width: 800 }} >
+                  <Stack style={{ paddingLeft: 12, paddingTop: 12, width: 1220 }} >
                      <DetailsList
                         isHeaderVisible={true}
                         items={stepItems}
@@ -1080,7 +1126,7 @@ const SummaryModal: React.FC<{ handler: PoolTelemetryHandler, onClose: () => voi
       title = `Device Platform Summary / ${minDate} - ${maxDate}`;
    }
 
-   return <Modal className={hordeClasses.modal} isOpen={true} isBlocking={true} topOffsetFixed={true} styles={{ main: { padding: 8, width: 940, backgroundColor: dashboard.darktheme ? modeColors.content : modeColors.background, hasBeenOpened: false, top: "72px", position: "absolute", height: "820px" } }} onDismiss={() => { onClose() }}>
+   return <Modal className={hordeClasses.modal} isOpen={true} isBlocking={true} topOffsetFixed={true} styles={{ main: { padding: 8, width: 1340, backgroundColor: dashboard.darktheme ? modeColors.content : modeColors.background, hasBeenOpened: false, top: "72px", position: "absolute", height: "820px" } }} onDismiss={() => { onClose() }}>
       <Stack>
          <Stack horizontal verticalAlign="center" style={{ paddingTop: 12 }}>
             <Stack style={{ paddingLeft: 16 }}>
@@ -1170,9 +1216,9 @@ export const DevicePoolGraph: React.FC = observer(() => {
       const maxDate = getHumanTime(handler.data[handler.data.length - 1].date);
       title = `Device Pool Telemetry / ${minDate} - ${maxDate}`;
    }
-   
+
    return <Stack>
-      {showSummary && !!handler && <SummaryModal handler={handler} onClose={() => setShowSummary(false)} />}
+      {showSummary && !!handler && <SummaryModal handler={handler} selectedPlatform={state.platform} onClose={() => setShowSummary(false)} />}
       <Stack style={{ position: "relative" }}>
          <Stack horizontal verticalAlign="start" style={{ width: 1412, paddingBottom: 18 }} tokens={{ childrenGap: 32 }}>
             <Stack style={{ paddingLeft: 4 }}>
@@ -1217,7 +1263,7 @@ export const DevicePoolTelemetryModal: React.FC<{ onClose?: () => void }> = ({ o
 
 
    return <Stack>
-      <Modal className={hordeClasses.modal} isOpen={true} isBlocking={true} topOffsetFixed={true} styles={{ main: { padding: 8, width: 1520, backgroundColor: modeColors.background, hasBeenOpened: false, top: "64px", position: "absolute", height: "760px" } }} onDismiss={() => { if (onClose) { onClose() } }}>
+      <Modal className={hordeClasses.modal} isOpen={true} isBlocking={true} topOffsetFixed={true} styles={{ main: { padding: 8, width: 1560, backgroundColor: modeColors.background, hasBeenOpened: false, top: "64px", position: "absolute", height: "760px" } }} onDismiss={() => { if (onClose) { onClose() } }}>
          <Stack>
             <Stack horizontal>
                <Stack grow />
