Horde: Allow output of non-credential environment variables at start of steps.

diff --git a/Engine/Source/Programs/Horde/Horde.Agent/Execution/BuildGraphExecutor.cs b/Engine/Source/Programs/Horde/Horde.Agent/Execution/BuildGraphExecutor.cs
--- a/Engine/Source/Programs/Horde/Horde.Agent/Execution/BuildGraphExecutor.cs
+++ b/Engine/Source/Programs/Horde/Horde.Agent/Execution/BuildGraphExecutor.cs
@@ -799,7 +799,7 @@ namespace Horde.Agent.Execution
 			}
 
 			// Log all the environment variables to the log
-			HashSet<string> credentialKeys = new HashSet<string>(newEnvVars.Select(x => x.Key), StringComparer.OrdinalIgnoreCase);
+			HashSet<string> credentialKeys = new HashSet<string>(step.Credentials.Keys, StringComparer.OrdinalIgnoreCase);
 			foreach (KeyValuePair<string, string> envVar in newEnvVars.OrderBy(x => x.Key))
 			{
 				string value = "[redacted]";
