Horde: Fix long gaps in P4 replication not correctly resetting record of replicated commits, despite commits missing from DB.

diff --git a/Engine/Source/Programs/Horde/Horde.Build/Perforce/PerforceServiceCache.cs b/Engine/Source/Programs/Horde/Horde.Build/Perforce/PerforceServiceCache.cs
--- a/Engine/Source/Programs/Horde/Horde.Build/Perforce/PerforceServiceCache.cs
+++ b/Engine/Source/Programs/Horde/Horde.Build/Perforce/PerforceServiceCache.cs
@@ -351,7 +351,7 @@ namespace Horde.Build.Perforce
 
 		async Task<ClusterState?> UpdateClusterAsync(string clusterName, List<StreamInfo> streamInfos, ClusterState state, CancellationToken cancellationToken)
 		{
-			const int MaxChanges = 100;
+			const int MaxChanges = 250;
 
 			using (IPooledPerforceConnection perforce = await ConnectAsync(clusterName, null, cancellationToken))
 			{
@@ -369,6 +369,14 @@ namespace Horde.Build.Perforce
 					return null;
 				}
 
+				// If we've retrieved the maximum number of changes from the server, we no longer have a complete chronological cache and need to reset it.
+				bool reset = false;
+				if (changes.Count >= MaxChanges)
+				{
+					_logger.LogWarning("Too many commits since last P4 cache update (query: {Spec}). Resetting cache for {Cluster}.", spec, clusterName);
+					reset = true;
+				}
+
 				// Get the server information
 				InfoRecord info = await perforce.GetInfoAsync(cancellationToken);
 
@@ -394,6 +402,7 @@ namespace Horde.Build.Perforce
 						if (files.Count > 0)
 						{
 							ICommit commit = await CreateCommitAsync(perforce, streamInfo.Stream, describeRecord, info, cancellationToken);
+							_logger.LogInformation("Replicating {StreamId} commit {Change}", streamInfo.Stream.Id, describeRecord.Number);
 
 							List<CommitTag> commitTags = new List<CommitTag>();
 							foreach (CommitTagInfo commitTagInfo in streamInfo.CommitTags)
@@ -416,7 +425,7 @@ namespace Horde.Build.Perforce
 				foreach (StreamInfo streamInfo in streamInfos)
 				{
 					int minChange;
-					if (!state.MinChanges.TryGetValue(streamInfo.Stream.Id, out minChange))
+					if (reset || !state.MinChanges.TryGetValue(streamInfo.Stream.Id, out minChange))
 					{
 						minChange = changes[^1].Number;
 					}
