Horde: Allow fleet managers to be configured on a per-pool basis

Uses the same kind of configuration style as pool sizing strategies.

diff --git a/Engine/Source/Programs/Horde/Horde.Build.Tests/Fleet/AutoscaleServiceV2Test.cs b/Engine/Source/Programs/Horde/Horde.Build.Tests/Fleet/AutoscaleServiceV2Test.cs
--- a/Engine/Source/Programs/Horde/Horde.Build.Tests/Fleet/AutoscaleServiceV2Test.cs
+++ b/Engine/Source/Programs/Horde/Horde.Build.Tests/Fleet/AutoscaleServiceV2Test.cs
@@ -4,16 +4,21 @@ extern alias HordeAgent;
 using System;
 using System.Collections.Generic;
 using System.Linq;
+using System.Text.Json;
 using System.Threading;
 using System.Threading.Tasks;
+using Amazon.AutoScaling;
+using Amazon.EC2;
 using Horde.Build.Agents;
 using Horde.Build.Agents.Fleet;
+using Horde.Build.Agents.Fleet.Providers;
 using Horde.Build.Agents.Pools;
 using Horde.Build.Utilities;
 using Microsoft.Extensions.DependencyInjection;
 using Microsoft.Extensions.Logging;
 using Microsoft.Extensions.Options;
 using Microsoft.VisualStudio.TestTools.UnitTesting;
+using Moq;
 using StatsdClient;
 
 namespace Horde.Build.Tests.Fleet
@@ -126,11 +131,17 @@ namespace Horde.Build.Tests.Fleet
 
 		private AutoscaleServiceV2 GetAutoscaleService(IFleetManager fleetManager)
 		{
-			ILogger<AutoscaleServiceV2> logger = ServiceProvider.GetRequiredService<ILogger<AutoscaleServiceV2>>();
+			ILoggerFactory loggerFactory = ServiceProvider.GetRequiredService<ILoggerFactory>();
 			IOptions<ServerSettings> serverSettingsOpt = ServiceProvider.GetRequiredService<IOptions<ServerSettings>>();
+			serverSettingsOpt.Value.FleetManagerV2 = FleetManagerType.AwsReuse;
+			
+			// Empty mocks to satisfy the constructor below
+			Mock<IAmazonEC2> mockEc2 = new ();
+			Mock<IAmazonAutoScaling> mockAwsAutoScaling = new ();
 
 			AutoscaleServiceV2 service = new (AgentCollection, GraphCollection, JobCollection, LeaseCollection, PoolCollection,
-				StreamService, fleetManager, _dogStatsD, Clock, serverSettingsOpt, logger);
+				StreamService, _dogStatsD, mockEc2.Object, mockAwsAutoScaling.Object, Clock, serverSettingsOpt, loggerFactory,
+				_ => fleetManager);
 			return service;
 		}
 	}
@@ -229,4 +240,77 @@ namespace Horde.Build.Tests.Fleet
 			return AutoscaleServiceV2.CreatePoolSizeStrategy(pool);
 		}
 	}
+
+	[TestClass]
+	public class FleetManagerFactoryTest : TestSetup
+	{
+		private static int s_poolCount;
+		
+		[TestMethod]
+		public async Task CreateNoOp()
+		{
+			IFleetManager fm = await CreateFleetManager(new FleetManagerInfo(FleetManagerType.NoOp, null, "{}"));
+			Assert.AreEqual(typeof(NoOpFleetManager), fm.GetType());
+		}
+		
+		[TestMethod]
+		public async Task CreateAws()
+		{
+			IFleetManager fm = await CreateFleetManager(new FleetManagerInfo(FleetManagerType.Aws, null, "{\"ImageId\": \"bogusImageId\"}"));
+			Assert.AreEqual(typeof(AwsFleetManager), fm.GetType());
+			Assert.AreEqual("bogusImageId", ((AwsFleetManager)fm).Settings.ImageId);
+		}
+		
+		[TestMethod]
+		public async Task CreateAwsReuse()
+		{
+			IFleetManager fm = await CreateFleetManager(new FleetManagerInfo(FleetManagerType.AwsReuse, null, "{}"));
+			Assert.AreEqual(typeof(AwsReuseFleetManager), fm.GetType());
+		}
+		
+		[TestMethod]
+		public async Task CreateAwsAsg()
+		{
+			IFleetManager fm = await CreateFleetManager(new FleetManagerInfo(FleetManagerType.AwsAsg, null, "{\"Name\": \"bogusName\"}"));
+			Assert.AreEqual(typeof(AwsAsgFleetManager), fm.GetType());
+			Assert.AreEqual("bogusName", ((AwsAsgFleetManager)fm).Settings.Name);
+		}
+
+		[TestMethod]
+		public async Task UnknownConfigFieldsAreIgnored()
+		{
+			IFleetManager fm = await CreateFleetManager(new FleetManagerInfo(FleetManagerType.AwsAsg, null, "{\"Name\": \"bogusName\", \"BAD-PROPERTY\": 1337}"));
+			Assert.AreEqual(typeof(AwsAsgFleetManager), fm.GetType());
+			Assert.AreEqual("bogusName", ((AwsAsgFleetManager)fm).Settings.Name);
+		}
+		
+		[TestMethod]
+		public async Task EmptyConfigThrowsException()
+		{
+			await Assert.ThrowsExceptionAsync<JsonException>(() => CreateFleetManager(new FleetManagerInfo(FleetManagerType.AwsAsg, null, "")));
+		}
+		
+		[TestMethod]
+		public async Task CreateFromEmptyList()
+		{
+			IFleetManager fm = await CreateFleetManager();
+			Assert.AreEqual(typeof(NoOpFleetManager), fm.GetType());
+		}
+		
+		[TestMethod]
+		public async Task ConditionSimple()
+		{
+			IFleetManager s = await CreateFleetManager(
+				new FleetManagerInfo(FleetManagerType.Aws, "false", "{}"),
+				new FleetManagerInfo(FleetManagerType.AwsReuse, "true", "{\"ScaleOutFactor\": 100, \"ScaleInFactor\": 200}")
+			);
+			Assert.AreEqual(typeof(AwsReuseFleetManager), s.GetType());
+		}
+		
+		private async Task<IFleetManager> CreateFleetManager(params FleetManagerInfo[] infos)
+		{
+			IPool pool = await PoolService.CreatePoolAsync("testPool-" + s_poolCount++, null, true, 0, 0, fleetManagers: infos.ToList(), sizeStrategy: PoolSizeStrategy.NoOp);
+			return AutoscaleServiceV2.CreateFleetManager(pool);
+		}
+	}
 }
\ No newline at end of file
diff --git a/Engine/Source/Programs/Horde/Horde.Build.Tests/TestSetup.cs b/Engine/Source/Programs/Horde/Horde.Build.Tests/TestSetup.cs
--- a/Engine/Source/Programs/Horde/Horde.Build.Tests/TestSetup.cs
+++ b/Engine/Source/Programs/Horde/Horde.Build.Tests/TestSetup.cs
@@ -6,6 +6,7 @@ using System.Reflection;
 using System.Security.Claims;
 using System.Text.Json;
 using System.Threading.Tasks;
+using Amazon.AutoScaling;
 using Amazon.EC2;
 using Datadog.Trace;
 using EpicGames.Core;
@@ -203,13 +204,12 @@ namespace Horde.Build.Tests
 			
 			// Empty mocked object to satisfy basic test runs
 			services.AddSingleton<IAmazonEC2>(sp => new Mock<IAmazonEC2>().Object);
+			services.AddSingleton<IAmazonAutoScaling>(sp => new Mock<IAmazonAutoScaling>().Object);
 
 			services.AddSingleton<AclService>();
 			services.AddSingleton<AgentService>();
 			services.AddSingleton<AgentSoftwareService>();
-			services.AddSingleton<AutoscaleService>();
 			services.AddSingleton<AutoscaleServiceV2>();
-			services.AddSingleton<AwsReuseFleetManager, AwsReuseFleetManager>();
 			services.AddSingleton<ConsistencyService>();
 			services.AddSingleton<ConfigCollection>();
 			services.AddSingleton<ComputeService>();
@@ -220,7 +220,6 @@ namespace Horde.Build.Tests
 			services.AddSingleton<IDowntimeService, DowntimeServiceStub>();
 			services.AddSingleton<IDogStatsd, NoOpDogStatsd>();
 			services.AddSingleton<IssueService>();
-			services.AddSingleton<IFleetManager, AwsReuseFleetManager>();
 			services.AddSingleton<JobService>();
 			services.AddSingleton<LifetimeService>();
 			services.AddSingleton<ILogStorage, NullLogStorage>();
diff --git a/Engine/Source/Programs/Horde/Horde.Build/Agents/Fleet/AutoscaleServiceV2.cs b/Engine/Source/Programs/Horde/Horde.Build/Agents/Fleet/AutoscaleServiceV2.cs
--- a/Engine/Source/Programs/Horde/Horde.Build/Agents/Fleet/AutoscaleServiceV2.cs
+++ b/Engine/Source/Programs/Horde/Horde.Build/Agents/Fleet/AutoscaleServiceV2.cs
@@ -7,6 +7,9 @@ using System.Linq;
 using System.Text.Json;
 using System.Threading;
 using System.Threading.Tasks;
+using Amazon.AutoScaling;
+using Amazon.EC2;
+using Horde.Build.Agents.Fleet.Providers;
 using Horde.Build.Agents.Leases;
 using Horde.Build.Agents.Pools;
 using Horde.Build.Jobs;
@@ -36,14 +39,20 @@ namespace Horde.Build.Agents.Fleet
 		private readonly ILeaseCollection _leaseCollection;
 		private readonly IPoolCollection _poolCollection;
 		private readonly StreamService _streamService;
-		private readonly IFleetManager _fleetManager;
 		private readonly IDogStatsd _dogStatsd;
+		private readonly IAmazonEC2 _ec2;
+		private readonly IAmazonAutoScaling _awsAutoScaling;
 		private readonly IClock _clock;
 		private readonly ITicker _ticker;
 		private readonly ITicker _tickerHighFrequency;
 		private readonly TimeSpan _defaultScaleOutCooldown;
 		private readonly TimeSpan _defaultScaleInCooldown;
+		private readonly IOptions<ServerSettings> _settings;
 		private readonly ILogger<AutoscaleServiceV2> _logger;
+		private readonly ILoggerFactory _loggerFactory;
+		
+		/// <summary>Allow overriding the fleet manager during testing</summary>
+		private readonly Func<IPool, IFleetManager> _fleetManagerFactory;
 
 		/// <summary>
 		/// Constructor
@@ -55,11 +64,13 @@ namespace Horde.Build.Agents.Fleet
 			ILeaseCollection leaseCollection,
 			IPoolCollection poolCollection,
 			StreamService streamService,
-			IFleetManager fleetManager,
 			IDogStatsd dogStatsd,
+			IAmazonEC2 ec2,
+			IAmazonAutoScaling awsAutoScaling,
 			IClock clock,
 			IOptions<ServerSettings> settings,
-			ILogger<AutoscaleServiceV2> logger)
+			ILoggerFactory loggerFactory,
+			Func<IPool, IFleetManager>? fleetManagerFactory = null)
 		{
 			_agentCollection = agentCollection;
 			_graphCollection = graphCollection;
@@ -67,14 +78,18 @@ namespace Horde.Build.Agents.Fleet
 			_leaseCollection = leaseCollection;
 			_poolCollection = poolCollection;
 			_streamService = streamService;
-			_fleetManager = fleetManager;
 			_dogStatsd = dogStatsd;
+			_ec2 = ec2;
+			_awsAutoScaling = awsAutoScaling;
 			_clock = clock;
-			_ticker = clock.AddSharedTicker<AutoscaleServiceV2>(TimeSpan.FromMinutes(5.0), TickLeaderAsync, logger);
-			_tickerHighFrequency = clock.AddSharedTicker("AutoscaleServiceV2.TickHighFrequency", TimeSpan.FromSeconds(30), TickHighFrequencyAsync, logger);
-			_logger = logger;
+			_logger = loggerFactory.CreateLogger<AutoscaleServiceV2>();
+			_loggerFactory = loggerFactory;
+			_ticker = clock.AddSharedTicker<AutoscaleServiceV2>(TimeSpan.FromMinutes(5.0), TickLeaderAsync, _logger);
+			_tickerHighFrequency = clock.AddSharedTicker("AutoscaleServiceV2.TickHighFrequency", TimeSpan.FromSeconds(30), TickHighFrequencyAsync, _logger);
+			_settings = settings;
 			_defaultScaleOutCooldown = TimeSpan.FromSeconds(settings.Value.AgentPoolScaleOutCooldownSeconds);
 			_defaultScaleInCooldown = TimeSpan.FromSeconds(settings.Value.AgentPoolScaleInCooldownSeconds);
+			_fleetManagerFactory = fleetManagerFactory ?? CreateFleetManager;
 		}
 
 		/// <inheritdoc/>
@@ -159,7 +174,8 @@ namespace Horde.Build.Agents.Fleet
 				int deltaAgentCount = desiredAgentCount - currentAgentCount;
 
 				_logger.LogInformation("{PoolName,-48} Current={Current,4} Target={Target,4} Delta={Delta,4} Status={Status}", pool.Name, currentAgentCount, desiredAgentCount, deltaAgentCount, poolSizeData.StatusMessage);
-				
+
+				IFleetManager fleetManager = _fleetManagerFactory(poolSizeData.Pool);
 				try
 				{
 					using IScope scope = GlobalTracer.Instance.BuildSpan("ScalingPool").StartActive();
@@ -175,7 +191,7 @@ namespace Horde.Build.Agents.Fleet
 						scope.Span.SetTag("isCoolingDown", isCoolingDown);
 						if (!isCoolingDown)
 						{
-							await _fleetManager.ExpandPoolAsync(pool, poolSizeData.Agents, deltaAgentCount);
+							await fleetManager.ExpandPoolAsync(pool, poolSizeData.Agents, deltaAgentCount);
 							await _poolCollection.TryUpdateAsync(pool, lastScaleUpTime: _clock.UtcNow);
 						}
 						else
@@ -192,7 +208,7 @@ namespace Horde.Build.Agents.Fleet
 						scope.Span.SetTag("isCoolingDown", isCoolingDown);
 						if (!isCoolingDown)
 						{
-							await _fleetManager.ShrinkPoolAsync(pool, poolSizeData.Agents, -deltaAgentCount);
+							await fleetManager.ShrinkPoolAsync(pool, poolSizeData.Agents, -deltaAgentCount);
 							await _poolCollection.TryUpdateAsync(pool, lastScaleDownTime: _clock.UtcNow);
 						}
 						else
@@ -212,6 +228,83 @@ namespace Horde.Build.Agents.Fleet
 				_dogStatsd.Gauge("agentpools.autoscale.current", currentAgentCount, tags: new []{"pool:" + pool.Name});
 			}
 		}
+		
+		private IEnumerable<string> GetPropValues(string name)
+		{
+			return name switch
+			{
+				"dayOfWeek" => new List<string> { _clock.UtcNow.DayOfWeek.ToString().ToLower() },
+				_ => Array.Empty<string>()
+			};
+		}
+		
+		/// <summary>
+		/// Instantiate a fleet manager using the list of conditions/configs in <see cref="IPool" />
+		/// </summary>
+		/// <param name="pool">Pool to use</param>
+		/// <returns>A fleet manager with parameters set, as dictated by the pool argument</returns>
+		/// <exception cref="ArgumentException">If fleet manager could not be instantiated</exception>
+		public IFleetManager CreateFleetManager(IPool pool)
+		{
+			foreach (FleetManagerInfo info in pool.FleetManagers)
+			{
+				if (info.Condition == null || info.Condition.Evaluate(GetPropValues))
+				{
+					return CreateFleetManager(info.Type, info.Config);
+				}
+			}
+
+			return GetDefaultFleetManager();
+		}
+		
+		/// <summary>
+		/// Instantiate the default fleet manager  />
+		/// </summary>
+		/// <returns>A fleet manager with parameters from server settings</returns>
+		/// <exception cref="ArgumentException">If fleet manager could not be instantiated</exception>
+		public IFleetManager GetDefaultFleetManager()
+		{
+			return CreateFleetManager(_settings.Value.FleetManagerV2, _settings.Value.FleetManagerV2Config ?? "{}");
+		}
+		
+		/// <summary>
+		/// Create a fleet manager
+		/// </summary>
+		/// <param name="type">Type of fleet manager</param>
+		/// <param name="config">Config as a serialized JSON string</param>
+		/// <returns>An instantiated fleet manager with parameters loaded from config</returns>
+		/// <exception cref="ArgumentException">If fleet manager could not be instantiated</exception>
+		private IFleetManager CreateFleetManager(FleetManagerType type, string config)
+		{
+			switch (type)
+			{
+				case FleetManagerType.NoOp:
+					return new NoOpFleetManager(_loggerFactory.CreateLogger<NoOpFleetManager>());
+				case FleetManagerType.Aws:
+					return new AwsFleetManager(_ec2, _agentCollection, DeserializeSettings<AwsFleetManagerSettings>(config), _loggerFactory.CreateLogger<AwsFleetManager>());
+				case FleetManagerType.AwsReuse:
+					return new AwsReuseFleetManager(_ec2, _agentCollection, _loggerFactory.CreateLogger<AwsReuseFleetManager>());
+				case FleetManagerType.AwsAsg:
+					return new AwsAsgFleetManager(_awsAutoScaling, DeserializeSettings<AwsAsgSettings>(config), _loggerFactory.CreateLogger<AwsAsgFleetManager>());	
+				default:
+					throw new ArgumentException("Unknown fleet manager type " + type);
+			}
+		}
+		
+		private static T DeserializeSettings<T>(string config)
+		{
+			if (String.IsNullOrEmpty(config)) { config = "{}"; }
+			try
+			{
+				T? settings = JsonSerializer.Deserialize<T>(config);
+				if (settings == null) throw new NullReferenceException($"Unable to deserialize");
+				return settings;
+			}
+			catch (ArgumentException e)
+			{
+				throw new ArgumentException($"Unable to deserialize {typeof(T)} config: '{config}'", e);
+			}
+		}
 
 		/// <summary>
 		/// Instantiate a sizing strategy for the given pool
@@ -222,16 +315,6 @@ namespace Horde.Build.Agents.Fleet
 		/// <exception cref="ArgumentException"></exception>
 		public IPoolSizeStrategy CreatePoolSizeStrategy(IPool pool)
 		{
-	
-			IEnumerable<string> GetPropValues(string name)
-			{
-				return name switch
-				{
-					"dayOfWeek" => new List<string> { _clock.UtcNow.DayOfWeek.ToString().ToLower() },
-					_ => Array.Empty<string>()
-				};
-			}
-
 			if (pool.SizeStrategies.Count > 0)
 			{
 				foreach (PoolSizeStrategyInfo info in pool.SizeStrategies)
diff --git a/Engine/Source/Programs/Horde/Horde.Build/Agents/Fleet/IFleetManager.cs b/Engine/Source/Programs/Horde/Horde.Build/Agents/Fleet/IFleetManager.cs
--- a/Engine/Source/Programs/Horde/Horde.Build/Agents/Fleet/IFleetManager.cs
+++ b/Engine/Source/Programs/Horde/Horde.Build/Agents/Fleet/IFleetManager.cs
@@ -3,10 +3,37 @@
 using System.Collections.Generic;
 using System.Threading;
 using System.Threading.Tasks;
+using Horde.Build.Agents.Fleet.Providers;
 using Horde.Build.Agents.Pools;
 
 namespace Horde.Build.Agents.Fleet
 {
+	/// <summary>
+	/// Available fleet managers
+	/// </summary>
+	public enum FleetManagerType
+	{
+		/// <summary>
+		/// <see cref="NoOpFleetManager" />
+		/// </summary>
+		NoOp,
+		
+		/// <summary>
+		/// <see cref="AwsFleetManager" />
+		/// </summary>
+		Aws,
+		
+		/// <summary>
+		/// <see cref="AwsReuseFleetManager" />
+		/// </summary>
+		AwsReuse,
+		
+		/// <summary>
+		/// <see cref="AwsAsgFleetManager" />
+		/// </summary>
+		AwsAsg
+	}
+	
 	/// <summary>
 	/// Service to manage a fleet of machines
 	/// </summary>
diff --git a/Engine/Source/Programs/Horde/Horde.Build/Agents/Fleet/Providers/AwsAsgFleetManager.cs b/Engine/Source/Programs/Horde/Horde.Build/Agents/Fleet/Providers/AwsAsgFleetManager.cs
--- a/Engine/Source/Programs/Horde/Horde.Build/Agents/Fleet/Providers/AwsAsgFleetManager.cs
+++ b/Engine/Source/Programs/Horde/Horde.Build/Agents/Fleet/Providers/AwsAsgFleetManager.cs
@@ -3,12 +3,11 @@
 using System;
 using System.Collections.Generic;
 using System.Net;
+using System.Text.Json.Serialization;
 using System.Threading;
 using System.Threading.Tasks;
-using Amazon;
 using Amazon.AutoScaling;
 using Amazon.AutoScaling.Model;
-using Amazon.Extensions.NETCore.Setup;
 using Horde.Build.Agents.Pools;
 using Horde.Build.Utilities;
 using Microsoft.Extensions.Logging;
@@ -22,11 +21,6 @@ namespace Horde.Build.Agents.Fleet.Providers
 	/// </summary>
 	public class AwsAsgSettings
 	{
-		/// <summary>
-		/// AWS region (e.g us-east-1)
-		/// </summary>
-		public string Region { get; }
-		
 		/// <summary>
 		/// Name of the auto-scaling group
 		/// </summary>
@@ -35,12 +29,12 @@ namespace Horde.Build.Agents.Fleet.Providers
 		/// <summary>
 		/// Constructor
 		/// </summary>
-		/// <param name="region"></param>
-		/// <param name="name"></param>
-		public AwsAsgSettings(string region, string name)
+		/// <param name="name">Name</param>
+		[JsonConstructor]
+		public AwsAsgSettings(string? name)
 		{
-			Region = region;
-			Name = name;
+			// The JSON deserializer can call the constructor with null parameters
+			Name = name ?? throw new ArgumentException($"Parameter {nameof(name)} is null");
 		}
 	}
 	
@@ -48,29 +42,20 @@ namespace Horde.Build.Agents.Fleet.Providers
 	/// Fleet manager for handling AWS EC2 instances
 	/// Uses an EC2 auto-scaling group for controlling the number of running instances.
 	/// </summary>
-	public sealed class AwsAsgFleetManager : IFleetManager, IDisposable
+	public sealed class AwsAsgFleetManager : IFleetManager
 	{
-		private readonly AmazonAutoScalingClient _client;
-		private readonly AwsAsgSettings _asgSettings;
+		internal AwsAsgSettings Settings { get; }
+		private readonly IAmazonAutoScaling _awsAutoScaling;
 		private readonly ILogger _logger;
 
 		/// <summary>
 		/// Constructor
 		/// </summary>
-		public AwsAsgFleetManager(AWSOptions awsOptions, AwsAsgSettings asgSettings, ILogger<AwsAsgFleetManager> logger)
+		public AwsAsgFleetManager(IAmazonAutoScaling awsAutoScaling, AwsAsgSettings asgSettings, ILogger<AwsAsgFleetManager> logger)
 		{
-			_asgSettings = asgSettings;
+			_awsAutoScaling = awsAutoScaling;
+			Settings = asgSettings;
 			_logger = logger;
-
-			AmazonAutoScalingConfig config = new() { RegionEndpoint = RegionEndpoint.GetBySystemName(_asgSettings.Region) };
-			logger.LogInformation("Initializing AWS ASG fleet manager for region {Region}", config.RegionEndpoint);
-			_client = new AmazonAutoScalingClient(awsOptions.Credentials, config);
-		}
-
-		/// <inheritdoc/>
-		public void Dispose()
-		{
-			_client.Dispose();
 		}
 
 		/// <inheritdoc/>
@@ -103,13 +88,13 @@ namespace Horde.Build.Agents.Fleet.Providers
 		
 		private async Task UpdateAsg(StringId<IPool> poolId, int desiredCapacity, IScope scope, CancellationToken cancellationToken)
 		{
-			UpdateAutoScalingGroupRequest request = new () { AutoScalingGroupName = _asgSettings.Name, DesiredCapacity = desiredCapacity };
-			UpdateAutoScalingGroupResponse response = await _client.UpdateAutoScalingGroupAsync(request, cancellationToken);
+			UpdateAutoScalingGroupRequest request = new () { AutoScalingGroupName = Settings.Name, DesiredCapacity = desiredCapacity };
+			UpdateAutoScalingGroupResponse response = await _awsAutoScaling.UpdateAutoScalingGroupAsync(request, cancellationToken);
 			
 			scope.Span.SetTag("res.statusCode", (int)response.HttpStatusCode);
 			if (response.HttpStatusCode != HttpStatusCode.OK)
 			{
-				_logger.LogError("Unable to update auto-scaling group {AutoScalingGroup} for pool {PoolId}.", _asgSettings.Name, poolId);
+				_logger.LogError("Unable to update auto-scaling group {AutoScalingGroup} for pool {PoolId}.", Settings.Name, poolId);
 			}
 		}
 
diff --git a/Engine/Source/Programs/Horde/Horde.Build/Agents/Fleet/Providers/AwsFleetManager.cs b/Engine/Source/Programs/Horde/Horde.Build/Agents/Fleet/Providers/AwsFleetManager.cs
--- a/Engine/Source/Programs/Horde/Horde.Build/Agents/Fleet/Providers/AwsFleetManager.cs
+++ b/Engine/Source/Programs/Horde/Horde.Build/Agents/Fleet/Providers/AwsFleetManager.cs
@@ -5,10 +5,8 @@ using System.Collections.Generic;
 using System.Linq;
 using System.Threading;
 using System.Threading.Tasks;
-using Amazon;
 using Amazon.EC2;
 using Amazon.EC2.Model;
-using Amazon.Extensions.NETCore.Setup;
 using Horde.Build.Agents.Pools;
 using Horde.Build.Auditing;
 using Microsoft.Extensions.Logging;
@@ -18,9 +16,9 @@ using OpenTracing.Util;
 namespace Horde.Build.Agents.Fleet.Providers
 {
 	/// <summary>
-	/// Settings for launching an EC2 instance
+	/// Settings for the AWS fleet manager
 	/// </summary>
-	public class AwsInstanceLaunchSettings
+	public class AwsFleetManagerSettings
 	{
 		/// <summary>
 		/// AWS region (e.g us-east-1)
@@ -81,7 +79,7 @@ namespace Horde.Build.Agents.Fleet.Providers
 		/// <param name="securityGroupIds"></param>
 		/// <param name="userData"></param>
 		/// <param name="keyName"></param>
-		public AwsInstanceLaunchSettings(string region, string imageId, string instanceType, string? iamInstanceProfile, IReadOnlyDictionary<string, string> tags,
+		public AwsFleetManagerSettings(string region, string imageId, string instanceType, string? iamInstanceProfile, IReadOnlyDictionary<string, string> tags,
 			List<string> subnetIds, List<string> securityGroupIds, string? userData, string? keyName)
 		{
 			Region = region;
@@ -126,7 +124,7 @@ namespace Horde.Build.Agents.Fleet.Providers
 	/// Fleet manager for handling AWS EC2 instances
 	/// Will create and/or terminate instances from scratch.
 	/// </summary>
-	public sealed class AwsFleetManager : IFleetManager, IDisposable
+	public sealed class AwsFleetManager : IFleetManager
 	{
 		/// <summary>
 		/// Prefix for signaling that an EC2 instance belong to an agent pool
@@ -134,29 +132,24 @@ namespace Horde.Build.Agents.Fleet.Providers
 		public const string PoolTagName = "Horde_Autoscale_Pool";
 		private const string AwsTagPropertyName = "aws-tag";
 
-		private readonly AmazonEC2Client _client;
+		/// <summary>
+		/// Settings for fleet manager
+		/// </summary>
+		public AwsFleetManagerSettings Settings { get; }
+		
+		private readonly IAmazonEC2 _ec2;
 		private readonly IAgentCollection _agentCollection;
-		private readonly AwsInstanceLaunchSettings _launchSettings;
 		private readonly ILogger _logger;
 
 		/// <summary>
 		/// Constructor
 		/// </summary>
-		public AwsFleetManager(AWSOptions awsOptions, IAgentCollection agentCollection, AwsInstanceLaunchSettings launchSettings, ILogger<AwsFleetManager> logger)
+		public AwsFleetManager(IAmazonEC2 ec2, IAgentCollection agentCollection, AwsFleetManagerSettings settings, ILogger<AwsFleetManager> logger)
 		{
+			_ec2 = ec2;
 			_agentCollection = agentCollection;
-			_launchSettings = launchSettings;
+			Settings = settings;
 			_logger = logger;
-
-			AmazonEC2Config config = new () { RegionEndpoint = RegionEndpoint.GetBySystemName(_launchSettings.Region) };
-			logger.LogInformation("Initializing AWS fleet manager for region {Region}", config.RegionEndpoint);
-			_client = new AmazonEC2Client(awsOptions.Credentials, config);
-		}
-
-		/// <inheritdoc/>
-		public void Dispose()
-		{
-			_client.Dispose();
 		}
 
 		/// <inheritdoc/>
@@ -167,20 +160,20 @@ namespace Horde.Build.Agents.Fleet.Providers
 			scope.Span.SetTag("numAgents", agents.Count);
 			scope.Span.SetTag("count", count);
 
-			List<Tag> tags = _launchSettings.Tags.Select(x => new Tag(x.Key, x.Value)).ToList();
+			List<Tag> tags = Settings.Tags.Select(x => new Tag(x.Key, x.Value)).ToList();
 			tags.Add(new Tag("Horde:RequestedPools", pool.Name));
 			tags.Add(new Tag("Horde:Timestamp:Ec2InstanceLaunchRequested", Convert.ToString(DateTimeOffset.UtcNow.ToUnixTimeMilliseconds())));
 			
 			RunInstancesRequest request = new ()
 			{
-				ImageId = _launchSettings.ImageId,
-				InstanceType = InstanceType.FindValue(_launchSettings.InstanceType),
-				KeyName = _launchSettings.KeyName,
-				SecurityGroupIds = _launchSettings.SecurityGroupIds,
-				SubnetId = _launchSettings.GetRandomSubnetId(), // Pick randomly to spread evenly among available subnets 
+				ImageId = Settings.ImageId,
+				InstanceType = InstanceType.FindValue(Settings.InstanceType),
+				KeyName = Settings.KeyName,
+				SecurityGroupIds = Settings.SecurityGroupIds,
+				SubnetId = Settings.GetRandomSubnetId(), // Pick randomly to spread evenly among available subnets 
 				MinCount = count,
 				MaxCount = count,
-				UserData = _launchSettings.GetUserDataAsBase64(),
+				UserData = Settings.GetUserDataAsBase64(),
 				InstanceInitiatedShutdownBehavior = ShutdownBehavior.Terminate,
 				MetadataOptions = new InstanceMetadataOptionsRequest
 				{
@@ -189,12 +182,12 @@ namespace Horde.Build.Agents.Fleet.Providers
 				}
 			};
 
-			if (_launchSettings.IamInstanceProfile != null)
+			if (Settings.IamInstanceProfile != null)
 			{
 				request.IamInstanceProfile =
-					_launchSettings.IamInstanceProfile.StartsWith("arn:aws", StringComparison.InvariantCulture)
-						? new IamInstanceProfileSpecification { Arn = _launchSettings.IamInstanceProfile }
-						: new IamInstanceProfileSpecification { Name = _launchSettings.IamInstanceProfile };
+					Settings.IamInstanceProfile.StartsWith("arn:aws", StringComparison.InvariantCulture)
+						? new IamInstanceProfileSpecification { Arn = Settings.IamInstanceProfile }
+						: new IamInstanceProfileSpecification { Name = Settings.IamInstanceProfile };
 			}
 			
 			if (tags.Count > 0)
@@ -202,7 +195,7 @@ namespace Horde.Build.Agents.Fleet.Providers
 				request.TagSpecifications = new() { new() { ResourceType = ResourceType.Instance, Tags = tags } };
 			}
 
-			RunInstancesResponse response = await _client.RunInstancesAsync(request, cancellationToken);
+			RunInstancesResponse response = await _ec2.RunInstancesAsync(request, cancellationToken);
 			int numStartedInstances = response.Reservation.Instances.Count;
 			scope.Span.SetTag("res.statusCode", (int)response.HttpStatusCode);
 			scope.Span.SetTag("res.numInstances", numStartedInstances);
@@ -276,7 +269,7 @@ namespace Horde.Build.Agents.Fleet.Providers
 			describeRequest.Filters.Add(new Filter("instance-state-name", new List<string> { InstanceStateName.Stopped.Value }));
 			describeRequest.Filters.Add(new Filter("tag:" + PoolTagName, new List<string> { pool.Name }));
 
-			DescribeInstancesResponse describeResponse = await _client.DescribeInstancesAsync(describeRequest, cancellationToken);
+			DescribeInstancesResponse describeResponse = await _ec2.DescribeInstancesAsync(describeRequest, cancellationToken);
 			return describeResponse.Reservations.SelectMany(x => x.Instances).Select(x => x.InstanceId).Distinct().Count();
 		}
 	}
diff --git a/Engine/Source/Programs/Horde/Horde.Build/Agents/Fleet/Providers/DefaultFleetManager.cs b/Engine/Source/Programs/Horde/Horde.Build/Agents/Fleet/Providers/DefaultFleetManager.cs
deleted file mode 100644
diff --git a/Engine/Source/Programs/Horde/Horde.Build/Agents/Fleet/Providers/NoOpFleetManager.cs b/Engine/Source/Programs/Horde/Horde.Build/Agents/Fleet/Providers/NoOpFleetManager.cs
new file mode 100644
--- /dev/null
+++ b/Engine/Source/Programs/Horde/Horde.Build/Agents/Fleet/Providers/NoOpFleetManager.cs
@@ -0,0 +1,44 @@
+// Copyright Epic Games, Inc. All Rights Reserved.
+
+using System.Collections.Generic;
+using System.Threading;
+using System.Threading.Tasks;
+using Horde.Build.Agents.Pools;
+using Microsoft.Extensions.Logging;
+
+namespace Horde.Build.Agents.Fleet.Providers
+{
+	/// <summary>
+	/// No-op implementation of <see cref="IFleetManager"/>
+	/// </summary>
+	public class NoOpFleetManager : IFleetManager
+	{
+		readonly ILogger _logger;
+
+		/// <summary>
+		/// Constructor
+		/// </summary>
+		/// <param name="logger">Logging device</param>
+		public NoOpFleetManager(ILogger<NoOpFleetManager> logger)
+		{
+			_logger = logger;
+		}
+
+		/// <inheritdoc/>
+		public Task ExpandPoolAsync(IPool pool, IReadOnlyList<IAgent> agents, int count, CancellationToken cancellationToken)
+		{
+			_logger.LogInformation("Expand pool {PoolId} by {Count} agents", pool.Id, count);
+			return Task.CompletedTask;
+		}
+
+		/// <inheritdoc/>
+		public Task ShrinkPoolAsync(IPool pool, IReadOnlyList<IAgent> agents, int count, CancellationToken cancellationToken)
+		{
+			_logger.LogInformation("Shrink pool {PoolId} by {Count} agents", pool.Id, count);
+			return Task.CompletedTask;
+		}
+
+		/// <inheritdoc/>
+		public Task<int> GetNumStoppedInstancesAsync(IPool pool, CancellationToken cancellationToken) => Task.FromResult(0);
+	}
+}
diff --git a/Engine/Source/Programs/Horde/Horde.Build/Agents/Pools/IPool.cs b/Engine/Source/Programs/Horde/Horde.Build/Agents/Pools/IPool.cs
--- a/Engine/Source/Programs/Horde/Horde.Build/Agents/Pools/IPool.cs
+++ b/Engine/Source/Programs/Horde/Horde.Build/Agents/Pools/IPool.cs
@@ -2,6 +2,8 @@
 
 using System;
 using System.Collections.Generic;
+using System.Text.Json;
+using System.Text.Json.Serialization;
 using EpicGames.Horde.Common;
 using EpicGames.Horde.Compute;
 using Horde.Build.Agents.Fleet;
@@ -49,6 +51,49 @@ namespace Horde.Build.Agents.Pools
 		}
 	}
 	
+	/// <summary>
+	/// Metadata for configuring and picking a fleet manager
+	/// </summary>
+	public class FleetManagerInfo
+	{
+		/// <summary>
+		/// Fleet manager type implementation to use
+		/// </summary>
+		public FleetManagerType Type { get; set; }
+		
+		/// <summary>
+		/// Condition if this strategy should be enabled (right now, using date/time as a distinguishing factor)
+		/// </summary>
+		public Condition? Condition { get; set; }
+
+		/// <summary>
+		/// Configuration for the strategy, serialized as JSON
+		/// </summary>
+		public string Config { get; set; } = "{}";
+
+		/// <summary>
+		/// Empty constructor for BSON/JSON serialization
+		/// </summary>
+		public FleetManagerInfo()
+		{
+		}
+		
+		/// <summary>
+		/// Constructor
+		/// </summary>
+		[JsonConstructor]
+		public FleetManagerInfo(FleetManagerType type, Condition? condition, string? config)
+		{
+			// Try deserializing to ensure the config is valid JSON
+			if (config == null) throw new ArgumentNullException(nameof(config));
+			JsonSerializer.Deserialize<dynamic>(config);
+			
+			Type = type;
+			Condition = condition;
+			Config = config;
+		}
+	}
+	
 	/// <summary>
 	/// A pool of machines
 	/// </summary>
@@ -133,6 +178,12 @@ namespace Horde.Build.Agents.Pools
 		/// List of pool sizing strategies for this pool. The first strategy with a matching condition will be picked.
 		/// </summary>
 		public IReadOnlyList<PoolSizeStrategyInfo> SizeStrategies { get; }
+
+		/// <summary>
+		/// List of fleet managers for this pool. The first strategy with a matching condition will be picked.
+		/// If empty or no conditions match, a default fleet manager will be used.
+		/// </summary>
+		public IReadOnlyList<FleetManagerInfo> FleetManagers { get; }
 		
 		/// <summary>
 		/// Settings for lease utilization pool sizing strategy (if used)
diff --git a/Engine/Source/Programs/Horde/Horde.Build/Agents/Pools/IPoolCollection.cs b/Engine/Source/Programs/Horde/Horde.Build/Agents/Pools/IPoolCollection.cs
--- a/Engine/Source/Programs/Horde/Horde.Build/Agents/Pools/IPoolCollection.cs
+++ b/Engine/Source/Programs/Horde/Horde.Build/Agents/Pools/IPoolCollection.cs
@@ -29,6 +29,7 @@ namespace Horde.Build.Agents.Pools
 		/// <param name="scaleOutCooldown">Cooldown time between scale-out events</param>
 		/// <param name="scaleInCooldown">Cooldown time between scale-in events</param>
 		/// <param name="sizeStrategies">Pool sizing strategies</param>
+		/// <param name="fleetManagers">Fleet managers</param>
 		/// <param name="sizeStrategy">Pool sizing strategy</param>
 		/// <param name="leaseUtilizationSettings">Settings for lease utilization strategy</param>
 		/// <param name="jobQueueSettings">Settings for job queue strategy</param>
@@ -46,6 +47,7 @@ namespace Horde.Build.Agents.Pools
 			TimeSpan? scaleOutCooldown = null,
 			TimeSpan? scaleInCooldown = null,
 			List<PoolSizeStrategyInfo>? sizeStrategies = null,
+			List<FleetManagerInfo>? fleetManagers = null,
 			PoolSizeStrategy? sizeStrategy = null,
 			LeaseUtilizationSettings? leaseUtilizationSettings = null,
 			JobQueueSettings? jobQueueSettings = null,
@@ -97,6 +99,7 @@ namespace Horde.Build.Agents.Pools
 		/// <param name="scaleInCooldown">Cooldown time between scale-in events</param>
 		/// <param name="sizeStrategy">Pool sizing strategy</param>
 		/// <param name="newSizeStrategies">List of pool sizing strategies</param>
+		/// <param name="newFleetManagers">List of fleet managers</param>
 		/// <param name="leaseUtilizationSettings">Settings for lease utilization strategy</param>
 		/// <param name="jobQueueSettings">Settings for job queue strategy</param>
 		/// <param name="computeQueueAwsMetricSettings">Settings for compute queue AWS metric strategy</param>
@@ -119,6 +122,7 @@ namespace Horde.Build.Agents.Pools
 			TimeSpan? scaleInCooldown = null,
 			PoolSizeStrategy? sizeStrategy = null,
 			List<PoolSizeStrategyInfo>? newSizeStrategies = null,
+			List<FleetManagerInfo>? newFleetManagers = null,
 			LeaseUtilizationSettings? leaseUtilizationSettings = null,
 			JobQueueSettings? jobQueueSettings = null,
 			ComputeQueueAwsMetricSettings? computeQueueAwsMetricSettings = null,
diff --git a/Engine/Source/Programs/Horde/Horde.Build/Agents/Pools/PoolCollection.cs b/Engine/Source/Programs/Horde/Horde.Build/Agents/Pools/PoolCollection.cs
--- a/Engine/Source/Programs/Horde/Horde.Build/Agents/Pools/PoolCollection.cs
+++ b/Engine/Source/Programs/Horde/Horde.Build/Agents/Pools/PoolCollection.cs
@@ -45,6 +45,7 @@ namespace Horde.Build.Agents.Pools
 			public PoolSizeStrategy? SizeStrategy { get; set; }
 			PoolSizeStrategy IPool.SizeStrategy => SizeStrategy ?? s_defaultStrategy!.Value;
 			public List<PoolSizeStrategyInfo> SizeStrategies { get; set; } = new ();
+			public List<FleetManagerInfo> FleetManagers { get; set; } = new ();
 			public LeaseUtilizationSettings? LeaseUtilizationSettings { get; set; }
 			public JobQueueSettings? JobQueueSettings { get; set; }
 			public ComputeQueueAwsMetricSettings? ComputeQueueAwsMetricSettings { get; set; }
@@ -54,6 +55,7 @@ namespace Horde.Build.Agents.Pools
 			IReadOnlyList<AgentWorkspace> IPool.Workspaces => Workspaces;
 			IReadOnlyDictionary<string, string> IPool.Properties => Properties;
 			IReadOnlyList<PoolSizeStrategyInfo> IPool.SizeStrategies => SizeStrategies;
+			IReadOnlyList<FleetManagerInfo> IPool.FleetManagers => FleetManagers;
 
 			public PoolDocument()
 			{
@@ -77,6 +79,7 @@ namespace Horde.Build.Agents.Pools
 				ScaleInCooldown = other.ScaleInCooldown;
 				SizeStrategy = other.SizeStrategy;
 				SizeStrategies.AddRange(other.SizeStrategies);
+				FleetManagers.AddRange(other.FleetManagers);
 				LeaseUtilizationSettings = other.LeaseUtilizationSettings;
 				JobQueueSettings = other.JobQueueSettings;
 				ComputeQueueAwsMetricSettings = other.ComputeQueueAwsMetricSettings;
@@ -114,6 +117,7 @@ namespace Horde.Build.Agents.Pools
 			TimeSpan? scaleOutCooldown,
 			TimeSpan? scaleInCooldown,
 			List<PoolSizeStrategyInfo>? sizeStrategies,
+			List<FleetManagerInfo>? fleetManagers,
 			PoolSizeStrategy? sizeStrategy,
 			LeaseUtilizationSettings? leaseUtilizationSettings,
 			JobQueueSettings? jobQueueSettings,
@@ -142,10 +146,9 @@ namespace Horde.Build.Agents.Pools
 			pool.LeaseUtilizationSettings = leaseUtilizationSettings;
 			pool.JobQueueSettings = jobQueueSettings;
 			pool.ComputeQueueAwsMetricSettings = computeQueueAwsMetricSettings;
-			if (sizeStrategies != null)
-			{
-				pool.SizeStrategies = sizeStrategies;
-			}
+			if (sizeStrategies != null) { pool.SizeStrategies = sizeStrategies; }
+			if (fleetManagers != null) { pool.FleetManagers = fleetManagers; }
+			
 			await _pools.InsertOneAsync(pool);
 			return pool;
 		}
@@ -237,6 +240,7 @@ namespace Horde.Build.Agents.Pools
 			TimeSpan? scaleInCooldown, 
 			PoolSizeStrategy? sizeStrategy,
 			List<PoolSizeStrategyInfo>? newSizeStrategies,
+			List<FleetManagerInfo>? newFleetManagers,
 			LeaseUtilizationSettings? leaseUtilizationSettings,
 			JobQueueSettings? jobQueueSettings, 
 			ComputeQueueAwsMetricSettings? computeQueueAwsMetricSettings, 
@@ -342,7 +346,10 @@ namespace Horde.Build.Agents.Pools
 			{
 				transaction.Set(x => x.SizeStrategies, newSizeStrategies);				
 			}
-
+			if (newFleetManagers != null)
+			{
+				transaction.Set(x => x.FleetManagers, newFleetManagers);				
+			}
 			if (leaseUtilizationSettings != null)
 			{
 				transaction.Set(x => x.LeaseUtilizationSettings, leaseUtilizationSettings);
diff --git a/Engine/Source/Programs/Horde/Horde.Build/Agents/Pools/PoolMessages.cs b/Engine/Source/Programs/Horde/Horde.Build/Agents/Pools/PoolMessages.cs
--- a/Engine/Source/Programs/Horde/Horde.Build/Agents/Pools/PoolMessages.cs
+++ b/Engine/Source/Programs/Horde/Horde.Build/Agents/Pools/PoolMessages.cs
@@ -53,6 +53,52 @@ namespace Horde.Build.Agents.Pools
 			return new PoolSizeStrategyInfo(Type, Condition, Config);
 		}
 	}
+	
+	/// <see cref="FleetManagerInfo" />
+	public class FleetManagerMessage
+	{
+		/// <summary>
+		/// Strategy implementation to use
+		/// </summary>
+		public FleetManagerType Type { get; set; }
+		
+		/// <summary>
+		/// Condition if this strategy should be enabled (right now, using date/time as a distinguishing factor)
+		/// </summary>
+		public Condition? Condition { get; set; }
+
+		/// <summary>
+		/// Configuration for the strategy, serialized as JSON
+		/// </summary>
+		public string Config { get; set; } = "";
+		
+		/// <summary>
+		/// Empty constructor for JSON serialization
+		/// </summary>
+		public FleetManagerMessage()
+		{
+		}
+
+		/// <summary>
+		/// Construct a public REST API representation from the internal one
+		/// </summary>
+		/// <param name="info"></param>
+		public FleetManagerMessage(FleetManagerInfo info)
+		{
+			Type = info.Type;
+			Condition = info.Condition;
+			Config = info.Config;
+		}
+		
+		/// <summary>
+		/// Convert this public REST API object to an internal representation
+		/// </summary>
+		/// <returns></returns>
+		public FleetManagerInfo Convert()
+		{
+			return new FleetManagerInfo(Type, Condition, Config);
+		}
+	}
 
 	/// <see cref="LeaseUtilizationSettings" />
 	public class LeaseUtilizationSettingsMessage
@@ -202,6 +248,11 @@ namespace Horde.Build.Agents.Pools
 		/// </summary>
 		public List<PoolSizeStrategyMessage>? SizeStrategies { get; set; }
 		
+		/// <summary>
+		/// Fleet managers
+		/// </summary>
+		public List<FleetManagerMessage>? FleetManagers { get; set; }
+		
 		/// <summary>
 		/// Pool sizing strategy
 		/// </summary>
@@ -302,6 +353,11 @@ namespace Horde.Build.Agents.Pools
 		/// Pool sizing strategies
 		/// </summary>
 		public List<PoolSizeStrategyMessage>? SizeStrategies { get; set; }
+		
+		/// <summary>
+		/// Fleet managers
+		/// </summary>
+		public List<FleetManagerMessage>? FleetManagers { get; set; }
 
 		/// <summary>
 		/// Set pool to use default strategy
@@ -405,6 +461,11 @@ namespace Horde.Build.Agents.Pools
 		/// Pool sizing strategies to be used for this pool
 		/// </summary>
 		public List<PoolSizeStrategyMessage> SizeStrategies { get; set; }
+		
+		/// <summary>
+		/// Fleet managers to be used for this pool
+		/// </summary>
+		public List<FleetManagerMessage> FleetManagers { get; set; }
 
 		/// <summary>
 		/// Pool sizing strategy to be used for this pool (deprecated, see SizeStrategies field)
@@ -460,6 +521,7 @@ namespace Horde.Build.Agents.Pools
 			ScaleOutCooldown = pool.ScaleOutCooldown == null ? null : (int)pool.ScaleOutCooldown.Value.TotalSeconds;
 			ScaleInCooldown = pool.ScaleInCooldown == null ? null : (int)pool.ScaleInCooldown.Value.TotalSeconds;
 			SizeStrategies = pool.SizeStrategies.Select(x => new PoolSizeStrategyMessage(x)).ToList();
+			FleetManagers = pool.FleetManagers.Select(x => new FleetManagerMessage(x)).ToList();
 			SizeStrategy = pool.SizeStrategy;
 			LeaseUtilizationSettings = pool.LeaseUtilizationSettings == null ? null : new LeaseUtilizationSettingsMessage(pool.LeaseUtilizationSettings);
 			JobQueueSettings = pool.JobQueueSettings == null ? null : new JobQueueSettingsMessage(pool.JobQueueSettings);
diff --git a/Engine/Source/Programs/Horde/Horde.Build/Agents/Pools/PoolService.cs b/Engine/Source/Programs/Horde/Horde.Build/Agents/Pools/PoolService.cs
--- a/Engine/Source/Programs/Horde/Horde.Build/Agents/Pools/PoolService.cs
+++ b/Engine/Source/Programs/Horde/Horde.Build/Agents/Pools/PoolService.cs
@@ -65,6 +65,7 @@ namespace Horde.Build.Agents.Pools
 		/// <param name="scaleOutCooldown">Cooldown time between scale-out events</param>
 		/// <param name="scaleInCooldown">Cooldown time between scale-in events</param>
 		/// <param name="sizeStrategies">Pool sizing strategies</param>
+		/// <param name="fleetManagers">Fleet managers</param>
 		/// <param name="sizeStrategy">Pool sizing strategy</param>
 		/// <param name="leaseUtilizationSettings">Settings for lease utilization strategy</param>
 		/// <param name="jobQueueSettings">Settings for job queue strategy</param>
@@ -81,6 +82,7 @@ namespace Horde.Build.Agents.Pools
 			TimeSpan? scaleOutCooldown = null,
 			TimeSpan? scaleInCooldown = null,
 			List<PoolSizeStrategyInfo>? sizeStrategies = null,
+			List<FleetManagerInfo>? fleetManagers = null,
 			PoolSizeStrategy? sizeStrategy = null,
 			LeaseUtilizationSettings? leaseUtilizationSettings = null,
 			JobQueueSettings? jobQueueSettings = null,
@@ -98,6 +100,7 @@ namespace Horde.Build.Agents.Pools
 				scaleOutCooldown,
 				scaleInCooldown,
 				sizeStrategies,
+				fleetManagers,
 				sizeStrategy,
 				leaseUtilizationSettings,
 				jobQueueSettings,
@@ -130,6 +133,7 @@ namespace Horde.Build.Agents.Pools
 		/// <param name="scaleInCooldown">Cooldown time between scale-in events</param>
 		/// <param name="sizeStrategy">New pool sizing strategy for the pool</param>
 		/// <param name="newSizeStrategies">New pool sizing strategies for the pool</param>
+		/// <param name="newFleetManagers">New fleet managers for the pool</param>
 		/// <param name="leaseUtilizationSettings">Settings for lease utilization-based strategy</param>
 		/// <param name="jobQueueSettings">Settings for job queue-based strategy</param>
 		/// <param name="computeQueueAwsMetricSettings">Settings for compute queue AWS metric strategy</param>
@@ -148,6 +152,7 @@ namespace Horde.Build.Agents.Pools
 			TimeSpan? scaleInCooldown = null,
 			PoolSizeStrategy? sizeStrategy = null,
 			List<PoolSizeStrategyInfo>? newSizeStrategies = null,
+			List<FleetManagerInfo>? newFleetManagers = null,
 			LeaseUtilizationSettings? leaseUtilizationSettings = null,
 			JobQueueSettings? jobQueueSettings = null,
 			ComputeQueueAwsMetricSettings? computeQueueAwsMetricSettings = null,
@@ -168,6 +173,7 @@ namespace Horde.Build.Agents.Pools
 					scaleInCooldown: scaleInCooldown,
 					sizeStrategy: sizeStrategy,
 					newSizeStrategies: newSizeStrategies,
+					newFleetManagers: newFleetManagers,
 					leaseUtilizationSettings: leaseUtilizationSettings,
 					jobQueueSettings: jobQueueSettings,
 					computeQueueAwsMetricSettings: computeQueueAwsMetricSettings,
diff --git a/Engine/Source/Programs/Horde/Horde.Build/Agents/Pools/PoolsController.cs b/Engine/Source/Programs/Horde/Horde.Build/Agents/Pools/PoolsController.cs
--- a/Engine/Source/Programs/Horde/Horde.Build/Agents/Pools/PoolsController.cs
+++ b/Engine/Source/Programs/Horde/Horde.Build/Agents/Pools/PoolsController.cs
@@ -58,6 +58,7 @@ namespace Horde.Build.Agents.Pools
 			}
 
 			List<PoolSizeStrategyInfo>? sizeStrategies = create.SizeStrategies?.Select(x => x.Convert()).ToList();
+			List<FleetManagerInfo>? fleetManagers = create.FleetManagers?.Select(x => x.Convert()).ToList();
 			LeaseUtilizationSettings? luSettings = create.LeaseUtilizationSettings?.Convert();
 			JobQueueSettings? jqSettings = create.JobQueueSettings?.Convert();
 			ComputeQueueAwsMetricSettings? cqamSettings = create.ComputeQueueAwsMetricSettings?.Convert();
@@ -68,7 +69,7 @@ namespace Horde.Build.Agents.Pools
 
 			IPool newPool = await _poolService.CreatePoolAsync(
 				create.Name, create.Condition, create.EnableAutoscaling, create.MinAgents, create.NumReserveAgents, conformInterval,
-				scaleOutCooldown, scaleInCooldown, sizeStrategies, create.SizeStrategy, luSettings, jqSettings, cqamSettings, create.Properties);
+				scaleOutCooldown, scaleInCooldown, sizeStrategies, fleetManagers, create.SizeStrategy, luSettings, jqSettings, cqamSettings, create.Properties);
 			return new CreatePoolResponse(newPool.Id.ToString());
 		}
 
@@ -148,12 +149,13 @@ namespace Horde.Build.Agents.Pools
 			}
 
 			List<PoolSizeStrategyInfo>? newSizeStrategies = update.SizeStrategies?.Select(x => x.Convert()).ToList();
+			List<FleetManagerInfo>? newFleetManagers = update.FleetManagers?.Select(x => x.Convert()).ToList();
 			TimeSpan? conformInterval = update.ConformInterval == null ? null : TimeSpan.FromHours(update.ConformInterval.Value);
 			TimeSpan? scaleOutCooldown = update.ScaleOutCooldown == null ? null : TimeSpan.FromSeconds(update.ScaleOutCooldown.Value);
 			TimeSpan? scaleInCooldown = update.ScaleInCooldown == null ? null : TimeSpan.FromSeconds(update.ScaleInCooldown.Value);
 
 			await _poolService.UpdatePoolAsync(pool, update.Name, update.Condition, update.EnableAutoscaling,
-				update.MinAgents, update.NumReserveAgents, update.Properties, conformInterval, scaleOutCooldown, scaleInCooldown, update.SizeStrategy, newSizeStrategies,
+				update.MinAgents, update.NumReserveAgents, update.Properties, conformInterval, scaleOutCooldown, scaleInCooldown, update.SizeStrategy, newSizeStrategies, newFleetManagers,
 				update.LeaseUtilizationSettings?.Convert(), update.JobQueueSettings?.Convert(), update.ComputeQueueAwsMetricSettings?.Convert(), update.UseDefaultStrategy);
 			return new OkResult();
 		}
diff --git a/Engine/Source/Programs/Horde/Horde.Build/ServerSettings.cs b/Engine/Source/Programs/Horde/Horde.Build/ServerSettings.cs
--- a/Engine/Source/Programs/Horde/Horde.Build/ServerSettings.cs
+++ b/Engine/Source/Programs/Horde/Horde.Build/ServerSettings.cs
@@ -120,22 +120,6 @@ namespace Horde.Build
 		public ChunkingOptions Chunking { get; set; } = new ChunkingOptions();
 	}
 
-	/// <summary>
-	/// Specifies the service to use for controlling the size of the fleet
-	/// </summary>
-	public enum FleetManagerType
-	{
-		/// <summary>
-		/// Default (empty) instance
-		/// </summary>
-		None,
-
-		/// <summary>
-		/// Use AWS EC2 instances
-		/// </summary>
-		Aws,
-	}
-	
 	/// <summary>
 	/// Authentication method used for logging users in
 	/// </summary>
@@ -408,9 +392,14 @@ namespace Horde.Build
 		public bool LogSessionRequests { get; set; } = false;
 
 		/// <summary>
-		/// Which fleet manager service to use
+		/// Default fleet manager to use (when not specified by pool)
+		/// </summary>
+		public FleetManagerType FleetManagerV2 { get; set; } = FleetManagerType.NoOp;
+		
+		/// <summary>
+		/// Config for the fleet manager (serialized JSON)
 		/// </summary>
-		public FleetManagerType FleetManager { get; set; } = FleetManagerType.None;
+		public string? FleetManagerV2Config { get; set; }
 
 		/// <summary>
 		/// Whether to run scheduled jobs.
diff --git a/Engine/Source/Programs/Horde/Horde.Build/Startup.cs b/Engine/Source/Programs/Horde/Horde.Build/Startup.cs
--- a/Engine/Source/Programs/Horde/Horde.Build/Startup.cs
+++ b/Engine/Source/Programs/Horde/Horde.Build/Startup.cs
@@ -13,6 +13,7 @@ using System.Text.Json.Serialization;
 using System.Threading;
 using System.Threading.Tasks;
 using Amazon;
+using Amazon.AutoScaling;
 using Amazon.CloudWatch;
 using Amazon.EC2;
 using Amazon.Extensions.NETCore.Setup;
@@ -385,22 +386,24 @@ namespace Horde.Build
 			services.AddSingleton(typeof(IAuditLogFactory<>), typeof(AuditLogFactory<>));
 			services.AddSingleton(typeof(ISingletonDocument<>), typeof(SingletonDocument<>));
 
-			services.AddSingleton<AutoscaleService>();
 			services.AddSingleton<AutoscaleServiceV2>();
 			services.AddSingleton<LeaseUtilizationStrategy>();
 			services.AddSingleton<JobQueueStrategy>();
 			services.AddSingleton<NoOpPoolSizeStrategy>();
 			services.AddSingleton<ComputeQueueAwsMetricStrategy>();
 			
-			switch (settings.FleetManager)
+			// Associate IFleetManager interface with the default implementation from config for convenience
+			// Though most fleet managers are created on a per-pool basis
+			services.AddSingleton<IFleetManager>(ctx =>
 			{
-				case FleetManagerType.Aws:
-					services.AddSingleton<IFleetManager, AwsReuseFleetManager>();
-					break;
-				default:
-					services.AddSingleton<IFleetManager, DefaultFleetManager>();
-					break;
-			}
+				AutoscaleServiceV2? autoScale = ctx.GetService<AutoscaleServiceV2>();
+				if (autoScale == null)
+				{
+					throw new InvalidOperationException($"{nameof(AutoscaleServiceV2)} not registered!");
+				}
+
+				return autoScale.GetDefaultFleetManager();
+			});
 
 			services.AddSingleton<AclService>();
 			services.AddSingleton<AgentService>();			
@@ -475,6 +478,7 @@ namespace Horde.Build
 			}
 			
 			services.AddAWSService<IAmazonCloudWatch>();
+			services.AddAWSService<IAmazonAutoScaling>();
 			services.AddAWSService<IAmazonEC2>();
 
 			ConfigureLogStorage(services);
