Missing file.

diff --git a/Engine/Source/Programs/Horde/Horde.Agent/Commands/Bundles/ExtractCommand.cs b/Engine/Source/Programs/Horde/Horde.Agent/Commands/Bundles/ExtractCommand.cs
--- a/Engine/Source/Programs/Horde/Horde.Agent/Commands/Bundles/ExtractCommand.cs
+++ b/Engine/Source/Programs/Horde/Horde.Agent/Commands/Bundles/ExtractCommand.cs
@@ -20,9 +20,9 @@ namespace Horde.Agent.Commands.Bundles
 
 		public override async Task<int> ExecuteAsync(ILogger logger)
 		{
-			using ITreeStore store = CreateTreeStore(logger);
+			IStorageClient storage = CreateStorageClient(logger);
 
-			DirectoryNode node = await store.ReadTreeAsync<DirectoryNode>(RefName);
+			DirectoryNode node = await storage.ReadTreeAsync<DirectoryNode>(RefName);
 			await node.CopyToDirectoryAsync(OutputDir.ToDirectoryInfo(), logger, CancellationToken.None);
 
 			return 0;
