Horde: Fix inconsistencies creating "default" StringId values, vs. using StringId.Empty.

diff --git a/Engine/Source/Programs/Horde/Horde.Build/Utilities/StringId.cs b/Engine/Source/Programs/Horde/Horde.Build/Utilities/StringId.cs
--- a/Engine/Source/Programs/Horde/Horde.Build/Utilities/StringId.cs
+++ b/Engine/Source/Programs/Horde/Horde.Build/Utilities/StringId.cs
@@ -20,13 +20,14 @@ namespace Horde.Build.Utilities
 	[JsonSchemaString]
 	[TypeConverter(typeof(StringIdTypeConverter))]
 	[CbConverter(typeof(CbStringIdConverter<>))]
+	[JsonConverter(typeof(JsonStringIdConverterFactory))]
 	public struct StringId<T> : IEquatable<StringId<T>>, IComparable<StringId<T>>
 	{
 		/// <summary>
 		/// Empty string
 		/// </summary>
 		[SuppressMessage("Design", "CA1000:Do not declare static members on generic types", Justification = "<Pending>")]
-		public static StringId<T> Empty { get; } = new StringId<T>(String.Empty);
+		public static StringId<T> Empty { get; } = default;
 
 		/// <summary>
 		/// The text representing this id
@@ -40,31 +41,33 @@ namespace Horde.Build.Utilities
 		[SuppressMessage("Globalization", "CA1308:Normalize strings to uppercase", Justification = "<Pending>")]
 		public StringId(string input)
 		{
-			_text = input;
-
-			if (_text.Length == 0)
+			if (String.IsNullOrEmpty(input))
 			{
-//				throw new ArgumentException("String id may not be empty");
+				_text = null!;
 			}
-
-			const int MaxLength = 64;
-			if (_text.Length > MaxLength)
+			else
 			{
-				throw new ArgumentException($"String id may not be longer than {MaxLength} characters");
-			}
+				_text = input;
 
-			for (int idx = 0; idx < _text.Length; idx++)
-			{
-				char character = _text[idx];
-				if (!IsValidCharacter(character))
+				const int MaxLength = 64;
+				if (_text.Length > MaxLength)
 				{
-					if (character >= 'A' && character <= 'Z')
-					{
-						_text = _text.ToLowerInvariant();
-					}
-					else
+					throw new ArgumentException($"String id may not be longer than {MaxLength} characters");
+				}
+
+				for (int idx = 0; idx < _text.Length; idx++)
+				{
+					char character = _text[idx];
+					if (!IsValidCharacter(character))
 					{
-						throw new ArgumentException($"{_text} is not a valid string id");
+						if (character >= 'A' && character <= 'Z')
+						{
+							_text = _text.ToLowerInvariant();
+						}
+						else
+						{
+							throw new ArgumentException($"{_text} is not a valid string id");
+						}
 					}
 				}
 			}
@@ -156,13 +159,13 @@ namespace Horde.Build.Utilities
 		/// <inheritdoc/>
 		public override int GetHashCode()
 		{
-			return _text.GetHashCode(StringComparison.Ordinal);
+			return String.GetHashCode(_text, StringComparison.Ordinal);
 		}
 
 		/// <inheritdoc/>
 		public bool Equals(StringId<T> other)
 		{
-			return _text.Equals(other._text, StringComparison.Ordinal);
+			return String.Equals(_text, other._text, StringComparison.Ordinal);
 		}
 
 		/// <inheritdoc/>
