Horde: Add a service to manage IStorageClient instances for different namespaces.

* Each IStorageClient is implemented using a ref store in the Horde database and cache in Redis.
* Namespaces are configured to use backends defined in the global config file, which may inherit and mutate settings from other backends.

diff --git a/Engine/Source/Programs/Horde/Horde.Build.Tests/BlobStoreTests.cs b/Engine/Source/Programs/Horde/Horde.Build.Tests/BlobStoreTests.cs
--- a/Engine/Source/Programs/Horde/Horde.Build.Tests/BlobStoreTests.cs
+++ b/Engine/Source/Programs/Horde/Horde.Build.Tests/BlobStoreTests.cs
@@ -3,43 +3,22 @@
 using System;
 using System.Collections.Generic;
 using System.Linq;
-using System.Runtime.InteropServices;
-using System.Text.Json;
 using System.Threading.Tasks;
 using EpicGames.Core;
-using Horde.Build.Jobs.Graphs;
-using Horde.Build.Jobs;
-using Horde.Build.Issues;
-using Horde.Build.Logs;
-using Horde.Build.Users;
-using Horde.Build.Projects;
-using Horde.Build.Streams;
-using Horde.Build.Server;
-using Horde.Build.Tests.Stubs.Services;
-using Horde.Build.Utilities;
-using HordeCommon;
-using Microsoft.Extensions.DependencyInjection;
-using Microsoft.Extensions.Logging;
 using Microsoft.VisualStudio.TestTools.UnitTesting;
-using MongoDB.Bson;
 using MongoDB.Driver;
-using Moq;
-using Horde.Build.Issues.Handlers;
-using System.Threading;
 using EpicGames.Horde.Storage;
-using Horde.Build.Storage;
-using Horde.Build.Storage.Backends;
-using Microsoft.Extensions.Caching.Memory;
 using System.Buffers;
+using EpicGames.Horde.Storage.Backends;
 
 namespace Horde.Build.Tests
 {
 	[TestClass]
 	public class BlobStoreTests : TestSetup
 	{
-		IStorageClient CreateBlobStore()
+		static IStorageClient CreateBlobStore()
 		{
-			return new BasicBlobStore(MongoService, new TransientStorageBackend(), ServiceProvider.GetRequiredService<IMemoryCache>(), ServiceProvider.GetRequiredService<ILogger<BasicBlobStore>>());
+			return new InMemoryBlobStore();
 		}
 
 		static byte[] CreateTestData(int length, int seed)
diff --git a/Engine/Source/Programs/Horde/Horde.Build.Tests/TestSetup.cs b/Engine/Source/Programs/Horde/Horde.Build.Tests/TestSetup.cs
--- a/Engine/Source/Programs/Horde/Horde.Build.Tests/TestSetup.cs
+++ b/Engine/Source/Programs/Horde/Horde.Build.Tests/TestSetup.cs
@@ -242,12 +242,6 @@ namespace Horde.Build.Tests
 			services.AddSingleton<IStorageBackend<ArtifactCollection>>(sp => new TransientStorageBackend().ForType<ArtifactCollection>());
 			services.AddSingleton<IStorageBackend<BasicStorageClient>>(sp => new TransientStorageBackend().ForType<BasicStorageClient>());
 
-			services.AddSingleton<IStorageClient<CommitService>>(sp =>
-			{
-				IStorageClient blobStore = new BasicBlobStore(sp.GetRequiredService<MongoService>(), new TransientStorageBackend(), sp.GetRequiredService<IMemoryCache>(), sp.GetRequiredService<ILogger<BasicBlobStore>>());
-				return blobStore.ForType<CommitService>();
-			});
-
 			services.AddSingleton<ILegacyStorageClient, BasicStorageClient>();
 
 			services.AddSingleton<ISingletonDocument<AgentSoftwareChannels>>(new SingletonDocumentStub<AgentSoftwareChannels>());
diff --git a/Engine/Source/Programs/Horde/Horde.Build/Commands/Bundles/PerforceCommand.cs b/Engine/Source/Programs/Horde/Horde.Build/Commands/Bundles/PerforceCommand.cs
--- a/Engine/Source/Programs/Horde/Horde.Build/Commands/Bundles/PerforceCommand.cs
+++ b/Engine/Source/Programs/Horde/Horde.Build/Commands/Bundles/PerforceCommand.cs
@@ -9,6 +9,7 @@ using EpicGames.Core;
 using EpicGames.Horde.Storage;
 using EpicGames.Horde.Storage.Nodes;
 using Horde.Build.Perforce;
+using Horde.Build.Storage;
 using Horde.Build.Streams;
 using Horde.Build.Utilities;
 using Microsoft.Extensions.Configuration;
@@ -65,7 +66,9 @@ namespace Horde.Build.Commands.Bundles
 			ICommitService commitService = serviceProvider.GetRequiredService<ICommitService>();
 			ReplicationService replicationService = serviceProvider.GetRequiredService<ReplicationService>();
 			IStreamCollection streamCollection = serviceProvider.GetRequiredService<IStreamCollection>();
-			ILegacyStorageClient storageClient = serviceProvider.GetRequiredService<ILegacyStorageClient>();
+			StorageService storageService = serviceProvider.GetRequiredService<StorageService>();
+
+			IStorageClient storage = await storageService.GetClientAsync(Namespace.Perforce);
 
 			IStream? stream = await streamCollection.GetAsync(new StreamId(StreamId));
 			if (stream == null)
@@ -83,7 +86,7 @@ namespace Horde.Build.Commands.Bundles
 			}
 			else
 			{
-				baseContents = await replicationService.ReadCommitTreeAsync(stream, BaseChange, Filter, RevisionsOnly, CancellationToken.None);
+				baseContents = await replicationService.ReadCommitTreeAsync(storage, stream, BaseChange, Filter, RevisionsOnly, CancellationToken.None);
 			}
 
 			await foreach (ICommit commit in commitService.GetCollection(stream).FindAsync(Change, null, Count))
@@ -95,7 +98,7 @@ namespace Horde.Build.Commands.Bundles
 
 				if (Content)
 				{
-					baseContents = await replicationService.WriteCommitTreeAsync(stream, commit.Number, BaseChange, baseContents, Filter, RevisionsOnly, CancellationToken.None);
+					baseContents = await replicationService.WriteCommitTreeAsync(storage, stream, commit.Number, BaseChange, baseContents, Filter, RevisionsOnly, CancellationToken.None);
 					BaseChange = commit.Number;
 				}
 			}
diff --git a/Engine/Source/Programs/Horde/Horde.Build/Perforce/ReplicationService.cs b/Engine/Source/Programs/Horde/Horde.Build/Perforce/ReplicationService.cs
--- a/Engine/Source/Programs/Horde/Horde.Build/Perforce/ReplicationService.cs
+++ b/Engine/Source/Programs/Horde/Horde.Build/Perforce/ReplicationService.cs
@@ -16,6 +16,7 @@ using EpicGames.Horde.Storage.Nodes;
 using EpicGames.Perforce;
 using EpicGames.Redis;
 using Horde.Build.Server;
+using Horde.Build.Storage;
 using Horde.Build.Streams;
 using Horde.Build.Utilities;
 using Microsoft.Extensions.Hosting;
@@ -163,7 +164,7 @@ namespace Horde.Build.Perforce
 
 		readonly IStreamCollection _streamCollection;
 		readonly IPerforceService _perforceService;
-		readonly IStorageClient _treeStore;
+		readonly StorageService _storageService;
 		readonly ILogger _logger;
 
 		const int MaxBackgroundTasks = 2;
@@ -183,7 +184,7 @@ namespace Horde.Build.Perforce
 		/// <summary>
 		/// Constructor
 		/// </summary>
-		public ReplicationService(RedisService redisService, /* CommitService commitService, */IStreamCollection streamCollection, IPerforceService perforceService, IStorageClient<ReplicationService> treeStore, IOptions<ReplicationServiceOptions> options, ILogger<ReplicationService> logger)
+		public ReplicationService(RedisService redisService, /* CommitService commitService, */IStreamCollection streamCollection, IPerforceService perforceService, StorageService storageService, IOptions<ReplicationServiceOptions> options, ILogger<ReplicationService> logger)
 		{
 			Options = options.Value;
 
@@ -196,7 +197,7 @@ namespace Horde.Build.Perforce
 
 			_streamCollection = streamCollection;
 			_perforceService = perforceService;
-			_treeStore = treeStore;
+			_storageService = storageService;
 			_logger = logger;
 		}
 
@@ -363,8 +364,10 @@ namespace Horde.Build.Perforce
 				}
 				else
 				{
+					IStorageClient storage = await _storageService.GetClientAsync(Namespace.Perforce);
+
 					// Update the stream, updating the reservation every 30 seconds
-					Task internalTask = Task.Run(() => UpdateStreamContentInternalAsync(stream, streamChanges, cancellationToken), cancellationToken);
+					Task internalTask = Task.Run(() => UpdateStreamContentInternalAsync(storage, stream, streamChanges, cancellationToken), cancellationToken);
 					while (!internalTask.IsCompleted)
 					{
 						Task delayTask = Task.Delay(TimeSpan.FromSeconds(ExtendReservationAfterSeconds), cancellationToken);
@@ -399,7 +402,7 @@ namespace Horde.Build.Perforce
 			await _redisReservations.RemoveAsync(streamId);
 		}
 
-		async Task UpdateStreamContentInternalAsync(IStream stream, RedisList<int> changes, CancellationToken cancellationToken)
+		async Task UpdateStreamContentInternalAsync(IStorageClient storage, IStream stream, RedisList<int> changes, CancellationToken cancellationToken)
 		{
 			int prevChange = 0;
 			ReplicationNode prevContents = new ReplicationNode();
@@ -417,11 +420,11 @@ namespace Horde.Build.Perforce
 				{
 					RefName prevRefName = GetRefName(stream, values[0]);
 
-					ReplicationNode? contents = await _treeStore.TryReadTreeAsync<ReplicationNode>(prevRefName, cancellationToken: cancellationToken);
+					ReplicationNode? contents = await storage.TryReadTreeAsync<ReplicationNode>(prevRefName, cancellationToken: cancellationToken);
 					if (contents == null)
 					{
 						_logger.LogInformation("No content for CL {Change}; creating full snapshot", values[0]);
-						prevContents = await WriteCommitTreeAsync(stream, values[0], 0, new ReplicationNode(), cancellationToken);
+						prevContents = await WriteCommitTreeAsync(storage, stream, values[0], 0, new ReplicationNode(), cancellationToken);
 					}
 					else
 					{
@@ -434,7 +437,7 @@ namespace Horde.Build.Perforce
 				else if (values.Length == 2)
 				{
 					// Perform a snapshot of the new change, then remove it from the list
-					prevContents = await WriteCommitTreeAsync(stream, values[1], prevChange, prevContents, cancellationToken);
+					prevContents = await WriteCommitTreeAsync(storage, stream, values[1], prevChange, prevContents, cancellationToken);
 					prevChange = values[1];
 
 					// Remove the first item from the list
@@ -488,31 +491,33 @@ namespace Horde.Build.Perforce
 		/// <summary>
 		/// Reads a tree from the given stream
 		/// </summary>
+		/// <param name="storage">Storage client instance</param>
 		/// <param name="stream"></param>
 		/// <param name="change"></param>
 		/// <param name="filter"></param>
 		/// <param name="revisionsOnly"></param>
 		/// <param name="cancellationToken"></param>
 		/// <returns></returns>
-		public async Task<ReplicationNode> ReadCommitTreeAsync(IStream stream, int change, string? filter, bool revisionsOnly, CancellationToken cancellationToken)
+		public async Task<ReplicationNode> ReadCommitTreeAsync(IStorageClient storage, IStream stream, int change, string? filter, bool revisionsOnly, CancellationToken cancellationToken)
 		{
 			RefName name = GetRefName(stream.Id, change, filter, revisionsOnly);
-			return await _treeStore.ReadTreeAsync<ReplicationNode>(name, cancellationToken: cancellationToken);
+			return await storage.ReadTreeAsync<ReplicationNode>(name, cancellationToken: cancellationToken);
 		}
 
 		/// <summary>
 		/// Replicates the contents of a stream to Horde storage, optionally using the given change as a starting point
 		/// </summary>
+		/// <param name="storage">Storage client instance</param>
 		/// <param name="stream">The stream to replicate</param>
 		/// <param name="change">Commit to store the tree ref</param>
 		/// <param name="baseChange">The base change to update from</param>
 		/// <param name="baseContents">Initial contents of the tree at baseChange</param>
 		/// <param name="cancellationToken">Cancellation token for the operation</param>
 		/// <returns>Root tree object</returns>
-		public async Task<ReplicationNode> WriteCommitTreeAsync(IStream stream, int change, int baseChange, ReplicationNode baseContents, CancellationToken cancellationToken)
+		public async Task<ReplicationNode> WriteCommitTreeAsync(IStorageClient storage, IStream stream, int change, int baseChange, ReplicationNode baseContents, CancellationToken cancellationToken)
 		{
 			bool revisionsOnly = stream.Config.ReplicationMode == ContentReplicationMode.RevisionsOnly;
-			return await WriteCommitTreeAsync(stream, change, baseChange, baseContents, stream.Config.ReplicationFilter, revisionsOnly, cancellationToken);
+			return await WriteCommitTreeAsync(storage, stream, change, baseChange, baseContents, stream.Config.ReplicationFilter, revisionsOnly, cancellationToken);
 		}
 
 		[DebuggerDisplay("{_path}")]
@@ -573,6 +578,7 @@ namespace Horde.Build.Perforce
 		/// <summary>
 		/// Replicates the contents of a stream to Horde storage, optionally using the given change as a starting point
 		/// </summary>
+		/// <param name="storage">Storage client</param>
 		/// <param name="stream">The stream to replicate</param>
 		/// <param name="change">Commit to store the tree ref</param>
 		/// <param name="baseChange">The base change to update from</param>
@@ -581,7 +587,7 @@ namespace Horde.Build.Perforce
 		/// <param name="revisionsOnly"></param>
 		/// <param name="cancellationToken">Cancellation token for the operation</param>
 		/// <returns>Root tree object</returns>
-		public async Task<ReplicationNode> WriteCommitTreeAsync(IStream stream, int change, int? baseChange, ReplicationNode baseContents, string? filter, bool revisionsOnly, CancellationToken cancellationToken)
+		public async Task<ReplicationNode> WriteCommitTreeAsync(IStorageClient storage, IStream stream, int change, int? baseChange, ReplicationNode baseContents, string? filter, bool revisionsOnly, CancellationToken cancellationToken)
 		{
 			// Create a client to replicate from this stream
 			ReplicationClient clientInfo = await FindOrAddReplicationClientAsync(stream);
@@ -600,11 +606,11 @@ namespace Horde.Build.Perforce
 
 			RefName refName = GetRefName(stream.Id, change, filter, revisionsOnly);
 			RefName incRefName = new RefName($"{refName}-inc");
-			ITreeWriter directoryWriter = _treeStore.CreateTreeWriter(refName.Text);
+			ITreeWriter directoryWriter = storage.CreateTreeWriter(refName.Text);
 			ITreeWriter fileWriter = directoryWriter.CreateChildWriter();
 
 			// Get the current sync state for this change
-			ReplicationNode? syncNode = await _treeStore.TryReadTreeAsync<ReplicationNode>(incRefName, cancellationToken: cancellationToken);
+			ReplicationNode? syncNode = await storage.TryReadTreeAsync<ReplicationNode>(incRefName, cancellationToken: cancellationToken);
 			bool deleteIncRef = syncNode != null;
 			if (syncNode == null)
 			{
@@ -877,7 +883,7 @@ namespace Horde.Build.Perforce
 			// Delete the incremental state
 			if (deleteIncRef)
 			{
-				await _treeStore.DeleteRefAsync(incRefName, cancellationToken);
+				await storage.DeleteRefAsync(incRefName, cancellationToken);
 			}
 			return syncNode;
 		}
diff --git a/Engine/Source/Programs/Horde/Horde.Build/Server/GlobalConfig.cs b/Engine/Source/Programs/Horde/Horde.Build/Server/GlobalConfig.cs
--- a/Engine/Source/Programs/Horde/Horde.Build/Server/GlobalConfig.cs
+++ b/Engine/Source/Programs/Horde/Horde.Build/Server/GlobalConfig.cs
@@ -10,6 +10,7 @@ using EpicGames.Perforce;
 using EpicGames.Serialization;
 using Horde.Build.Acls;
 using Horde.Build.Projects;
+using Horde.Build.Storage;
 using Horde.Build.Tools;
 using Horde.Build.Utilities;
 
@@ -65,9 +66,9 @@ namespace Horde.Build.Server
 		public int MaxConformCount { get; set; }
 
 		/// <summary>
-		/// List of storage namespaces
+		/// Storage configuration
 		/// </summary>
-		public StorageConfig? Storage { get; set; }
+		public StorageConfig Storage { get; set; } = new StorageConfig();
 
 		/// <summary>
 		/// Access control list
@@ -169,38 +170,6 @@ namespace Horde.Build.Server
 		public string Path { get; set; } = null!;
 	}
 
-	/// <summary>
-	/// Configuration for storage system
-	/// </summary>
-	public class StorageConfig
-	{
-		/// <summary>
-		/// List of storage namespaces
-		/// </summary>
-		public List<NamespaceConfig> Namespaces { get; set; } = new List<NamespaceConfig>();
-	}
-
-	/// <summary>
-	/// Configuration for a storage namespace
-	/// </summary>
-	public class NamespaceConfig
-	{
-		/// <summary>
-		/// Identifier for this namespace
-		/// </summary>
-		public string Id { get; set; } = String.Empty;
-
-		/// <summary>
-		/// Buckets within this namespace
-		/// </summary>
-		public List<BucketConfig> Buckets { get; set; } = new List<BucketConfig>();
-
-		/// <summary>
-		/// Access control for this namespace
-		/// </summary>
-		public AclConfig? Acl { get; set; }
-	}
-
 	/// <summary>
 	/// Configuration for a device platform 
 	/// </summary>
@@ -233,18 +202,6 @@ namespace Horde.Build.Server
 		public List<DevicePlatformConfig> Platforms { get; set; } = new List<DevicePlatformConfig>();
 	}
 
-
-	/// <summary>
-	/// Configuration for a bucket
-	/// </summary>
-	public class BucketConfig
-	{
-		/// <summary>
-		/// Identifier for the bucket
-		/// </summary>
-		public string Id { get; set; } = String.Empty;
-	}
-
 	/// <summary>
 	/// Describes the monetary cost of agents matching a particular criteris
 	/// </summary>
diff --git a/Engine/Source/Programs/Horde/Horde.Build/Startup.cs b/Engine/Source/Programs/Horde/Horde.Build/Startup.cs
--- a/Engine/Source/Programs/Horde/Horde.Build/Startup.cs
+++ b/Engine/Source/Programs/Horde/Horde.Build/Startup.cs
@@ -277,12 +277,6 @@ namespace Horde.Build
 			}
 		}
 
-		static IStorageClient CreateStorageClient(IServiceProvider sp, BlobStoreOptions options)
-		{
-			IStorageBackend backend = CreateStorageBackend(sp, options);
-			return new BasicBlobStore(sp.GetRequiredService<MongoService>(), backend, sp.GetRequiredService<IMemoryCache>(), sp.GetRequiredService<ILogger<BasicBlobStore>>());
-		}
-
 		public Startup(IConfiguration configuration)
 		{
 			Configuration = configuration;
@@ -358,6 +352,7 @@ namespace Horde.Build
 			services.AddSingleton<INotificationTriggerCollection, NotificationTriggerCollection>();
 			services.AddSingleton<IPoolCollection, PoolCollection>();
 			services.AddSingleton<IProjectCollection, ProjectCollection>();
+			services.AddSingleton<IRefCollection, RefCollection>();
 			services.AddSingleton<ISessionCollection, SessionCollection>();
 			services.AddSingleton<IServiceAccountCollection, ServiceAccountCollection>();
 			services.AddSingleton<ISubscriptionCollection, SubscriptionCollection>();
@@ -447,6 +442,7 @@ namespace Horde.Build
 			services.AddSingleton<StreamService>();
 			services.AddSingleton<DeviceService>();						
 			services.AddSingleton<NoticeService>();
+			services.AddSingleton<StorageService>();
 
 			if (settings.JiraUrl != null)
 			{
@@ -460,7 +456,6 @@ namespace Horde.Build
 			// Storage providers
 			services.AddSingleton(sp => CreateStorageBackend(sp, settings.LogStorage).ForType<PersistentLogStorage>());
 			services.AddSingleton(sp => CreateStorageBackend(sp, settings.ArtifactStorage).ForType<ArtifactCollection>());
-			services.AddSingleton(sp => CreateStorageClient(sp, settings.CommitStorage).ForType<ReplicationService>());
 
 			services.AddHordeStorage(settings => configSection.GetSection("Storage").Bind(settings));
 			
diff --git a/Engine/Source/Programs/Horde/Horde.Build/Storage/BasicBlobStore.cs b/Engine/Source/Programs/Horde/Horde.Build/Storage/BasicBlobStore.cs
deleted file mode 100644
diff --git a/Engine/Source/Programs/Horde/Horde.Build/Storage/IRefCollection.cs b/Engine/Source/Programs/Horde/Horde.Build/Storage/IRefCollection.cs
new file mode 100644
--- /dev/null
+++ b/Engine/Source/Programs/Horde/Horde.Build/Storage/IRefCollection.cs
@@ -0,0 +1,48 @@
+// Copyright Epic Games, Inc. All Rights Reserved.
+
+using System;
+using System.Threading;
+using System.Threading.Tasks;
+using Amazon.Runtime.Internal.Util;
+using EpicGames.Horde.Storage;
+using Horde.Build.Server;
+using Microsoft.Extensions.Caching.Memory;
+using MongoDB.Bson;
+using MongoDB.Bson.Serialization.Attributes;
+using MongoDB.Driver;
+
+namespace Horde.Build.Storage
+{
+	/// <summary>
+	/// Collection of ref name to blob mappings
+	/// </summary>
+	public interface IRefCollection
+	{
+		/// <summary>
+		/// Read a ref from storage
+		/// </summary>
+		/// <param name="namespaceId">Namespace for the ref</param>
+		/// <param name="name">The ref name</param>
+		/// <param name="cacheTime">Minimum age for a cached value to be returned</param>
+		/// <param name="cancellationToken">Cancellation token for the operation</param>
+		/// <returns>The ref value</returns>
+		Task<BlobLocator> TryReadRefTargetAsync(NamespaceId namespaceId, RefName name, DateTime cacheTime = default, CancellationToken cancellationToken = default);
+
+		/// <summary>
+		/// Writes a ref to storage
+		/// </summary>
+		/// <param name="namespaceId">Namespace for the ref</param>
+		/// <param name="name">The ref name</param>
+		/// <param name="locator">Value for the ref</param>
+		/// <param name="cancellationToken">Cancellation token for the operation</param>
+		Task WriteRefTargetAsync(NamespaceId namespaceId, RefName name, BlobLocator locator, CancellationToken cancellationToken = default);
+
+		/// <summary>
+		/// Delete a ref from storage
+		/// </summary>
+		/// <param name="namespaceId">Namespace for the ref</param>
+		/// <param name="name">The ref name</param>
+		/// <param name="cancellationToken">Cancellation token for the operation</param>
+		Task DeleteRefAsync(NamespaceId namespaceId, RefName name, CancellationToken cancellationToken = default);
+	}
+}
diff --git a/Engine/Source/Programs/Horde/Horde.Build/Storage/RefCollection.cs b/Engine/Source/Programs/Horde/Horde.Build/Storage/RefCollection.cs
new file mode 100644
--- /dev/null
+++ b/Engine/Source/Programs/Horde/Horde.Build/Storage/RefCollection.cs
@@ -0,0 +1,133 @@
+// Copyright Epic Games, Inc. All Rights Reserved.
+
+using System;
+using System.Threading;
+using System.Threading.Tasks;
+using EpicGames.Horde.Storage;
+using Horde.Build.Server;
+using Microsoft.Extensions.Caching.Memory;
+using MongoDB.Bson;
+using MongoDB.Bson.Serialization.Attributes;
+using MongoDB.Driver;
+
+namespace Horde.Build.Storage
+{
+	/// <summary>
+	/// Collection of ref name to blob mappings
+	/// </summary>
+	public class RefCollection : IRefCollection
+	{
+		class RefKey
+		{
+			public NamespaceId NamespaceId { get; set; }
+			public RefName Name { get; set; }
+
+			public RefKey(NamespaceId namespaceId, RefName name)
+			{
+				NamespaceId = namespaceId;
+				Name = name;
+			}
+
+			/// <inheritdoc/>
+			public override bool Equals(object? obj) => obj is RefKey refKey && NamespaceId == refKey.NamespaceId && Name == refKey.Name;
+
+			/// <inheritdoc/>
+			public override int GetHashCode() => HashCode.Combine(NamespaceId, Name);
+		}
+
+		class RefDocument
+		{
+			public ObjectId Id { get; set; }
+
+			[BsonElement("ns")]
+			public NamespaceId NamespaceId { get; set; }
+
+			[BsonElement("k")]
+			public RefName Name { get; set; }
+
+			[BsonElement("v")]
+			public BlobLocator Value { get; set; }
+
+			public RefDocument()
+			{
+				Name = RefName.Empty;
+				Value = BlobLocator.Empty;
+			}
+
+			public RefDocument(NamespaceId namespaceId, RefName name, BlobLocator value)
+			{
+				NamespaceId = namespaceId;
+				Name = name;
+				Value = value;
+			}
+		}
+
+		class CachedRefValue
+		{
+			public RefName Name { get; }
+			public BlobLocator Value { get; }
+			public DateTime Time { get; }
+
+			public CachedRefValue(RefName name, BlobLocator value)
+			{
+				Name = name;
+				Value = value;
+				Time = DateTime.UtcNow;
+			}
+		}
+
+		readonly IMongoCollection<RefDocument> _refs;
+		readonly IMemoryCache _cache;
+
+		/// <summary>
+		/// Constructor
+		/// </summary>
+		/// <param name="mongoService">The mongo service implementation</param>
+		/// <param name="cache">Cache for ref values</param>
+		public RefCollection(MongoService mongoService, IMemoryCache cache)
+		{
+			_refs = mongoService.GetCollection<RefDocument>("Refs", keys => keys.Ascending(x => x.NamespaceId).Ascending(x => x.Name), unique: true);
+			_cache = cache;
+		}
+
+		CachedRefValue AddRefToCache(NamespaceId namespaceId, RefName name, BlobLocator value)
+		{
+			CachedRefValue cacheValue = new CachedRefValue(name, value);
+			using (ICacheEntry newEntry = _cache.CreateEntry(new RefKey(namespaceId, name)))
+			{
+				newEntry.Value = cacheValue;
+				newEntry.SetSize(name.Text.Length + value.Inner.Length);
+				newEntry.SetSlidingExpiration(TimeSpan.FromMinutes(5.0));
+			}
+			return cacheValue;
+		}
+
+		/// <inheritdoc/>
+		public async Task DeleteRefAsync(NamespaceId namespaceId, RefName name, CancellationToken cancellationToken = default)
+		{
+			await _refs.DeleteOneAsync(x => x.NamespaceId == namespaceId && x.Name == name, cancellationToken);
+			AddRefToCache(namespaceId, name, BlobLocator.Empty);
+		}
+
+		/// <inheritdoc/>
+		public async Task<BlobLocator> TryReadRefTargetAsync(NamespaceId namespaceId, RefName name, DateTime cacheTime = default, CancellationToken cancellationToken = default)
+		{
+			CachedRefValue entry;
+			if (!_cache.TryGetValue(name, out entry) || entry.Time < cacheTime)
+			{
+				RefDocument? refDocument = await _refs.Find(x => x.NamespaceId == namespaceId && x.Name == name).FirstOrDefaultAsync(cancellationToken);
+				BlobLocator value = (refDocument == null) ? BlobLocator.Empty : refDocument.Value;
+				entry = AddRefToCache(namespaceId, name, value);
+			}
+			return entry.Value;
+		}
+
+		/// <inheritdoc/>
+		public async Task WriteRefTargetAsync(NamespaceId namespaceId, RefName name, BlobLocator locator, CancellationToken cancellationToken = default)
+		{
+			RefDocument refDocument = new RefDocument(namespaceId, name, locator);
+			await _refs.ReplaceOneAsync(x => x.NamespaceId == namespaceId && x.Name == name, refDocument, new ReplaceOptions { IsUpsert = true }, cancellationToken);
+			AddRefToCache(namespaceId, name, locator);
+		}
+	}
+}
diff --git a/Engine/Source/Programs/Horde/Horde.Build/Storage/StorageConfig.cs b/Engine/Source/Programs/Horde/Horde.Build/Storage/StorageConfig.cs
new file mode 100644
--- /dev/null
+++ b/Engine/Source/Programs/Horde/Horde.Build/Storage/StorageConfig.cs
@@ -0,0 +1,141 @@
+// Copyright Epic Games, Inc. All Rights Reserved.
+
+using System.Collections.Generic;
+using System.ComponentModel.DataAnnotations;
+using EpicGames.Horde.Storage;
+using Horde.Build.Acls;
+using Horde.Build.Storage.Backends;
+using Horde.Build.Utilities;
+
+namespace Horde.Build.Storage
+{
+	using BackendId = StringId<IStorageBackend>;
+
+	/// <summary>
+	/// Well-known namespace identifiers
+	/// </summary>
+	public static class Namespace
+	{
+		/// <summary>
+		/// Storage of artifacts 
+		/// </summary>
+		public static NamespaceId Artifacts { get; } = new NamespaceId("horde-artifacts");
+
+		/// <summary>
+		/// Replicated Perforce data
+		/// </summary>
+		public static NamespaceId Perforce { get; } = new NamespaceId("horde-perforce");
+	}
+
+	/// <summary>
+	/// Configuration for storage
+	/// </summary>
+	public class StorageConfig
+	{
+		/// <summary>
+		/// List of storage backends
+		/// </summary>
+		public List<BackendConfig> Backends { get; set; } = new List<BackendConfig>();
+
+		/// <summary>
+		/// List of namespaces for storage
+		/// </summary>
+		public List<NamespaceConfig> Namespaces { get; set; } = new List<NamespaceConfig>();
+	}
+
+	/// <summary>
+	/// Types of storage backend to use
+	/// </summary>
+	public enum BackendType
+	{
+		/// <summary>
+		/// Local filesystem
+		/// </summary>
+		FileSystem,
+
+		/// <summary>
+		/// AWS S3
+		/// </summary>
+		Aws,
+
+		/// <summary>
+		/// In-memory only (for testing)
+		/// </summary>
+		Transient,
+
+		/// <summary>
+		/// Relay to another server (useful for testing against prod)
+		/// </summary>
+		Relay,
+	};
+
+	/// <summary>
+	/// Common settings object for different providers
+	/// </summary>
+	public class BackendConfig : IFileSystemStorageOptions, IAwsStorageOptions, IRelayStorageOptions
+	{
+		/// <summary>
+		/// The storage backend ID
+		/// </summary>
+		[Required]
+		public BackendId Id { get; set; }
+
+		/// <summary>
+		/// Base backend to copy default settings from
+		/// </summary>
+		public BackendId Base { get; set; }
+
+		/// <inheritdoc/>
+		public BackendType? Type { get; set; }
+
+		/// <inheritdoc/>
+		public string? BaseDir { get; set; }
+
+		/// <inheritdoc/>
+		public string? AwsBucketName { get; set; }
+
+		/// <inheritdoc/>
+		public string? AwsBucketPath { get; set; }
+
+		/// <inheritdoc/>
+		public AwsCredentialsType AwsCredentials { get; set; }
+
+		/// <inheritdoc/>
+		public string? AwsRole { get; set; }
+
+		/// <inheritdoc/>
+		public string? AwsProfile { get; set; }
+
+		/// <inheritdoc/>
+		public string? AwsRegion { get; set; }
+
+		/// <inheritdoc/>
+		public string? RelayServer { get; set; }
+
+		/// <inheritdoc/>
+		public string? RelayToken { get; set; }
+	}
+
+	/// <summary>
+	/// Configuration of a particular namespace
+	/// </summary>
+	public class NamespaceConfig
+	{
+		/// <summary>
+		/// Identifier for this namespace
+		/// </summary>
+		[Required]
+		public NamespaceId Id { get; set; }
+
+		/// <summary>
+		/// Backend to use for this namespace
+		/// </summary>
+		[Required]
+		public BackendId Backend { get; set; }
+
+		/// <summary>
+		/// Access list for this namespace
+		/// </summary>
+		public AclConfig Acl { get; set; } = new AclConfig();
+	}
+}
diff --git a/Engine/Source/Programs/Horde/Horde.Build/Storage/StorageService.cs b/Engine/Source/Programs/Horde/Horde.Build/Storage/StorageService.cs
new file mode 100644
--- /dev/null
+++ b/Engine/Source/Programs/Horde/Horde.Build/Storage/StorageService.cs
@@ -0,0 +1,390 @@
+// Copyright Epic Games, Inc. All Rights Reserved.
+
+using System;
+using System.Buffers;
+using System.Collections;
+using System.Collections.Generic;
+using System.Reflection;
+using System.Security.Claims;
+using System.Text.Json;
+using System.Threading;
+using System.Threading.Tasks;
+using EpicGames.Core;
+using EpicGames.Horde.Storage;
+using Horde.Build.Acls;
+using Horde.Build.Server;
+using Horde.Build.Storage.Backends;
+using Horde.Build.Utilities;
+using Microsoft.Extensions.Caching.Memory;
+using Microsoft.Extensions.Configuration;
+using Microsoft.Extensions.DependencyInjection;
+using Microsoft.Extensions.Logging;
+
+namespace Horde.Build.Storage
+{
+	using BackendId = StringId<IStorageBackend>;
+
+	/// <summary>
+	/// Exception thrown by the <see cref="StorageService"/>
+	/// </summary>
+	public sealed class StorageException : Exception
+	{
+		/// <summary>
+		/// Constructor
+		/// </summary>
+		public StorageException(string message, Exception? inner = null)
+			: base(message, inner)
+		{
+		}
+	}
+
+	/// <summary>
+	/// Functionality related to the storage service
+	/// </summary>
+	public sealed class StorageService : IDisposable
+	{
+		sealed class StorageClient : StorageClientBase, IDisposable
+		{
+			public NamespaceConfig Config { get; }
+
+			NamespaceId NamespaceId => Config.Id;
+
+			readonly IStorageBackend _backend;
+			readonly IRefCollection _refs;
+
+			public StorageClient(NamespaceConfig config, IStorageBackend backend, IRefCollection refs, IMemoryCache? cache)
+				: base(cache)
+			{
+				Config = config;
+				_backend = backend;
+				_refs = refs;
+			}
+
+			public void Dispose()
+			{
+				if (_backend is IDisposable disposable)
+				{
+					disposable.Dispose();
+				}
+			}
+
+			static string GetBlobPath(BlobId blobId) => $"{blobId}.blob";
+
+			#region Blobs
+
+			public override async Task<Bundle> ReadBundleAsync(BlobLocator locator, CancellationToken cancellationToken = default)
+			{
+				string path = GetBlobPath(locator.BlobId);
+
+				ReadOnlyMemory<byte>? data = await _backend.ReadBytesAsync(path, cancellationToken);
+				if (data == null)
+				{
+					throw new StorageException($"Unable to read data from {path}");
+				}
+
+				return new Bundle(new MemoryReader(data.Value));
+			}
+
+			/// <inheritdoc/>
+			public override async Task<BlobLocator> WriteBundleAsync(Bundle bundle, Utf8String prefix = default, CancellationToken cancellationToken = default)
+			{
+				BlobId id = BlobId.CreateNew(prefix);
+				string path = GetBlobPath(id);
+
+				ReadOnlySequence<byte> sequence = bundle.AsSequence();
+				using (ReadOnlySequenceStream stream = new ReadOnlySequenceStream(sequence))
+				{
+					await _backend.WriteAsync(path, stream, cancellationToken);
+				}
+
+				return new BlobLocator(HostId.Empty, id);
+			}
+
+			#endregion
+
+			#region Refs
+
+			/// <inheritdoc/>
+			public override Task<BlobLocator> TryReadRefTargetAsync(RefName name, DateTime cacheTime = default, CancellationToken cancellationToken = default) => _refs.TryReadRefTargetAsync(NamespaceId, name, cacheTime, cancellationToken);
+
+			/// <inheritdoc/>
+			public override Task WriteRefTargetAsync(RefName name, BlobLocator locator, CancellationToken cancellationToken = default) => _refs.WriteRefTargetAsync(NamespaceId, name, locator, cancellationToken);
+
+			/// <inheritdoc/>
+			public override Task DeleteRefAsync(RefName name, CancellationToken cancellationToken = default) => _refs.DeleteRefAsync(NamespaceId, name, cancellationToken);
+
+			#endregion
+		}
+
+		class NamespaceInfo : IDisposable
+		{
+			public NamespaceConfig Config { get; }
+			public IStorageClient Client { get; }
+			public Acl? Acl { get; }
+
+			public NamespaceInfo(NamespaceConfig config, IStorageClient client)
+			{
+				Config = config;
+				Client = client;
+				Acl = Acl.Merge(null, config.Acl);
+			}
+
+			public void Dispose() => (Client as IDisposable)?.Dispose();
+		}
+
+		class State : IDisposable
+		{
+			public StorageConfig Config { get; }
+			public Dictionary<NamespaceId, NamespaceInfo> Namespaces { get; } = new Dictionary<NamespaceId, NamespaceInfo>();
+
+			public State(StorageConfig config)
+			{
+				Config = config;
+			}
+
+			public void Dispose()
+			{
+				foreach (NamespaceInfo namespaceInfo in Namespaces.Values)
+				{
+					namespaceInfo.Dispose();
+				}
+			}
+		}
+
+		readonly GlobalsService _globalsService;
+		readonly AclService _aclService;
+		readonly IRefCollection _refCollection;
+		readonly IMemoryCache _cache;
+		readonly IServiceProvider _serviceProvider;
+
+		State? _lastState;
+		string? _lastConfigRevision;
+
+		readonly AsyncCachedValue<State> _cachedState;
+
+		static readonly BackendConfig s_defaultBackendConfig = new BackendConfig();
+
+		/// <summary>
+		/// Constructor
+		/// </summary>
+		public StorageService(GlobalsService globalsService, AclService aclService, IRefCollection refCollection, IMemoryCache cache, IServiceProvider serviceProvider)
+		{
+			_globalsService = globalsService;
+			_aclService = aclService;
+			_refCollection = refCollection;
+			_cache = cache;
+			_serviceProvider = serviceProvider;
+			_cachedState = new AsyncCachedValue<State>(() => GetNextState(), TimeSpan.FromMinutes(1.0)); 
+		}
+
+		/// <inheritdoc/>
+		public void Dispose()
+		{
+			if (_lastState != null)
+			{
+				_lastState.Dispose();
+				_lastState = null;
+			}
+		}
+
+		async ValueTask<NamespaceInfo> GetNamespaceInfo(NamespaceId namespaceId)
+		{
+			State state = await _cachedState.GetAsync();
+			if (!state.Namespaces.TryGetValue(namespaceId, out NamespaceInfo? namespaceInfo))
+			{
+				throw new StorageException($"No namespace '{namespaceId}' is configured.");
+			}
+			return namespaceInfo;
+		}
+
+		/// <summary>
+		/// Gets a storage client for the given namespace
+		/// </summary>
+		/// <param name="namespaceId">Namespace identifier</param>
+		/// <returns></returns>
+		public async ValueTask<IStorageClient> GetClientAsync(NamespaceId namespaceId)
+		{
+			NamespaceInfo namespaceInfo = await GetNamespaceInfo(namespaceId);
+			return namespaceInfo.Client;
+		}
+
+		/// <summary>
+		/// Authorize operations for the given store
+		/// </summary>
+		/// <param name="namespaceId"></param>
+		/// <param name="user"></param>
+		/// <param name="action"></param>
+		/// <param name="cache">Global permissions cache</param>
+		/// <returns></returns>
+		public async Task<bool> AuthorizeAsync(NamespaceId namespaceId, ClaimsPrincipal user, AclAction action, GlobalPermissionsCache? cache)
+		{
+			NamespaceInfo namespaceInfo = await GetNamespaceInfo(namespaceId);
+
+			bool? result = namespaceInfo.Acl?.Authorize(action, user);
+			if (result == null)
+			{
+				return await _aclService.AuthorizeAsync(action, user, cache);
+			}
+			else
+			{
+				return result.Value;
+			}
+		}
+
+		async Task<State> GetNextState()
+		{
+			IGlobals globals = await _globalsService.GetAsync();
+			if (_lastState == null || !String.Equals(_lastConfigRevision, globals.ConfigRevision, StringComparison.Ordinal))
+			{
+				StorageConfig storageConfig = globals.Config.Storage;
+
+				// Create a lookup for backend config objects by their id
+				Dictionary<BackendId, BackendConfig> backendIdToConfig = new Dictionary<BackendId, BackendConfig>();
+				foreach (BackendConfig backendConfig in storageConfig.Backends)
+				{
+					backendIdToConfig.Add(backendConfig.Id, backendConfig);
+				}
+
+				// Create a lookup of namespace id to config, to ensure there are no duplicates
+				Dictionary<NamespaceId, NamespaceConfig> namespaceIdToConfig = new Dictionary<NamespaceId, NamespaceConfig>();
+				foreach (NamespaceConfig namespaceConfig in storageConfig.Namespaces)
+				{
+					namespaceIdToConfig.Add(namespaceConfig.Id, namespaceConfig);
+				}
+
+				// Configure the new clients
+				State nextState = new State(storageConfig);
+				try
+				{
+					Dictionary<BackendId, BackendConfig> mergedIdToConfig = new Dictionary<BackendId, BackendConfig>();
+					foreach (NamespaceConfig namespaceConfig in storageConfig.Namespaces)
+					{
+						BackendConfig backendConfig = GetBackendConfig(namespaceConfig.Backend, backendIdToConfig, mergedIdToConfig);
+						IStorageBackend backend = CreateStorageBackend(backendConfig);
+
+#pragma warning disable CA2000 // Dispose objects before losing scope
+						StorageClient client = new StorageClient(namespaceConfig, backend, _refCollection, _cache);
+						nextState.Namespaces.Add(namespaceConfig.Id, new NamespaceInfo(namespaceConfig, client));
+#pragma warning restore CA2000 // Dispose objects before losing scope
+					}
+				}
+				catch
+				{
+					nextState.Dispose();
+					throw;
+				}
+
+				_lastState?.Dispose();
+				_lastState = nextState;
+
+				_lastConfigRevision = globals.ConfigRevision;
+			}
+			return _lastState;
+		}
+
+		/// <summary>
+		/// Creates a storage backend with the given configuration
+		/// </summary>
+		/// <param name="config">Configuration for the backend</param>
+		/// <returns>New storage backend instance</returns>
+		IStorageBackend CreateStorageBackend(BackendConfig config)
+		{
+			switch (config.Type ?? BackendType.FileSystem)
+			{
+				case BackendType.FileSystem:
+					return new FileSystemStorageBackend(config);
+				case BackendType.Aws:
+					return new AwsStorageBackend(_serviceProvider.GetRequiredService<IConfiguration>(), config, _serviceProvider.GetRequiredService<ILogger<AwsStorageBackend>>());
+				case BackendType.Transient:
+					return new TransientStorageBackend();
+				case BackendType.Relay:
+					return new RelayStorageBackend(config);
+				default:
+					throw new NotImplementedException();
+			}
+		}
+
+		/// <summary>
+		/// Gets the configuration for a particular backend
+		/// </summary>
+		/// <param name="backendId">Identifier for the backend</param>
+		/// <param name="baseIdToConfig">Lookup for all backend configuration data</param>
+		/// <param name="mergedIdToConfig">Lookup for computed hierarchical config objects</param>
+		/// <returns>Merged config for the given backend</returns>
+		BackendConfig GetBackendConfig(BackendId backendId, Dictionary<BackendId, BackendConfig> baseIdToConfig, Dictionary<BackendId, BackendConfig> mergedIdToConfig)
+		{
+			BackendConfig? config;
+			if (mergedIdToConfig.TryGetValue(backendId, out config))
+			{
+				if (config == null)
+				{
+					throw new StorageException($"Configuration for storage backend '{backendId}' is recursive.");
+				}
+			}
+			else
+			{
+				mergedIdToConfig.Add(backendId, null!);
+
+				config = baseIdToConfig[backendId];
+				if (config.Base != BackendId.Empty)
+				{
+					BackendConfig baseConfig = GetBackendConfig(config.Base, baseIdToConfig, mergedIdToConfig);
+					MergeConfigs(baseConfig, config, s_defaultBackendConfig);
+				}
+
+				mergedIdToConfig[backendId] = config;
+			}
+			return config;
+		}
+
+		/// <summary>
+		/// Allows one configuration object to override another
+		/// </summary>
+		static void MergeConfigs<T>(T sourceObject, T targetObject, T defaultObject)
+		{
+			PropertyInfo[] properties = typeof(T).GetProperties(BindingFlags.Public | BindingFlags.Instance);
+			foreach (PropertyInfo property in properties)
+			{
+				object? targetValue = property.GetValue(targetObject);
+				object? defaultValue = property.GetValue(defaultObject);
+
+				if (ValueEquals(targetValue, defaultValue))
+				{
+					object? sourceValue = property.GetValue(sourceObject);
+					property.SetValue(targetObject, sourceValue);
+				}
+			}
+		}
+
+		/// <summary>
+		/// Tests two objects for value equality
+		/// </summary>
+		static bool ValueEquals(object? source, object? target)
+		{
+			if (ReferenceEquals(source, null) || ReferenceEquals(target, null))
+			{
+				return ReferenceEquals(source, target);
+			}
+
+			if (source is IEnumerable _)
+			{
+				IEnumerator sourceEnumerator = ((IEnumerable)source).GetEnumerator();
+				IEnumerator targetEnumerator = ((IEnumerable)target).GetEnumerator();
+
+				for (; ; )
+				{
+					if (!sourceEnumerator.MoveNext())
+					{
+						return !targetEnumerator.MoveNext();
+					}
+					if (!targetEnumerator.MoveNext() || !ValueEquals(sourceEnumerator.Current, targetEnumerator.Current))
+					{
+						return false;
+					}
+				}
+			}
+
+			return source.Equals(target);
+		}
+	}
+}
