Horde: Parallelize auto-scaling operations

diff --git a/Engine/Source/Programs/Horde/Horde.Build.Tests/Fleet/AutoscaleServiceV2Test.cs b/Engine/Source/Programs/Horde/Horde.Build.Tests/Fleet/AutoscaleServiceV2Test.cs
--- a/Engine/Source/Programs/Horde/Horde.Build.Tests/Fleet/AutoscaleServiceV2Test.cs
+++ b/Engine/Source/Programs/Horde/Horde.Build.Tests/Fleet/AutoscaleServiceV2Test.cs
@@ -82,7 +82,7 @@ namespace Horde.Build.Tests.Fleet
 			await CreateAgentAsync(pool, false);
 
 			List<PoolSizeData> poolSizeDatas = await service.GetPoolSizeDataAsync();
-			poolSizeDatas = await service.CalculateDesiredPoolSizesAsync(poolSizeDatas);
+			poolSizeDatas = await service.CalculatePoolSizesAsync(poolSizeDatas, CancellationToken.None);
 			Assert.AreEqual(1, poolSizeDatas.Count);
 			Assert.AreEqual(2, poolSizeDatas[0].Agents.Count);
 		}
diff --git a/Engine/Source/Programs/Horde/Horde.Build/Agents/Fleet/AutoscaleServiceV2.cs b/Engine/Source/Programs/Horde/Horde.Build/Agents/Fleet/AutoscaleServiceV2.cs
--- a/Engine/Source/Programs/Horde/Horde.Build/Agents/Fleet/AutoscaleServiceV2.cs
+++ b/Engine/Source/Programs/Horde/Horde.Build/Agents/Fleet/AutoscaleServiceV2.cs
@@ -1,6 +1,7 @@
 // Copyright Epic Games, Inc. All Rights Reserved.
 
 using System;
+using System.Collections.Concurrent;
 using System.Collections.Generic;
 using System.Diagnostics;
 using System.Linq;
@@ -33,6 +34,11 @@ namespace Horde.Build.Agents.Fleet
 	/// </summary>
 	public sealed class AutoscaleServiceV2 : IHostedService, IDisposable
 	{
+		/// <summary>
+		/// Max number of auto-scaling calculations to be done concurrently (sizing calculations and fleet manager calls)
+		/// </summary>
+		private const int MaxParallelTasks = 10;
+		
 		private readonly IAgentCollection _agentCollection;
 		private readonly IGraphCollection _graphCollection;
 		private readonly IJobCollection _jobCollection;
@@ -118,8 +124,8 @@ namespace Horde.Build.Agents.Fleet
 			using IScope _ = GlobalTracer.Instance.BuildSpan("AutoscaleService.TickAsync").StartActive();
 			
 			List<PoolSizeData> input = await GetPoolSizeDataAsync();
-			List<PoolSizeData> output = await CalculateDesiredPoolSizesAsync(input);
-			await ResizePoolsAsync(output);
+			List<PoolSizeData> output = await CalculatePoolSizesAsync(input, stoppingToken);
+			await ResizePoolsAsync(output, stoppingToken);
 
 			stopwatch.Stop();
 			_logger.LogInformation("Autoscaling pools took {ElapsedTime} ms", stopwatch.ElapsedMilliseconds);
@@ -138,15 +144,20 @@ namespace Horde.Build.Agents.Fleet
 			_logger.LogInformation("Autoscaling pools (high frequency) took {ElapsedTime} ms", stopwatch.ElapsedMilliseconds);
 		}
 
-		internal async Task<List<PoolSizeData>> CalculateDesiredPoolSizesAsync(List<PoolSizeData> poolSizeDatas)
+		internal async Task<List<PoolSizeData>> CalculatePoolSizesAsync(List<PoolSizeData> poolSizeDatas, CancellationToken cancellationToken)
 		{
-			List<PoolSizeData> result = new();
-			foreach (PoolSizeData psd in poolSizeDatas)
+			using IScope _ = GlobalTracer.Instance.BuildSpan("AutoscaleService.CalculatePoolSizesAsync").StartActive();
+			ConcurrentQueue<PoolSizeData> results = new ();
+			ParallelOptions options = new () { MaxDegreeOfParallelism = MaxParallelTasks, CancellationToken = cancellationToken };
+
+			await Parallel.ForEachAsync(poolSizeDatas, options, async (input, innerCt) =>
 			{
-				result.Add((await CreatePoolSizeStrategy(psd.Pool).CalcDesiredPoolSizesAsync(new List<PoolSizeData> { psd }))[0]);
-			}
+				IPoolSizeStrategy sizeStrategy = CreatePoolSizeStrategy(input.Pool);
+				PoolSizeData output = (await sizeStrategy.CalcDesiredPoolSizesAsync(new List<PoolSizeData> { input }))[0];
+				results.Enqueue(output);
+			});
 
-			return result;
+			return results.ToList();
 		}
 
 		internal async Task<List<PoolSizeData>> GetPoolSizeDataAsync()
@@ -158,75 +169,83 @@ namespace Horde.Build.Agents.Fleet
 			return pools.Select(pool => new PoolSizeData(pool, GetAgentsInPool(pool.Id), null)).ToList();
 		}
 
-		internal async Task ResizePoolsAsync(List<PoolSizeData> poolSizeDatas)
+		internal async Task ResizePoolAsync(PoolSizeData poolSizeData, CancellationToken cancellationToken)
 		{
-			foreach (PoolSizeData poolSizeData in poolSizeDatas.OrderByDescending(x => x.Agents.Count))
+			IPool pool = poolSizeData.Pool;
+
+			if (!pool.EnableAutoscaling || poolSizeData.DesiredAgentCount == null)
 			{
-				IPool pool = poolSizeData.Pool;
+				return;
+			}
 
-				if (!pool.EnableAutoscaling || poolSizeData.DesiredAgentCount == null)
-				{
-					continue;
-				}
+			int currentAgentCount = poolSizeData.Agents.Count;
+			int desiredAgentCount = poolSizeData.DesiredAgentCount.Value;
+			int deltaAgentCount = desiredAgentCount - currentAgentCount;
 
-				int currentAgentCount = poolSizeData.Agents.Count;
-				int desiredAgentCount = poolSizeData.DesiredAgentCount.Value;
-				int deltaAgentCount = desiredAgentCount - currentAgentCount;
+			_logger.LogInformation("{PoolName,-48} Current={Current,4} Target={Target,4} Delta={Delta,4} Status={Status}", pool.Name, currentAgentCount, desiredAgentCount, deltaAgentCount, poolSizeData.StatusMessage);
 
-				_logger.LogInformation("{PoolName,-48} Current={Current,4} Target={Target,4} Delta={Delta,4} Status={Status}", pool.Name, currentAgentCount, desiredAgentCount, deltaAgentCount, poolSizeData.StatusMessage);
+			IFleetManager fleetManager = _fleetManagerFactory(poolSizeData.Pool);
+			try
+			{
+				using IScope scope = GlobalTracer.Instance.BuildSpan("ScalingPool").StartActive();
+				scope.Span.SetTag("poolName", pool.Name);
+				scope.Span.SetTag("current", currentAgentCount);
+				scope.Span.SetTag("desired", desiredAgentCount);
+				scope.Span.SetTag("delta", deltaAgentCount);
 
-				IFleetManager fleetManager = _fleetManagerFactory(poolSizeData.Pool);
-				try
+				if (deltaAgentCount > 0)
 				{
-					using IScope scope = GlobalTracer.Instance.BuildSpan("ScalingPool").StartActive();
-					scope.Span.SetTag("poolName", pool.Name);
-					scope.Span.SetTag("current", currentAgentCount);
-					scope.Span.SetTag("desired", desiredAgentCount);
-					scope.Span.SetTag("delta", deltaAgentCount);
-
-					if (deltaAgentCount > 0)
+					TimeSpan scaleOutCooldown = pool.ScaleOutCooldown ?? _defaultScaleOutCooldown;
+					bool isCoolingDown = pool.LastScaleUpTime != null && pool.LastScaleUpTime + scaleOutCooldown > _clock.UtcNow;
+					scope.Span.SetTag("isCoolingDown", isCoolingDown);
+					if (!isCoolingDown)
 					{
-						TimeSpan scaleOutCooldown = pool.ScaleOutCooldown ?? _defaultScaleOutCooldown;
-						bool isCoolingDown = pool.LastScaleUpTime != null && pool.LastScaleUpTime + scaleOutCooldown > _clock.UtcNow;
-						scope.Span.SetTag("isCoolingDown", isCoolingDown);
-						if (!isCoolingDown)
-						{
-							await fleetManager.ExpandPoolAsync(pool, poolSizeData.Agents, deltaAgentCount);
-							await _poolCollection.TryUpdateAsync(pool, lastScaleUpTime: _clock.UtcNow);
-						}
-						else
-						{
-							TimeSpan? cooldownTimeLeft = pool.LastScaleUpTime + _defaultScaleOutCooldown - _clock.UtcNow;
-							_logger.LogDebug("Cannot scale out {PoolName}, it's cooling down for another {TimeLeft} secs", pool.Name, cooldownTimeLeft?.TotalSeconds);
-						}
+						await fleetManager.ExpandPoolAsync(pool, poolSizeData.Agents, deltaAgentCount, cancellationToken);
+						await _poolCollection.TryUpdateAsync(pool, lastScaleUpTime: _clock.UtcNow);
 					}
-
-					if (deltaAgentCount < 0)
+					else
 					{
-						TimeSpan scaleInCooldown = pool.ScaleInCooldown ?? _defaultScaleInCooldown;
-						bool isCoolingDown = pool.LastScaleDownTime != null && pool.LastScaleDownTime + scaleInCooldown > _clock.UtcNow;
-						scope.Span.SetTag("isCoolingDown", isCoolingDown);
-						if (!isCoolingDown)
-						{
-							await fleetManager.ShrinkPoolAsync(pool, poolSizeData.Agents, -deltaAgentCount);
-							await _poolCollection.TryUpdateAsync(pool, lastScaleDownTime: _clock.UtcNow);
-						}
-						else
-						{
-							TimeSpan? cooldownTimeLeft = pool.LastScaleDownTime + _defaultScaleInCooldown - _clock.UtcNow;
-							_logger.LogDebug("Cannot scale in {PoolName}, it's cooling down for another {TimeLeft} secs", pool.Name, cooldownTimeLeft?.TotalSeconds);
-						}
+						TimeSpan? cooldownTimeLeft = pool.LastScaleUpTime + _defaultScaleOutCooldown - _clock.UtcNow;
+						_logger.LogDebug("Cannot scale out {PoolName}, it's cooling down for another {TimeLeft} secs", pool.Name, cooldownTimeLeft?.TotalSeconds);
 					}
 				}
-				catch (Exception ex)
+
+				if (deltaAgentCount < 0)
 				{
-					_logger.LogError(ex, "Failed to scale {PoolName}:\n{Exception}", pool.Name, ex);
-					continue;
+					TimeSpan scaleInCooldown = pool.ScaleInCooldown ?? _defaultScaleInCooldown;
+					bool isCoolingDown = pool.LastScaleDownTime != null && pool.LastScaleDownTime + scaleInCooldown > _clock.UtcNow;
+					scope.Span.SetTag("isCoolingDown", isCoolingDown);
+					if (!isCoolingDown)
+					{
+						await fleetManager.ShrinkPoolAsync(pool, poolSizeData.Agents, -deltaAgentCount, cancellationToken);
+						await _poolCollection.TryUpdateAsync(pool, lastScaleDownTime: _clock.UtcNow);
+					}
+					else
+					{
+						TimeSpan? cooldownTimeLeft = pool.LastScaleDownTime + _defaultScaleInCooldown - _clock.UtcNow;
+						_logger.LogDebug("Cannot scale in {PoolName}, it's cooling down for another {TimeLeft} secs", pool.Name, cooldownTimeLeft?.TotalSeconds);
+					}
 				}
-
-				_dogStatsd.Gauge("agentpools.autoscale.target", desiredAgentCount, tags: new []{"pool:" + pool.Name});
-				_dogStatsd.Gauge("agentpools.autoscale.current", currentAgentCount, tags: new []{"pool:" + pool.Name});
 			}
+			catch (Exception ex)
+			{
+				_logger.LogError(ex, "Failed to scale {PoolName}:\n{Exception}", pool.Name, ex);
+				return;
+			}
+
+			_dogStatsd.Gauge("agentpools.autoscale.target", desiredAgentCount, tags: new []{"pool:" + pool.Name});
+			_dogStatsd.Gauge("agentpools.autoscale.current", currentAgentCount, tags: new []{"pool:" + pool.Name});
+		}
+		
+		internal async Task ResizePoolsAsync(List<PoolSizeData> poolSizeDatas, CancellationToken cancellationToken = default)
+		{
+			using IScope _ = GlobalTracer.Instance.BuildSpan("AutoscaleService.ResizePoolsAsync").StartActive();
+			
+			ParallelOptions options = new () { MaxDegreeOfParallelism = MaxParallelTasks, CancellationToken = cancellationToken };
+			await Parallel.ForEachAsync(poolSizeDatas.OrderByDescending(x => x.Agents.Count), options, async (poolSizeData, innerCt) =>
+			{
+				await ResizePoolAsync(poolSizeData, innerCt);
+			});
 		}
 		
 		private IEnumerable<string> GetPropValues(string name)
