Horde: Add new job change query options and unify stream template responses
#rnx
#skipci

diff --git a/Engine/Source/Programs/Horde/HordeDashboard/src/backend/Api.ts b/Engine/Source/Programs/Horde/HordeDashboard/src/backend/Api.ts
--- a/Engine/Source/Programs/Horde/HordeDashboard/src/backend/Api.ts
+++ b/Engine/Source/Programs/Horde/HordeDashboard/src/backend/Api.ts
@@ -229,10 +229,6 @@ export type StreamData = GetStreamResponse & {
 
 }
 
-export type TemplateData = GetTemplateResponse & {
-	ref?: GetTemplateRefResponse;
-}
-
 export type JobData = GetJobResponse & {
 	graphRef?: GetGraphResponse;
 }
@@ -950,19 +946,6 @@ export type CreateLabelRequest = {
 	includedNodes: string[];
 }
 
-/** Query selecting the base changelist to use */
-export type ChangeQueryRequest = {
-
-	/** The template id to query */
-	templateId?: string;
-
-	/** The target to query	*/
-	target?: string;
-
-	/** Whether to match a job that produced warnings */
-	outcomes?: JobStepOutcome[];
-}
-
 
 /**Parameters required to create a job */
 export type CreateJobRequest = {
@@ -979,8 +962,8 @@ export type CreateJobRequest = {
 	/** The changelist number to build. Can be null for latest. */
 	change?: number;
 
-	/** Parameters to use when selecting the change to execute at. */
-	changeQuery?: ChangeQueryRequest;
+	/// List of change queries to evaluate
+	changeQueries?: ChangeQueryConfig[];
 
 	/** The preflight changelist number */
 	preflightChange?: number;
@@ -1910,13 +1893,61 @@ export type UpdateStepStateRequest = {
 
 }
 
-/**Information about a template in this stream */
-export type GetTemplateRefResponse = {
 
-	/** The name of the template */
+/// Response describing a template
+export type GetTemplateResponseBase = {
+		
+	/// Name of the template
 	name: string;
 
-	/**Unique id of this template ref */
+	/// Default priority for this job
+	priority?: Priority;
+	
+	/// Whether to allow preflights of this template
+	allowPreflights: boolean;
+	
+	/// Whether to always update issues on jobs using this template
+	updateIssues: boolean;
+	
+	/// The initial agent type to parse the BuildGraph script on
+	initialAgentType?: string;
+	
+	/// Path to a file within the stream to submit to generate a new changelist for jobs
+	submitNewChange?: string;
+	
+	/// Parameters for the job.
+	arguments: string[];
+	
+	/// List of parameters for this template
+	parameters: ParameterData[];
+}
+
+/// Query selecting the base changelist to use
+export type ChangeQueryConfig = {
+	
+	/// Name of this query, for display on the dashboard.
+	name?: string;
+	
+	/// Condition to evaluate before deciding to use this query. May query tags in a preflight.	
+	condition?: any;
+		
+	/// The template id to query	
+	templateId?: string;
+	
+	/// The target to query	
+	target?: string;
+	
+	/// Whether to match a job that produced warnings	
+	outcomes?: JobStepOutcome[];
+	
+	/// Finds the last commit with this tag	
+	commitTag?: any;
+}
+
+/**Information about a template in this stream */
+export type GetTemplateRefResponse = GetTemplateResponseBase & {
+
+	/**Unique id of this template ref, (sanitized name) */
 	id: string;
 
 	/**Hash of the template definition */
@@ -1927,19 +1958,17 @@ export type GetTemplateRefResponse = {
 
 	/** Step state for template in stream */
 	stepStates?: GetTemplateStepStateResponse[];
+
+	/** List of queries for the default changelist */
+	defaultChange?: ChangeQueryConfig[];
+
 }
 
 /** Specifies defaults for running a preflight */
-export type DefaultPreflightRequest = {
+export type DefaultPreflightConfig = {
 
 	/** The template id to query */
 	templateId?: string;
-
-	/** Query for the change to use */
-	change?: ChangeQueryRequest;
-
-	/**  The last successful job type to use for the base changelist (DEPRECATED!)*/
-	changeTemplateId?: string;
 }
 
 
@@ -1975,11 +2004,8 @@ export type GetStreamResponse = {
 	/**List of tabs to display for this stream*/
 	tabs: GetStreamTabResponse[];
 
-	/**  Default template to use for preflights (deprecated)*/
-	defaultPreflightTemplate?: string;
-
 	/** Default template for running preflights */
-	defaultPreflight?: DefaultPreflightRequest;
+	defaultPreflight?: DefaultPreflightConfig;
 
 	/**Map of agent name to type */
 	agentTypes: { [key: string]: GetAgentTypeResponse };
@@ -2182,41 +2208,6 @@ export type BoolParameterData = ParameterData & {
 
 }
 
-
-/**Response describing a template */
-export type GetTemplateResponse = {
-
-	/**Unique id of the template */
-	id: string;
-
-	/**Name of the template */
-	name: string;
-
-	/**Default priority for this job */
-	priority?: Priority;
-
-	/**Whether to allow preflights of this template */
-	allowPreflights: boolean;
-
-	/**Whether to always update issues on jobs that use this template */
-	updateIssues: boolean;
-
-	/**List of node groups for this job */
-	groups: CreateGroupRequest[];
-
-	/**List of template aggregates */
-	aggregates: CreateAggregateRequest[];
-
-	/**Parameters for the job. */
-	arguments: string[];
-
-	/**List of parameters for this template */
-	parameters: ParameterData[];
-
-	/**Custom permissions for this object */
-	acl?: GetAclResponse;
-}
-
 /** Information about a commit */
 export type GetChangeSummaryResponse = {
 
@@ -2232,7 +2223,6 @@ export type GetChangeSummaryResponse = {
 
 }
 
-
 /**Information about a device attached to this agent */
 export type GetDeviceCapabilitiesResponse = {
 
@@ -3869,7 +3859,7 @@ export type StreamConfig = {
 	defaultPreflightTemplate?: string;
 
 	/// Default template for running preflights
-	defaultPreflight?: DefaultPreflightRequest;
+	defaultPreflight?: DefaultPreflightConfig;
 
 	/// <summary>
 	/// List of tabs to show for the new stream
diff --git a/Engine/Source/Programs/Horde/HordeDashboard/src/backend/Backend.ts b/Engine/Source/Programs/Horde/HordeDashboard/src/backend/Backend.ts
--- a/Engine/Source/Programs/Horde/HordeDashboard/src/backend/Backend.ts
+++ b/Engine/Source/Programs/Horde/HordeDashboard/src/backend/Backend.ts
@@ -3,7 +3,7 @@
 import { configure } from "mobx";
 import { isNumber } from 'util';
 import templateCache from '../backend/TemplateCache';
-import { AgentData, AgentQuery, ArtifactData, AuditLogEntry, AuditLogQuery, BatchUpdatePoolRequest, ChangeSummaryData, CreateDeviceRequest, CreateDeviceResponse, CreateExternalIssueRequest, CreateExternalIssueResponse, CreateJobRequest, CreateJobResponse, CreateNoticeRequest, CreatePoolRequest, CreateSoftwareResponse, CreateSubscriptionRequest, CreateSubscriptionResponse, DashboardPreference, DevicePoolTelemetryQuery, DeviceTelemetryQuery, EventData, FindIssueResponse, FindJobTimingsResponse, GetAgentSoftwareChannelResponse, GetArtifactZipRequest, GetDashboardConfigResponse, GetDevicePlatformResponse, GetDevicePoolResponse, GetDevicePoolTelemetryResponse, GetDeviceReservationResponse, GetDeviceResponse, GetDeviceTelemetryResponse, GetExternalIssueProjectResponse, GetExternalIssueResponse, GetGraphResponse, GetIssueStreamResponse, GetJobsTabResponse, GetJobStepRefResponse, GetJobStepTraceResponse, GetJobTimingResponse, GetLogEventResponse, GetNoticeResponse, GetNotificationResponse, GetPerforceServerStatusResponse, GetPoolResponse, GetServerInfoResponse, GetServerSettingsResponse, GetSoftwareResponse, GetSubscriptionResponse, GetUserResponse, GetUtilizationTelemetryResponse, GlobalConfig, IssueData, IssueQuery, IssueQueryV2, JobData, JobQuery, JobsTabColumnType, JobStepOutcome, JobStreamQuery, JobTimingsQuery, LeaseData, LogData, LogLineData, PoolData, ProjectData, ScheduleData, ScheduleQuery, SearchLogFileResponse, ServerUpdateResponse, SessionData, StreamData, TabType, TemplateData, TestData, UpdateAgentRequest, UpdateDeviceRequest, UpdateGlobalConfigRequest, UpdateIssueRequest, UpdateJobRequest, UpdateLeaseRequest, UpdateNoticeRequest, UpdateNotificationsRequest, UpdatePoolRequest, UpdateServerSettingsRequest, UpdateStepRequest, UpdateStepResponse, UpdateTemplateRefRequest, UpdateUserRequest, UsersQuery } from './Api';
+import { AgentData, AgentQuery, ArtifactData, AuditLogEntry, AuditLogQuery, BatchUpdatePoolRequest, ChangeSummaryData, CreateDeviceRequest, CreateDeviceResponse, CreateExternalIssueRequest, CreateExternalIssueResponse, CreateJobRequest, CreateJobResponse, CreateNoticeRequest, CreatePoolRequest, CreateSoftwareResponse, CreateSubscriptionRequest, CreateSubscriptionResponse, DashboardPreference, DevicePoolTelemetryQuery, DeviceTelemetryQuery, EventData, FindIssueResponse, FindJobTimingsResponse, GetAgentSoftwareChannelResponse, GetArtifactZipRequest, GetDashboardConfigResponse, GetDevicePlatformResponse, GetDevicePoolResponse, GetDevicePoolTelemetryResponse, GetDeviceReservationResponse, GetDeviceResponse, GetDeviceTelemetryResponse, GetExternalIssueProjectResponse, GetExternalIssueResponse, GetGraphResponse, GetIssueStreamResponse, GetJobsTabResponse, GetJobStepRefResponse, GetJobStepTraceResponse, GetJobTimingResponse, GetLogEventResponse, GetNoticeResponse, GetNotificationResponse, GetPerforceServerStatusResponse, GetPoolResponse, GetServerInfoResponse, GetServerSettingsResponse, GetSoftwareResponse, GetSubscriptionResponse, GetUserResponse, GetUtilizationTelemetryResponse, GlobalConfig, IssueData, IssueQuery, IssueQueryV2, JobData, JobQuery, JobsTabColumnType, JobStepOutcome, JobStreamQuery, JobTimingsQuery, LeaseData, LogData, LogLineData, PoolData, ProjectData, ScheduleData, ScheduleQuery, SearchLogFileResponse, ServerUpdateResponse, SessionData, StreamData, TabType, TestData, UpdateAgentRequest, UpdateDeviceRequest, UpdateGlobalConfigRequest, UpdateIssueRequest, UpdateJobRequest, UpdateLeaseRequest, UpdateNoticeRequest, UpdateNotificationsRequest, UpdatePoolRequest, UpdateServerSettingsRequest, UpdateStepRequest, UpdateStepResponse, UpdateTemplateRefRequest, UpdateUserRequest, UsersQuery } from './Api';
 import dashboard from './Dashboard';
 import { ChallengeStatus, Fetch } from './Fetch';
 import graphCache, { GraphQuery } from './GraphCache';
@@ -92,14 +92,6 @@ export class Backend {
 
     }
 
-    getTemplate(streamId: string, templateRefId: string): Promise<TemplateData> {
-        return new Promise<TemplateData>((resolve, reject) => {
-            this.backend.get(`/api/v1/streams/${streamId}/templates/${templateRefId}`)
-                .then(response => resolve(response.data as TemplateData))
-                .catch(reason => reject(reason));
-        });
-    }
-
     updateTemplateRef(streamId: string, templateRefId: string, request: UpdateTemplateRefRequest): Promise<boolean> {
 
         return new Promise<boolean>((resolve, reject) => {
@@ -236,22 +228,6 @@ export class Backend {
         });
     }
 
-
-    getJobTemplate(jobId: string): Promise<TemplateData> {
-
-        return new Promise<TemplateData>((resolve, reject) => {
-
-            this.backend.get(`/api/v1/jobs/${jobId}/template`, {
-                params: {}
-            }).then((value) => {
-                resolve(value.data as TemplateData);
-            }).catch(reason => {
-                reject(reason);
-            });
-        });
-
-    }
-
     retryJobStep(jobId: string, batchId: string, stepId: string): Promise<any> {
 
         return new Promise<any>((resolve, reject) => {
diff --git a/Engine/Source/Programs/Horde/HordeDashboard/src/backend/JobDetails.ts b/Engine/Source/Programs/Horde/HordeDashboard/src/backend/JobDetails.ts
--- a/Engine/Source/Programs/Horde/HordeDashboard/src/backend/JobDetails.ts
+++ b/Engine/Source/Programs/Horde/HordeDashboard/src/backend/JobDetails.ts
@@ -5,7 +5,7 @@ import { getTheme, mergeStyles, mergeStyleSets } from '@fluentui/react/lib/Styli
 import backend from '.';
 import { getBatchInitElapsed, getNiceTime, getStepElapsed, getStepETA, getStepFinishTime, getStepTimingDelta } from '../base/utilities/timeUtils';
 import { getBatchText } from '../components/JobDetailCommon';
-import { AgentData, ArtifactData, BatchData, EventData, GetGroupResponse, GetJobStepRefResponse, GetJobTimingResponse, GetLabelResponse, GetLabelStateResponse, GetLabelTimingInfoResponse, GetTemplateResponse, GroupData, IssueData, JobData, JobState, JobStepBatchState, JobStepError, JobStepOutcome, JobStepState, LabelState, NodeData, ReportPlacement, StepData, StreamData, TestData } from './Api';
+import { AgentData, ArtifactData, BatchData, EventData, GetGroupResponse, GetJobStepRefResponse, GetJobTimingResponse, GetLabelResponse, GetLabelStateResponse, GetLabelTimingInfoResponse, GroupData, IssueData, JobData, JobState, JobStepBatchState, JobStepError, JobStepOutcome, JobStepState, LabelState, NodeData, ReportPlacement, StepData, StreamData, TestData } from './Api';
 import { projectStore } from './ProjectStore';
 import moment from 'moment';
 
@@ -561,7 +561,6 @@ export class JobDetails {
         this.issues = [];
         this.history = undefined;
         this.labels = [];
-        this.template = undefined;
         this.fatalError = undefined;
         this.state = JobStepState.Waiting;
         this.reportData = new Map<string, string>();
@@ -904,14 +903,6 @@ export class JobDetails {
                 requests.push(this.getLogEvents(logId));
             }
 
-            this.template = JobDetails.templates.get(this.jobdata!.templateHash!);
-
-            let templateIdx = -1;
-            if (!this.template && !this.isLogView) {
-                templateIdx = requests.length;
-                requests.push(backend.getJobTemplate(this.jobdata!.id));
-            }
-
             let historyIdx = -1;
             if (stepName) {
                 historyIdx = requests.length;
@@ -934,13 +925,6 @@ export class JobDetails {
                     this.history = values[historyIdx] as any;
                 }
 
-                if (templateIdx !== -1) {
-                    const template = values[templateIdx] as GetTemplateResponse;
-                    JobDetails.templates.set(this.jobdata!.id, template);
-                    this.template = template;
-
-                }
-
                 this.process();
 
                 const agentIds: string[] = [];
@@ -1014,7 +998,6 @@ export class JobDetails {
     history?: GetJobStepRefResponse[] = undefined;
     private timing?: GetJobTimingResponse;
     agents: Map<string, AgentData> = new Map();
-    template?: GetTemplateResponse;
 
     outcome: JobStepOutcome = JobStepOutcome.Failure;
     state: JobStepState = JobStepState.Waiting;
@@ -1041,7 +1024,6 @@ export class JobDetails {
 
     suppressIssues: boolean = false;
 
-    private static templates = new Map<string, GetTemplateResponse>();
 }
 
 
diff --git a/Engine/Source/Programs/Horde/HordeDashboard/src/backend/TemplateCache.ts b/Engine/Source/Programs/Horde/HordeDashboard/src/backend/TemplateCache.ts
--- a/Engine/Source/Programs/Horde/HordeDashboard/src/backend/TemplateCache.ts
+++ b/Engine/Source/Programs/Horde/HordeDashboard/src/backend/TemplateCache.ts
@@ -1,85 +1,19 @@
 // Copyright Epic Games, Inc. All Rights Reserved.
 
-import backend from '.';
-import { GetTemplateRefResponse, StreamData, TemplateData } from './Api';
-import { LocalCache } from './LocalCache';
+import { GetTemplateRefResponse, StreamData } from './Api';
 
 export class TemplateCache {
 
-	get(ref: GetTemplateRefResponse): TemplateData | undefined {
+	getStreamTemplates(stream: StreamData): Promise<GetTemplateRefResponse[]> {
 
-		const key = ref.id + ref.hash;
-		return this.cache.get(key);
-	}
-
-	set(ref: GetTemplateRefResponse, template: TemplateData) {
-
-		const key = ref.id + ref.hash;
-		this.cache.set(key, template);
-
-	}
-
-	getStreamTemplates(stream: StreamData): Promise<TemplateData[]> {
-
-		return new Promise<TemplateData[]>(async (resolve, reject) => {
-
-			const templates: TemplateData[] = [];
-
-			const unresolved: GetTemplateRefResponse[] = [];
-
-			for (let i = 0; i < stream.templates.length; i++) {
-
-				const tref = stream.templates[i];
-
-				let template = this.get(tref);
-
-				if (!template) {
-					template = await templateDatabase.getItem(tref.id + tref.hash);
-					if (template) {						
-						this.set(tref, template);
-					}
-				}
-
-				if (template) {
-					templates.push(template);
-				} else {
-					unresolved.push(tref);
-				}
-
-			}
-
-			// get stream templates
-			Promise.all(unresolved.map(async (tref) => {
-				return backend.getTemplate(stream.id, tref.id).then((response) => {
-					response.ref = tref;
-					this.set(tref, response);
-
-					templateDatabase.storeItem(tref.id + tref.hash, response);
-
-					templates.push(response);
-				});
-			})).then(() => {
-				resolve(templates);
-			}).catch(reason => reject(reason));
+		// if this seems odd, it is, part of cleanup to solely use GetTemplateRefResponse and effectively ignpre legacy GetTemplateResponse
+		return new Promise<GetTemplateRefResponse[]>(resolve => {
+			resolve(stream.templates ?? []);
 		});
 	}
 
 	initialize() {
-		templateDatabase.initialize();
 	}
-
-	cache: Map<string, TemplateData> = new Map();
-
 }
 
-class TemplateDatabase extends LocalCache<TemplateData> {
-
-	constructor() {
-		super("Horde", "TemplateCache", "Cache for Horde template objects");
-	}
-
-}
-
-const templateDatabase = new TemplateDatabase();
-
 export default new TemplateCache();
\ No newline at end of file
diff --git a/Engine/Source/Programs/Horde/HordeDashboard/src/components/IssueViewV2.tsx b/Engine/Source/Programs/Horde/HordeDashboard/src/components/IssueViewV2.tsx
--- a/Engine/Source/Programs/Horde/HordeDashboard/src/components/IssueViewV2.tsx
+++ b/Engine/Source/Programs/Horde/HordeDashboard/src/components/IssueViewV2.tsx
@@ -5,7 +5,7 @@ import moment from "moment";
 import React, { useEffect, useRef, useState } from "react";
 import { Link, useHistory, useLocation } from "react-router-dom";
 import backend from "../backend";
-import { CreateExternalIssueResponse, CreateJobRequest, EventSeverity, GetChangeSummaryResponse, GetExternalIssueProjectResponse, GetExternalIssueResponse, GetIssueResponse, GetIssueSpanResponse, GetIssueStepResponse, GetIssueStreamResponse, GetLogEventResponse, GetTemplateRefResponse, GetThinUserInfoResponse, GetUserResponse, IssueSeverity, TemplateData, UpdateIssueRequest } from "../backend/Api";
+import { CreateExternalIssueResponse, CreateJobRequest, EventSeverity, GetChangeSummaryResponse, GetExternalIssueProjectResponse, GetExternalIssueResponse, GetIssueResponse, GetIssueSpanResponse, GetIssueStepResponse, GetIssueStreamResponse, GetLogEventResponse, GetTemplateRefResponse, GetThinUserInfoResponse, GetUserResponse, IssueSeverity, UpdateIssueRequest } from "../backend/Api";
 import dashboard, { StatusColor } from "../backend/Dashboard";
 import { projectStore } from "../backend/ProjectStore";
 import templateCache from '../backend/TemplateCache';
@@ -2954,7 +2954,7 @@ export const IssueForceCloseModal: React.FC<{ onClose: () => void }> = ({ onClos
 const TestFixModal: React.FC<{ onClose: () => void }> = ({ onClose }) => {
 
    const history = useHistory();
-   const [state, setState] = useState<{ loadingTemplates?: boolean, templates?: Map<string, TemplateData[]>, error?: string, submitting?: boolean, shelvedCL?: string, baseCL?: string, streamId?: string, templateId?: string, target?: string, updateIssues?: boolean }>({ updateIssues: true });
+   const [state, setState] = useState<{ loadingTemplates?: boolean, templates?: Map<string, GetTemplateRefResponse[]>, error?: string, submitting?: boolean, shelvedCL?: string, baseCL?: string, streamId?: string, templateId?: string, target?: string, updateIssues?: boolean }>({ updateIssues: true });
 
    const issue = details.issue!;
 
@@ -2979,7 +2979,7 @@ const TestFixModal: React.FC<{ onClose: () => void }> = ({ onClose }) => {
 
          const getTemplates = async () => {
 
-            const templateMap: Map<string, TemplateData[]> = new Map();
+            const templateMap: Map<string, GetTemplateRefResponse[]> = new Map();
 
             for (let i = 0; i < testStreams.length; i++) {
 
@@ -3066,7 +3066,7 @@ const TestFixModal: React.FC<{ onClose: () => void }> = ({ onClose }) => {
          error = "Target not set";
       }
 
-      const template = state.templates?.get(state.streamId!)?.find(t => t.ref?.id === state.templateId);
+      const template = state.templates?.get(state.streamId!)?.find(t => t.id === state.templateId);
       if (!template) {
          error = `Unable to find stream ${state.streamId} template ${state.templateId}`;
          console.log(state.templates?.get(state.streamId!));
diff --git a/Engine/Source/Programs/Horde/HordeDashboard/src/components/JobFilterBar.tsx b/Engine/Source/Programs/Horde/HordeDashboard/src/components/JobFilterBar.tsx
--- a/Engine/Source/Programs/Horde/HordeDashboard/src/components/JobFilterBar.tsx
+++ b/Engine/Source/Programs/Horde/HordeDashboard/src/components/JobFilterBar.tsx
@@ -4,7 +4,7 @@ import { action, observable } from "mobx";
 import { Dropdown, IDropdownOption, Stack } from "@fluentui/react";
 import React, { useEffect, useState } from "react";
 import { useBackend } from "../backend";
-import { GetJobsTabResponse, TabType, TemplateData } from "../backend/Api";
+import { GetJobsTabResponse, TabType, GetTemplateRefResponse } from "../backend/Api";
 import { FilterStatus } from "../backend/JobHandler";
 import templateCache from '../backend/TemplateCache';
 import { hordeClasses, modeColors } from "../styles/Styles";
@@ -43,7 +43,7 @@ export class JobFilter {
         this.updated++;
     }
 
-    set(templates?: TemplateData[], status?: FilterStatus[]) {
+    set(templates?: GetTemplateRefResponse[], status?: FilterStatus[]) {
 
         let templatesDirty = true;
 
@@ -90,7 +90,7 @@ export class JobFilter {
     @observable
     updated: number = 0;
 
-    templates?: TemplateData[];
+    templates?: GetTemplateRefResponse[];
     status?: FilterStatus[];
 }
 
@@ -98,7 +98,7 @@ export const jobFilter = new JobFilter();
 
 interface CategoryItem extends IDropdownOption {
 
-    templates: TemplateData[];
+    templates: GetTemplateRefResponse[];
 }
 
 // Job filter bar for "all" jobs view
@@ -106,7 +106,7 @@ export const JobFilterBar: React.FC<{ streamId: string }> = ({ streamId }) => {
 
     const { projectStore } = useBackend();
 
-    const [state, setState] = useState<{ streamId?: string; templates?: TemplateData[], categories?: string[], status?: FilterStatus | "All" | undefined }>({});
+    const [state, setState] = useState<{ streamId?: string; templates?: GetTemplateRefResponse[], categories?: string[], status?: FilterStatus | "All" | undefined }>({});
 
     const stream = projectStore.streamById(streamId);
 
@@ -132,13 +132,13 @@ export const JobFilterBar: React.FC<{ streamId: string }> = ({ streamId }) => {
         return null;
     } 
 
-    let templates: TemplateData[] = state.templates.map(t => t);
+    let templates: GetTemplateRefResponse[] = state.templates.map(t => t);
 
     if (!templates.length) {
         return null;
     }
 
-    const catMap = new Map<string, TemplateData[]>();
+    const catMap = new Map<string, GetTemplateRefResponse[]>();
 
     stream.tabs?.forEach(t => {
 
@@ -148,7 +148,7 @@ export const JobFilterBar: React.FC<{ streamId: string }> = ({ streamId }) => {
 
         let tab = t as GetJobsTabResponse;
 
-        const ctemps = tab.templates?.map(tid => templates.find(t => t.ref?.id === tid)).filter(t => !!t) as TemplateData[];
+        const ctemps = tab.templates?.map(tid => templates.find(t => t.id === tid)).filter(t => !!t) as GetTemplateRefResponse[];
 
         if (!ctemps?.length) {
             return;
diff --git a/Engine/Source/Programs/Horde/HordeDashboard/src/components/JobSearchSimple.tsx b/Engine/Source/Programs/Horde/HordeDashboard/src/components/JobSearchSimple.tsx
--- a/Engine/Source/Programs/Horde/HordeDashboard/src/components/JobSearchSimple.tsx
+++ b/Engine/Source/Programs/Horde/HordeDashboard/src/components/JobSearchSimple.tsx
@@ -4,7 +4,7 @@ import moment from "moment";
 import { useState } from "react";
 import { Link } from "react-router-dom";
 import backend from "../backend";
-import { GetJobResponse, GetJobsTabResponse, GetJobStepRefResponse, JobData, JobQuery, TabType, TemplateData } from "../backend/Api";
+import { GetJobResponse, GetJobsTabResponse, GetJobStepRefResponse, JobData, JobQuery, TabType, GetTemplateRefResponse } from "../backend/Api";
 import dashboard from "../backend/Dashboard";
 import { projectStore } from "../backend/ProjectStore";
 import templateCache from '../backend/TemplateCache';
@@ -25,7 +25,7 @@ export const JobSearchSimpleModal: React.FC<{ streamId: string, onClose: () => v
    };
 
    const [searchState, setSearchState] = useState<{ items: JobItem[], groups?: IGroup[], templateId?: string, name?: string, userId?: string, global?: boolean, streamIdOverride?: string, preflights: boolean, preflightOnly?: boolean, minDate?: Date, maxDate?: Date, minCL?: string, maxCL?: string, querying: boolean, containsStep?: string }>({ items: [], querying: false, preflights: false });
-   const [templateData, setTemplateData] = useState<{ streamId: string; templates: TemplateData[] } | undefined>(undefined);
+   const [templateData, setTemplateData] = useState<{ streamId: string; templates: GetTemplateRefResponse[] } | undefined>(undefined);
 
    const stream = projectStore.streamById(streamId);
 
@@ -62,7 +62,7 @@ export const JobSearchSimpleModal: React.FC<{ streamId: string, onClose: () => v
    let templateOptions: IContextualMenuItem[] = [];
 
 
-   const sorted = new Map<string, TemplateData[]>();
+   const sorted = new Map<string, GetTemplateRefResponse[]>();
 
    templates.forEach(t => {
 
@@ -72,7 +72,7 @@ export const JobSearchSimpleModal: React.FC<{ streamId: string, onClose: () => v
          }
 
          const jtab = tab as GetJobsTabResponse;
-         if (!jtab.templates?.find(template => template === t.ref?.id)) {
+         if (!jtab.templates?.find(template => template === t.id)) {
             return;
          }
 
@@ -273,10 +273,10 @@ export const JobSearchSimpleModal: React.FC<{ streamId: string, onClose: () => v
 
             let stepTemplates: string[] = [];
 
-            if (template?.ref?.id) {
-               stepTemplates.push(template.ref.id);
+            if (template?.id) {
+               stepTemplates.push(template.id);
             } else {
-               stepTemplates = templates.filter(t => !!t.ref?.id).map(t => t.ref!.id!);
+               stepTemplates = templates.filter(t => !!t.id).map(t => t.id);
             }
 
             const numTemplates = stepTemplates.length;
@@ -313,7 +313,7 @@ export const JobSearchSimpleModal: React.FC<{ streamId: string, onClose: () => v
                filter: "id,name,change,createTime,streamId,preflightChange,startedByUserInfo,streamId",
                minChange: minCL,
                maxChange: maxCL,
-               template: (!searchState.global && template?.ref?.id && !searchState.containsStep) ? [template.ref.id] : undefined,
+               template: (!searchState.global && template?.id && !searchState.containsStep) ? [template.id] : undefined,
                minCreateTime: minCreateTime?.toUTCString(),
                maxCreateTime: maxCreateTime?.toUTCString(),
                includePreflight: searchState.preflights || !!searchState.preflightOnly,
diff --git a/Engine/Source/Programs/Horde/HordeDashboard/src/components/JobViewAll.tsx b/Engine/Source/Programs/Horde/HordeDashboard/src/components/JobViewAll.tsx
--- a/Engine/Source/Programs/Horde/HordeDashboard/src/components/JobViewAll.tsx
+++ b/Engine/Source/Programs/Horde/HordeDashboard/src/components/JobViewAll.tsx
@@ -517,7 +517,7 @@ const JobViewAllInner: React.FC<{ filter: JobFilterSimple }> = observer(({ filte
             preflightStartedByUserId = dashboard.userId;
         }
 
-        if (jobHandler.filter(stream, jobFilter.templates?.map(t => t.ref!.id), preflightStartedByUserId)) {
+        if (jobHandler.filter(stream, jobFilter.templates?.map(t => t.id), preflightStartedByUserId)) {
 
         }
     }
diff --git a/Engine/Source/Programs/Horde/HordeDashboard/src/components/NewBuild.tsx b/Engine/Source/Programs/Horde/HordeDashboard/src/components/NewBuild.tsx
--- a/Engine/Source/Programs/Horde/HordeDashboard/src/components/NewBuild.tsx
+++ b/Engine/Source/Programs/Horde/HordeDashboard/src/components/NewBuild.tsx
@@ -5,7 +5,7 @@ import { observer } from 'mobx-react-lite';
 import moment from 'moment';
 import React, { useState } from 'react';
 import backend, { useBackend } from '../backend';
-import { BoolParameterData, ChangeQueryRequest, CreateJobRequest, GetJobsTabResponse, GroupParameterData, JobsTabData, ListParameterData, ListParameterItemData, ListParameterStyle, ParameterData, ParameterType, Priority, TabType, TemplateData, TextParameterData } from '../backend/Api';
+import { BoolParameterData, CreateJobRequest, GetJobsTabResponse, GroupParameterData, JobsTabData, ListParameterData, ListParameterItemData, ListParameterStyle, ParameterData, ParameterType, Priority, TabType, GetTemplateRefResponse, TextParameterData, ChangeQueryConfig } from '../backend/Api';
 
 import templateCache from '../backend/TemplateCache';
 import { ErrorHandler } from "../components/ErrorHandler";
@@ -27,7 +27,7 @@ class BuildParameters {
 
    static paramKey = 0;
 
-   constructor(template: TemplateData, jobDetails?: JobDetailsV2, existing?: BuildParameters) {
+   constructor(template: GetTemplateRefResponse, jobDetails?: JobDetailsV2, existing?: BuildParameters) {
 
       BuildParameters.paramKey++;
 
@@ -292,7 +292,7 @@ class BuildParameters {
    copySettings(src: BuildParameters) {
 
       this.change = src.change;
-      this.changeSuccess = src.changeSuccess;
+      this.changeOption = src.changeOption;
       this.preflight = src.preflight;
       this.autoSubmit = src.autoSubmit;
 
@@ -518,7 +518,7 @@ class BuildParameters {
 
          try {
             regex = new RegExp(param.validation, 'gm');
-         } catch (reason:any) {
+         } catch (reason: any) {
 
             errors.push({
                param: param,
@@ -577,7 +577,7 @@ class BuildParameters {
       this.changed++;
    }
 
-   template: TemplateData;
+   template: GetTemplateRefResponse;
 
    jobDetails?: JobDetailsV2;
 
@@ -587,7 +587,7 @@ class BuildParameters {
    // base settings
 
    change?: string;
-   changeSuccess?: boolean;
+   changeOption?: string;
    preflight?: string;
    autoSubmit?: boolean;
 
@@ -608,11 +608,11 @@ class BuildParameters {
 
 export const NewBuild: React.FC<{ streamId: string; show: boolean; onClose: (newJobId: string | undefined) => void, jobKey?: string; jobDetails?: JobDetailsV2, readOnly?: boolean }> = observer(({ streamId, jobKey, show, onClose, jobDetails, readOnly }) => {
 
-   const { projectStore } = useBackend();   
+   const { projectStore } = useBackend();
    const query = useQuery();
-   const [build, setBuild] = useState<{ template?: TemplateData; buildParams?: BuildParameters }>({});
+   const [build, setBuild] = useState<{ template?: GetTemplateRefResponse; buildParams?: BuildParameters }>({});
    const [submitting, setSubmitting] = useState(false);
-   const [templateData, setTemplateData] = useState<{ streamId: string; templates: TemplateData[] } | undefined>(undefined);
+   const [templateData, setTemplateData] = useState<{ streamId: string; templates: GetTemplateRefResponse[] } | undefined>(undefined);
    const [mode, setMode] = useState<"Basic" | "Advanced">("Basic");
    const [showAllTemplates, setShowAllTemplates] = useState<boolean>(false);
    const [validationErrors, setValidationErrors] = useState<ValidationError[]>([]);
@@ -635,7 +635,7 @@ export const NewBuild: React.FC<{ streamId: string; show: boolean; onClose: (new
       return null;
    }
 
-   let templates: TemplateData[] = [];
+   let templates: GetTemplateRefResponse[] = [];
 
    // @todo: better way of checking whether we are in a preflight submit
    const isPreflightSubmit = !!query.get("shelvedchange");
@@ -672,21 +672,6 @@ export const NewBuild: React.FC<{ streamId: string; show: boolean; onClose: (new
       return null;
    }
 
-   let preflightDefaultTemplate: TemplateData | undefined;
-   let preflightDefaultTemplateRefId = stream.defaultPreflight?.templateId ?? stream.defaultPreflightTemplate;
-   const preflightChangeTemplateRefId = stream.defaultPreflight?.change?.templateId ?? stream.defaultPreflight?.changeTemplateId;
-   let preflightChangeTemplate: TemplateData | undefined;
-
-   if (preflightDefaultTemplateRefId) {
-      preflightDefaultTemplate = templateData.templates.find(t => t.ref?.id === preflightDefaultTemplateRefId)
-   }
-
-   if (preflightChangeTemplateRefId) {
-      preflightChangeTemplate = templateData.templates.find(t => t.ref?.id === preflightChangeTemplateRefId)
-      if (!preflightChangeTemplate) {
-         console.error(`Unable to find preflight change template ${preflightChangeTemplateRefId} in stream templates`);
-      }
-   }
 
    if (!jobDetails) {
 
@@ -709,7 +694,7 @@ export const NewBuild: React.FC<{ streamId: string; show: boolean; onClose: (new
          templates = [];
 
          tab.templates?.forEach(templateName => {
-            const t = templateData.templates.find(t => t.ref?.id === templateName);
+            const t = templateData.templates.find(t => t.id === templateName);
             if (!t) {
                console.error(`Could not find template ${templateName}`);
                return;
@@ -746,13 +731,33 @@ export const NewBuild: React.FC<{ streamId: string; show: boolean; onClose: (new
    const template = build.template;
    const buildParams = build.buildParams;
 
+   let defaultPreflightQuery: ChangeQueryConfig[] | undefined;
+
+   if (template?.defaultChange) {
+      defaultPreflightQuery = template.defaultChange;
+   }
+
+   const defaultStreamPreflightTemplate = stream.templates.find(t => t.id === stream.defaultPreflight?.templateId);
+
+   if (!defaultPreflightQuery && stream.defaultPreflight?.templateId) {
+
+      if (defaultStreamPreflightTemplate) {
+         defaultPreflightQuery = defaultStreamPreflightTemplate.defaultChange;
+      }
+
+      if (!defaultPreflightQuery) {
+         console.error(`Unable to find default stream preflight template ${stream.defaultPreflight?.templateId} in stream templates`);
+      }
+   }
+
+
    if (buildParams?.preflight) {
       templates = templates.filter(t => t.allowPreflights);
    }
 
    if (!template || templates.indexOf(template) === -1) {
 
-      let t: TemplateData | undefined;
+      let t: GetTemplateRefResponse | undefined;
 
       if (jobDetails?.template) {
          t = templates.find(t => t.name === jobDetails.template!.name);
@@ -761,9 +766,9 @@ export const NewBuild: React.FC<{ streamId: string; show: boolean; onClose: (new
       // handle preflight redirect case
       if (!t && defaultShelvedChange && !jobDetails) {
 
-         t = preflightDefaultTemplate;
+         t = defaultStreamPreflightTemplate;
          if (!t) {
-            console.error(`Stream default preflight template cannot be found for stream ${stream.fullname} : defaultPreflightTemplate ${preflightDefaultTemplateRefId}, will use first template in list`);
+            console.error(`Stream default preflight template cannot be found for stream ${stream.fullname} : stream defaultPreflightTemplate ${stream.defaultPreflight?.templateId}, will use first template in list`);
          }
       }
 
@@ -771,7 +776,7 @@ export const NewBuild: React.FC<{ streamId: string; show: boolean; onClose: (new
       if (!t && (showAllTemplates) && stream.tabs.length > 0) {
          const stab = stream.tabs[0] as JobsTabData;
          if (stab.templates && stab.templates.length > 0) {
-            t = templates.find(t => t.ref?.id === stab.templates![0]);
+            t = templates.find(t => t.id === stab.templates![0]);
          }
       }
 
@@ -797,7 +802,7 @@ export const NewBuild: React.FC<{ streamId: string; show: boolean; onClose: (new
          // when coming from P4V, always clear the base CL
          if (isFromP4V) {
             p.change = undefined;
-            p.changeSuccess = false;
+            p.changeOption = undefined;
          }
 
          setBuild({ template: t, buildParams: p });
@@ -831,7 +836,7 @@ export const NewBuild: React.FC<{ streamId: string; show: boolean; onClose: (new
       }).sort((a, b) => a.text < b.text ? -1 : 1);
    } else {
 
-      const sorted = new Map<string, TemplateData[]>();
+      const sorted = new Map<string, GetTemplateRefResponse[]>();
 
       templates.forEach(t => {
 
@@ -841,7 +846,7 @@ export const NewBuild: React.FC<{ streamId: string; show: boolean; onClose: (new
             }
 
             const jtab = tab as GetJobsTabResponse;
-            if (!jtab.templates?.find(template => template === t.ref?.id)) {
+            if (!jtab.templates?.find(template => template === t.id)) {
                return;
             }
 
@@ -1245,21 +1250,21 @@ export const NewBuild: React.FC<{ streamId: string; show: boolean; onClose: (new
          }
       }
 
-      let changeQuery: ChangeQueryRequest | undefined;
+      let changeQueries: ChangeQueryConfig[] | undefined;
 
-      if (buildParams.changeSuccess && preflightChangeTemplateRefId) {
-         change = undefined;
-
-         if (stream.defaultPreflight?.change) {
-            changeQuery = stream.defaultPreflight?.change;
-         } else {
-            changeQuery = {
-               templateId: preflightChangeTemplateRefId
-            };
+      if (typeof change !== 'number') {
+         const changeOption = buildParams.changeOption ?? "Default Change";
+         if (changeOption === "Default Change" && defaultPreflightQuery) {
+            changeQueries = defaultPreflightQuery;
+         }
+         else {
+            const changeQuery = defaultPreflightQuery?.find(p => p.name === buildParams.changeOption);
+            if (changeQuery) {
+               changeQueries = [{ ...changeQuery, condition: undefined }];
+            }
          }
       }
 
-
       if (buildParams!.preflight) {
          preflight = parseInt(buildParams!.preflight);
          if (isNaN(preflight)) {
@@ -1275,7 +1280,6 @@ export const NewBuild: React.FC<{ streamId: string; show: boolean; onClose: (new
 
       const args = buildParams!.generateArguments();
 
-
       let additionalArgs: string[] = [];
 
       if (buildParams!.advAdditionalArgs) {
@@ -1313,7 +1317,7 @@ export const NewBuild: React.FC<{ streamId: string; show: boolean; onClose: (new
 
       allArgs = allArgs.map(a => a.trim());
 
-      const templateId = template?.ref?.id ? template.ref.id : jobDetails?.jobData?.templateId;
+      const templateId = template?.id ? template.id : jobDetails?.jobData?.templateId;
       const updateIssues = template?.updateIssues ? undefined : buildParams!.advUpdateIssues;
 
       const data: CreateJobRequest = {
@@ -1322,7 +1326,7 @@ export const NewBuild: React.FC<{ streamId: string; show: boolean; onClose: (new
          name: buildParams!.advJobName,
          priority: buildParams!.advJobPriority,
          updateIssues: updateIssues,
-         changeQuery: changeQuery,
+         changeQueries: changeQueries,
          arguments: allArgs
       };
 
@@ -1375,7 +1379,7 @@ export const NewBuild: React.FC<{ streamId: string; show: boolean; onClose: (new
                }
 
                redirected = true;
-               onClose(data.id);               
+               onClose(data.id);
             }).catch(reason => {
                // "Not Found" is generally a permissions error
                errorReason = reason ? reason : "Unknown";
@@ -1453,22 +1457,21 @@ export const NewBuild: React.FC<{ streamId: string; show: boolean; onClose: (new
 
    // change setup
 
+   const changeOptions = ["Default Change"];
 
-   const changeItems: IComboBoxOption[] = [
-      { key: 'ChangeLatest', text: 'Latest Change' }
-   ];
-
-   if (preflightChangeTemplate) {
-      changeItems.push({ key: 'ChangeSuccess', text: `Latest Success - ${preflightChangeTemplate.name}` });
+   if (defaultPreflightQuery) {
+      changeOptions.push(...defaultPreflightQuery.filter(n => !!n.name).map(n => n.name!));
    }
 
+   const changeItems: IComboBoxOption[] = changeOptions.map(name => { return { key: `key_change_option_${name}`, text: name } });
+
    let changeText: string | undefined;
 
    if (buildParams.change) {
       changeText = buildParams.change;
    }
-   if (buildParams.changeSuccess) {
-      changeText = "Latest Success";
+   if (buildParams.changeOption) {
+      changeText = buildParams.changeOption;
    }
 
    const parameterGap = 6;
@@ -1558,26 +1561,20 @@ export const NewBuild: React.FC<{ streamId: string; show: boolean; onClose: (new
                      </Stack>
                      <Stack grow />
                      <Stack>
-                        <ComboBox label="Change" text={changeText} options={changeItems} disabled={readOnly} allowFreeform autoComplete="off" placeholder="Latest Change" defaultValue={jobDetails?.jobData?.change?.toString()} onChange={(ev, option, index, value) => {
+                        <ComboBox style={{ width: 240 }} label="Change" text={changeText} options={changeItems} disabled={readOnly} allowFreeform autoComplete="off" placeholder="Default Change" defaultValue={jobDetails?.jobData?.change?.toString()} onChange={(ev, option, index, value) => {
                            ev.preventDefault();
 
                            if (option) {
-                              if (option.key === "ChangeLatest") {
-                                 buildParams.change = undefined;
-                                 buildParams.changeSuccess = false;
-                              } else if (option.key === "ChangeSuccess") {
-                                 buildParams.change = undefined;
-                                 buildParams.changeSuccess = true;
-                              }
+                              buildParams.change = undefined;
+                              buildParams.changeOption = option.text;
                            } else {
-                              buildParams.changeSuccess = false;
+                              buildParams.changeOption = undefined;
                               if (!value) {
                                  buildParams.change = undefined;
                               } else {
                                  buildParams.change = value;
                               }
                            }
-
                            buildParams!.setChanged();
                         }}
                         />
diff --git a/Engine/Source/Programs/Horde/HordeDashboard/src/components/TrendsReportView.tsx b/Engine/Source/Programs/Horde/HordeDashboard/src/components/TrendsReportView.tsx
deleted file mode 100644
diff --git a/Engine/Source/Programs/Horde/HordeDashboard/src/components/jobDetailsV2/JobDetailStepTrends.tsx b/Engine/Source/Programs/Horde/HordeDashboard/src/components/jobDetailsV2/JobDetailStepTrends.tsx
--- a/Engine/Source/Programs/Horde/HordeDashboard/src/components/jobDetailsV2/JobDetailStepTrends.tsx
+++ b/Engine/Source/Programs/Horde/HordeDashboard/src/components/jobDetailsV2/JobDetailStepTrends.tsx
@@ -473,7 +473,7 @@ const StepTrendsPanelInner: React.FC<{ jobDetails: JobDetailsV2; dataView: Trend
          TemplateCache.getStreamTemplates(jobDetails.stream).then(templates => {
             templates.forEach(template => {
                if (template.id === jobDetails.template?.id || template.name === jobDetails.template?.name) {
-                  setSelectedTemplateRef(template.ref);
+                  setSelectedTemplateRef(template);
                   return;
                }
             });
diff --git a/Engine/Source/Programs/Horde/HordeDashboard/src/components/jobDetailsV2/JobDetailsViewCommon.tsx b/Engine/Source/Programs/Horde/HordeDashboard/src/components/jobDetailsV2/JobDetailsViewCommon.tsx
--- a/Engine/Source/Programs/Horde/HordeDashboard/src/components/jobDetailsV2/JobDetailsViewCommon.tsx
+++ b/Engine/Source/Programs/Horde/HordeDashboard/src/components/jobDetailsV2/JobDetailsViewCommon.tsx
@@ -5,7 +5,7 @@ import moment from "moment";
 import React, { useState } from "react";
 import { useHistory, useLocation } from "react-router-dom";
 import backend from "../../backend";
-import { BatchData, GetJobTimingResponse, GetLabelResponse, GroupData, JobData, JobState, LabelState, NodeData, ReportPlacement, StepData, StreamData, TemplateData } from "../../backend/Api";
+import { BatchData, GetJobTimingResponse, GetLabelResponse, GetTemplateRefResponse, GroupData, JobData, JobState, LabelState, NodeData, ReportPlacement, StepData, StreamData } from "../../backend/Api";
 import { JobLabel } from "../../backend/JobDetails";
 import { PollBase } from "../../backend/PollBase";
 import { projectStore } from "../../backend/ProjectStore";
@@ -633,7 +633,7 @@ export class JobDetailsV2 extends PollBase {
       let forceUpdate = false;
       if (initialRequest) {
          const templates = await TemplateCache.getStreamTemplates(this.stream);
-         this.template = templates.find(t => t.ref?.id === jobData.templateId);
+         this.template = templates.find(t => t.id === jobData.templateId);
          if (!this.template) {
             throw new Error(`Unable to get stream template for job ${this.jobId}`)
          }
@@ -704,7 +704,7 @@ export class JobDetailsV2 extends PollBase {
    groups: GroupData[] = [];
    nodes: NodeData[] = [];
    batches: BatchData[] = [];
-   template?: TemplateData;
+   template?: GetTemplateRefResponse;
 
    // artifact id => report contents (markdown)
    private reportData = new Map<string, string>();
