Horde: Cache data fetched inside pool sizing strategies

Recent refactoring of per-pool pool size strategies removed the built-in batching of calculations. By caching expensive lookups for a brief period speeds up scaling when strategies are run.

diff --git a/Engine/Source/Programs/Horde/Horde.Build.Tests/Fleet/AutoscaleServiceV2Test.cs b/Engine/Source/Programs/Horde/Horde.Build.Tests/Fleet/AutoscaleServiceV2Test.cs
--- a/Engine/Source/Programs/Horde/Horde.Build.Tests/Fleet/AutoscaleServiceV2Test.cs
+++ b/Engine/Source/Programs/Horde/Horde.Build.Tests/Fleet/AutoscaleServiceV2Test.cs
@@ -140,7 +140,7 @@ namespace Horde.Build.Tests.Fleet
 			Mock<IAmazonAutoScaling> mockAwsAutoScaling = new ();
 
 			AutoscaleServiceV2 service = new (AgentCollection, GraphCollection, JobCollection, LeaseCollection, PoolCollection,
-				StreamService, _dogStatsD, mockEc2.Object, mockAwsAutoScaling.Object, Clock, serverSettingsOpt, loggerFactory,
+				StreamService, _dogStatsD, mockEc2.Object, mockAwsAutoScaling.Object, Clock, Cache, serverSettingsOpt, loggerFactory,
 				_ => fleetManager);
 			return service;
 		}
diff --git a/Engine/Source/Programs/Horde/Horde.Build.Tests/Fleet/JobQueueStrategyTest.cs b/Engine/Source/Programs/Horde/Horde.Build.Tests/Fleet/JobQueueStrategyTest.cs
--- a/Engine/Source/Programs/Horde/Horde.Build.Tests/Fleet/JobQueueStrategyTest.cs
+++ b/Engine/Source/Programs/Horde/Horde.Build.Tests/Fleet/JobQueueStrategyTest.cs
@@ -144,7 +144,7 @@ namespace Horde.Build.Tests.Fleet
 			}
 
 			JobQueueSettings settings = new ();
-			return (new (JobCollection, GraphCollection, StreamService, Clock, settings), poolSize);
+			return (new (JobCollection, GraphCollection, StreamService, Clock, Cache, settings), poolSize);
 		}
 		
 		private async Task<IJob> AddPlaceholderJob(IGraph graph, StreamId streamId, string nodeNameToExecute)
diff --git a/Engine/Source/Programs/Horde/Horde.Build.Tests/Fleet/LeaseUtilizationStrategyTest.cs b/Engine/Source/Programs/Horde/Horde.Build.Tests/Fleet/LeaseUtilizationStrategyTest.cs
--- a/Engine/Source/Programs/Horde/Horde.Build.Tests/Fleet/LeaseUtilizationStrategyTest.cs
+++ b/Engine/Source/Programs/Horde/Horde.Build.Tests/Fleet/LeaseUtilizationStrategyTest.cs
@@ -97,7 +97,7 @@ namespace Horde.Build.Tests.Fleet
 
 		private async Task AssertPoolSizeAsync(IPool pool, int expectedNumAgents)
 		{
-			LeaseUtilizationStrategy strategy = new (AgentCollection, PoolCollection, LeaseCollection, Clock);
+			LeaseUtilizationStrategy strategy = new (AgentCollection, PoolCollection, LeaseCollection, Clock, Cache);
 			List<PoolSizeData> output = await strategy.CalcDesiredPoolSizesAsync(new() { new(pool, _poolAgents, null) });
 			Assert.AreEqual(1, output.Count);
 			Assert.AreEqual(expectedNumAgents, output[0].DesiredAgentCount);
diff --git a/Engine/Source/Programs/Horde/Horde.Build.Tests/TestSetup.cs b/Engine/Source/Programs/Horde/Horde.Build.Tests/TestSetup.cs
--- a/Engine/Source/Programs/Horde/Horde.Build.Tests/TestSetup.cs
+++ b/Engine/Source/Programs/Horde/Horde.Build.Tests/TestSetup.cs
@@ -76,6 +76,7 @@ namespace Horde.Build.Tests
 	public class TestSetup : DatabaseIntegrationTest
 	{
 		public FakeClock Clock => ServiceProvider.GetRequiredService<FakeClock>();
+		public IMemoryCache Cache => ServiceProvider.GetRequiredService<IMemoryCache>();
 
 		public ConfigCollection ConfigCollection => ServiceProvider.GetRequiredService<ConfigCollection>();
 		public IGraphCollection GraphCollection => ServiceProvider.GetRequiredService<IGraphCollection>();
diff --git a/Engine/Source/Programs/Horde/Horde.Build/Agents/Fleet/AutoscaleServiceV2.cs b/Engine/Source/Programs/Horde/Horde.Build/Agents/Fleet/AutoscaleServiceV2.cs
--- a/Engine/Source/Programs/Horde/Horde.Build/Agents/Fleet/AutoscaleServiceV2.cs
+++ b/Engine/Source/Programs/Horde/Horde.Build/Agents/Fleet/AutoscaleServiceV2.cs
@@ -18,6 +18,7 @@ using Horde.Build.Jobs.Graphs;
 using Horde.Build.Streams;
 using Horde.Build.Utilities;
 using HordeCommon;
+using Microsoft.Extensions.Caching.Memory;
 using Microsoft.Extensions.Hosting;
 using Microsoft.Extensions.Logging;
 using Microsoft.Extensions.Options;
@@ -49,6 +50,7 @@ namespace Horde.Build.Agents.Fleet
 		private readonly IAmazonEC2 _ec2;
 		private readonly IAmazonAutoScaling _awsAutoScaling;
 		private readonly IClock _clock;
+		private readonly IMemoryCache _cache;
 		private readonly ITicker _ticker;
 		private readonly ITicker _tickerHighFrequency;
 		private readonly TimeSpan _defaultScaleOutCooldown;
@@ -74,6 +76,7 @@ namespace Horde.Build.Agents.Fleet
 			IAmazonEC2 ec2,
 			IAmazonAutoScaling awsAutoScaling,
 			IClock clock,
+			IMemoryCache cache,
 			IOptions<ServerSettings> settings,
 			ILoggerFactory loggerFactory,
 			Func<IPool, IFleetManager>? fleetManagerFactory = null)
@@ -88,6 +91,7 @@ namespace Horde.Build.Agents.Fleet
 			_ec2 = ec2;
 			_awsAutoScaling = awsAutoScaling;
 			_clock = clock;
+			_cache = cache;
 			_logger = loggerFactory.CreateLogger<AutoscaleServiceV2>();
 			_loggerFactory = loggerFactory;
 			_ticker = clock.AddSharedTicker<AutoscaleServiceV2>(TimeSpan.FromMinutes(5.0), TickLeaderAsync, _logger);
@@ -345,11 +349,11 @@ namespace Horde.Build.Agents.Fleet
 							case PoolSizeStrategy.JobQueue:
 								JobQueueSettings? settings = JsonSerializer.Deserialize<JobQueueSettings>(info.Config);
 								if (settings == null) throw new ArgumentException("Unable to deserialize pool sizing config");
-								return new JobQueueStrategy(_jobCollection, _graphCollection, _streamService, _clock, settings);
+								return new JobQueueStrategy(_jobCollection, _graphCollection, _streamService, _clock, _cache, settings);
 							
 							case PoolSizeStrategy.LeaseUtilization:
 								// No settings for lease utilization strategy
-								return new LeaseUtilizationStrategy(_agentCollection, _poolCollection, _leaseCollection, _clock);
+								return new LeaseUtilizationStrategy(_agentCollection, _poolCollection, _leaseCollection, _clock, _cache);
 							
 							case PoolSizeStrategy.NoOp:
 								return new NoOpPoolSizeStrategy();
@@ -365,9 +369,9 @@ namespace Horde.Build.Agents.Fleet
 			switch (pool.SizeStrategy)
 			{
 				case PoolSizeStrategy.JobQueue:
-					return new JobQueueStrategy(_jobCollection, _graphCollection, _streamService, _clock, pool.JobQueueSettings);
+					return new JobQueueStrategy(_jobCollection, _graphCollection, _streamService, _clock, _cache, pool.JobQueueSettings);
 				case PoolSizeStrategy.LeaseUtilization:
-					return new LeaseUtilizationStrategy(_agentCollection, _poolCollection, _leaseCollection, _clock);
+					return new LeaseUtilizationStrategy(_agentCollection, _poolCollection, _leaseCollection, _clock, _cache);
 				case PoolSizeStrategy.NoOp:
 					return new NoOpPoolSizeStrategy();
 				default:
diff --git a/Engine/Source/Programs/Horde/Horde.Build/Agents/Fleet/JobQueueStrategy.cs b/Engine/Source/Programs/Horde/Horde.Build/Agents/Fleet/JobQueueStrategy.cs
--- a/Engine/Source/Programs/Horde/Horde.Build/Agents/Fleet/JobQueueStrategy.cs
+++ b/Engine/Source/Programs/Horde/Horde.Build/Agents/Fleet/JobQueueStrategy.cs
@@ -11,6 +11,7 @@ using Horde.Build.Jobs.Graphs;
 using Horde.Build.Streams;
 using Horde.Build.Utilities;
 using HordeCommon;
+using Microsoft.Extensions.Caching.Memory;
 
 namespace Horde.Build.Agents.Fleet
 {
@@ -64,12 +65,14 @@ namespace Horde.Build.Agents.Fleet
 	/// </summary>
 	public class JobQueueStrategy : IPoolSizeStrategy
 	{
+		private const string CacheKey = nameof(JobQueueStrategy);
 		internal JobQueueSettings Settings { get; }
 		
 		private readonly IJobCollection _jobs;
 		private readonly IGraphCollection _graphs;
 		private readonly StreamService _streamService;
 		private readonly IClock _clock;
+		private readonly IMemoryCache _cache;
 		
 		/// <summary>
 		/// How far back in time to look for job batches (that potentially are in the queue)
@@ -93,14 +96,16 @@ namespace Horde.Build.Agents.Fleet
 		/// <param name="graphs"></param>
 		/// <param name="streamService"></param>
 		/// <param name="clock"></param>
+		/// <param name="cache"></param>
 		/// <param name="settings"></param>
 		/// <param name="samplePeriod">Time period for each sample</param>
-		public JobQueueStrategy(IJobCollection jobs, IGraphCollection graphs, StreamService streamService, IClock clock, JobQueueSettings? settings = null, TimeSpan? samplePeriod = null)
+		public JobQueueStrategy(IJobCollection jobs, IGraphCollection graphs, StreamService streamService, IClock clock, IMemoryCache cache, JobQueueSettings? settings = null, TimeSpan? samplePeriod = null)
 		{
 			_jobs = jobs;
 			_graphs = graphs;
 			_streamService = streamService;
 			_clock = clock;
+			_cache = cache;
 			_samplePeriod = samplePeriod ?? _samplePeriod;
 			Settings = settings ?? _defaultPoolSettings;
 		}
@@ -173,7 +178,15 @@ namespace Horde.Build.Agents.Fleet
 		public async Task<List<PoolSizeData>> CalcDesiredPoolSizesAsync(List<PoolSizeData> pools)
 		{
 			DateTimeOffset minCreateTime = _clock.UtcNow - _samplePeriod;
-			Dictionary<PoolId, int> poolQueueSizes = await GetPoolQueueSizesAsync(minCreateTime);
+			Dictionary<PoolId, int> poolQueueSizes;// = await GetPoolQueueSizesAsync(minCreateTime);
+
+			// Cache pool queue sizes for a short while for faster runs when many pools are scaled
+			if (!_cache.TryGetValue(CacheKey, out poolQueueSizes))
+			{
+				// Pool sizes haven't been cached, update them (might happen from multiple tasks but that is fine)
+				poolQueueSizes = await GetPoolQueueSizesAsync(minCreateTime);
+				_cache.Set(CacheKey, poolQueueSizes, TimeSpan.FromSeconds(60));
+			}
 
 			return pools.Select(current =>
 			{
diff --git a/Engine/Source/Programs/Horde/Horde.Build/Agents/Fleet/LeaseUtilizationStrategy.cs b/Engine/Source/Programs/Horde/Horde.Build/Agents/Fleet/LeaseUtilizationStrategy.cs
--- a/Engine/Source/Programs/Horde/Horde.Build/Agents/Fleet/LeaseUtilizationStrategy.cs
+++ b/Engine/Source/Programs/Horde/Horde.Build/Agents/Fleet/LeaseUtilizationStrategy.cs
@@ -11,6 +11,7 @@ using Horde.Build.Agents.Pools;
 using Horde.Build.Utilities;
 using HordeCommon;
 using HordeCommon.Rpc.Tasks;
+using Microsoft.Extensions.Caching.Memory;
 using OpenTracing;
 using OpenTracing.Util;
 
@@ -30,6 +31,8 @@ namespace Horde.Build.Agents.Fleet
 	/// </summary>
 	public class LeaseUtilizationStrategy : IPoolSizeStrategy
 	{
+		private const string CacheKey = nameof(LeaseUtilizationStrategy);
+		
 		struct UtilizationSample
 		{
 			public double _jobWork;
@@ -90,6 +93,7 @@ namespace Horde.Build.Agents.Fleet
 		private readonly IPoolCollection _poolCollection;
 		private readonly ILeaseCollection _leaseCollection;
 		private readonly IClock _clock;
+		private readonly IMemoryCache _cache;
 		private readonly TimeSpan _sampleTime = TimeSpan.FromMinutes(6.0);
 		private readonly int _numSamples = 10;
 		private readonly int _numSamplesForResult = 9;
@@ -101,15 +105,17 @@ namespace Horde.Build.Agents.Fleet
 		/// <param name="poolCollection"></param>
 		/// <param name="leaseCollection"></param>
 		/// <param name="clock"></param>
+		/// <param name="cache"></param>
 		/// <param name="numSamples">Number of samples to collect for calculating lease utilization</param>
 		/// <param name="numSamplesForResult">Min number of samples for a valid result</param>
 		/// <param name="sampleTime">Time period for each sample</param>
-		public LeaseUtilizationStrategy(IAgentCollection agentCollection, IPoolCollection poolCollection, ILeaseCollection leaseCollection, IClock clock, int numSamples = 10, int numSamplesForResult = 9, TimeSpan? sampleTime = null)
+		public LeaseUtilizationStrategy(IAgentCollection agentCollection, IPoolCollection poolCollection, ILeaseCollection leaseCollection, IClock clock, IMemoryCache cache, int numSamples = 10, int numSamplesForResult = 9, TimeSpan? sampleTime = null)
 		{
 			_agentCollection = agentCollection;
 			_poolCollection = poolCollection;
 			_leaseCollection = leaseCollection;
 			_clock = clock;
+			_cache = cache;
 			_numSamples = numSamples;
 			_numSamplesForResult = numSamplesForResult;
 			_sampleTime = sampleTime ?? _sampleTime;
@@ -195,7 +201,16 @@ namespace Horde.Build.Agents.Fleet
 		/// <inheritdoc/>
 		public async Task<List<PoolSizeData>> CalcDesiredPoolSizesAsync(List<PoolSizeData> pools)
 		{
-			Dictionary<PoolId, PoolData> poolToData = await GetPoolDataAsync();
+			Dictionary<PoolId, PoolData> poolToData;
+			
+			// Cache pool data for a short while for faster runs when many pools are scaled
+			if (!_cache.TryGetValue(CacheKey, out poolToData))
+			{
+				// Pool sizes haven't been cached, update them (might happen from multiple tasks but that is fine)
+				poolToData = await GetPoolDataAsync();
+				_cache.Set(CacheKey, poolToData, TimeSpan.FromSeconds(60));
+			}
+			
 			List<PoolSizeData> result = new();
 
 			foreach (PoolData poolData in poolToData.Values.OrderByDescending(x => x.Agents.Count))
