Horde: Add a basic controller for reading/writing primitives to the storage service.

diff --git a/Engine/Source/Programs/Horde/Horde.Build/Commands/Bundles/PerforceCommand.cs b/Engine/Source/Programs/Horde/Horde.Build/Commands/Bundles/PerforceCommand.cs
--- a/Engine/Source/Programs/Horde/Horde.Build/Commands/Bundles/PerforceCommand.cs
+++ b/Engine/Source/Programs/Horde/Horde.Build/Commands/Bundles/PerforceCommand.cs
@@ -68,7 +68,7 @@ namespace Horde.Build.Commands.Bundles
 			IStreamCollection streamCollection = serviceProvider.GetRequiredService<IStreamCollection>();
 			StorageService storageService = serviceProvider.GetRequiredService<StorageService>();
 
-			IStorageClient storage = await storageService.GetClientAsync(Namespace.Perforce);
+			IStorageClient storage = await storageService.GetClientAsync(Namespace.Perforce, CancellationToken.None);
 
 			IStream? stream = await streamCollection.GetAsync(new StreamId(StreamId));
 			if (stream == null)
diff --git a/Engine/Source/Programs/Horde/Horde.Build/Perforce/ReplicationService.cs b/Engine/Source/Programs/Horde/Horde.Build/Perforce/ReplicationService.cs
--- a/Engine/Source/Programs/Horde/Horde.Build/Perforce/ReplicationService.cs
+++ b/Engine/Source/Programs/Horde/Horde.Build/Perforce/ReplicationService.cs
@@ -364,7 +364,7 @@ namespace Horde.Build.Perforce
 				}
 				else
 				{
-					IStorageClient storage = await _storageService.GetClientAsync(Namespace.Perforce);
+					IStorageClient storage = await _storageService.GetClientAsync(Namespace.Perforce, cancellationToken);
 
 					// Update the stream, updating the reservation every 30 seconds
 					Task internalTask = Task.Run(() => UpdateStreamContentInternalAsync(storage, stream, streamChanges, cancellationToken), cancellationToken);
diff --git a/Engine/Source/Programs/Horde/Horde.Build/Storage/StorageController.cs b/Engine/Source/Programs/Horde/Horde.Build/Storage/StorageController.cs
new file mode 100644
--- /dev/null
+++ b/Engine/Source/Programs/Horde/Horde.Build/Storage/StorageController.cs
@@ -0,0 +1,184 @@
+// Copyright Epic Games, Inc. All Rights Reserved.
+
+using System.IO;
+using System.Threading;
+using System.Threading.Tasks;
+using Amazon.EC2.Model;
+using EpicGames.Core;
+using EpicGames.Horde.Storage;
+using Horde.Build.Acls;
+using Horde.Build.Utilities;
+using Microsoft.AspNetCore.Authorization;
+using Microsoft.AspNetCore.Http;
+using Microsoft.AspNetCore.Mvc;
+
+namespace Horde.Build.Storage
+{
+	/// <summary>
+	/// Response from uploading a bundle
+	/// </summary>
+	public class WriteBundleResponse
+	{
+		/// <summary>
+		/// Locator for the uploaded bundle
+		/// </summary>
+		public BlobLocator Locator { get; set; }
+	}
+
+	/// <summary>
+	/// Request object for writing a ref
+	/// </summary>
+	public class WriteRefRequest
+	{
+		/// <summary>
+		/// Locator for the target blob
+		/// </summary>
+		public BlobLocator Locator { get; set; }
+
+		/// <summary>
+		/// Export index for the ref
+		/// </summary>
+		public int ExportIdx { get; set; }
+	}
+
+	/// <summary>
+	/// Response object for reading a ref
+	/// </summary>
+	public class ReadRefResponse
+	{
+		/// <summary>
+		/// Locator for the target blob
+		/// </summary>
+		public BlobLocator Locator { get; set; }
+
+		/// <summary>
+		/// Export index for the ref
+		/// </summary>
+		public int ExportIdx { get; set; }
+
+		/// <summary>
+		/// Constructor
+		/// </summary>
+		public ReadRefResponse(RefTarget target)
+		{
+			Locator = target.Locator;
+			ExportIdx = target.ExportIdx;
+		}
+	}
+
+	/// <summary>
+	/// Controller for the /api/v1/storage endpoint
+	/// </summary>
+	[Authorize]
+	[ApiController]
+	[Route("[controller]")]
+	public class StorageController : HordeControllerBase
+	{
+		readonly StorageService _storageService;
+
+		/// <summary>
+		/// Constructor
+		/// </summary>
+		/// <param name="storageService"></param>
+		public StorageController(StorageService storageService)
+		{
+			_storageService = storageService;
+		}
+
+		/// <summary>
+		/// Uploads data to the storage service. 
+		/// </summary>
+		/// <param name="namespaceId">Namespace to fetch from</param>
+		/// <param name="file">Data to be uploaded</param>
+		/// <param name="cancellationToken">Cancellation token for the operation</param>
+		[HttpPost]
+		[Route("/api/v1/storage/{namespaceId}/bundles")]
+		public async Task<ActionResult<WriteBundleResponse>> WriteBundleAsync(NamespaceId namespaceId, IFormFile file, CancellationToken cancellationToken)
+		{
+			if (!await _storageService.AuthorizeAsync(namespaceId, User, AclAction.WriteBlobs, null, cancellationToken))
+			{
+				return Forbid(AclAction.WriteBlobs);
+			}
+
+			IStorageClient client = await _storageService.GetClientAsync(namespaceId, cancellationToken);
+			using (Stream stream = file.OpenReadStream())
+			{
+				Bundle bundle = await Bundle.FromStreamAsync(stream, cancellationToken);
+				BlobLocator locator = await client.WriteBundleAsync(bundle, cancellationToken: cancellationToken);
+				return new WriteBundleResponse { Locator = locator };
+			}
+		}
+
+		/// <summary>
+		/// Retrieves data from the storage service. 
+		/// </summary>
+		/// <param name="namespaceId">Namespace to fetch from</param>
+		/// <param name="locator">Bundle to retrieve</param>
+		/// <param name="cancellationToken">Cancellation token for the operation</param>
+		[HttpGet]
+		[Route("/api/v1/storage/{namespaceId}/bundles/{locator}")]
+		public async Task<ActionResult> ReadBundleAsync(NamespaceId namespaceId, BlobLocator locator, CancellationToken cancellationToken)
+		{
+			if (!await _storageService.AuthorizeAsync(namespaceId, User, AclAction.ReadBlobs, null, cancellationToken))
+			{
+				return Forbid(AclAction.WriteBlobs);
+			}
+
+			IStorageClient client = await _storageService.GetClientAsync(namespaceId, cancellationToken);
+			Bundle bundle = await client.ReadBundleAsync(locator, cancellationToken);
+
+#pragma warning disable CA2000 // Dispose objects before losing scope
+			return File(new ReadOnlySequenceStream(bundle.AsSequence()), "application/x-ue-bundle");
+#pragma warning restore CA2000 // Dispose objects before losing scope
+		}
+
+		/// <summary>
+		/// Writes a ref to the storage service.
+		/// </summary>
+		/// <param name="namespaceId">Namespace to write to</param>
+		/// <param name="refName">Name of the ref</param>
+		/// <param name="request">Request for the ref to write</param>
+		/// <param name="cancellationToken">Cancellation token for the operation</param>
+		[HttpPut]
+		[Route("/api/v1/storage/{namespaceId}/refs/{refName}")]
+		public async Task<ActionResult> WriteRefAsync(NamespaceId namespaceId, RefName refName, [FromBody] WriteRefRequest request, CancellationToken cancellationToken)
+		{
+			if (!await _storageService.AuthorizeAsync(namespaceId, User, AclAction.WriteRefs, null, cancellationToken))
+			{
+				return Forbid(AclAction.WriteBlobs);
+			}
+
+			IStorageClient client = await _storageService.GetClientAsync(namespaceId, cancellationToken);
+			RefTarget target = new RefTarget(request.Locator, request.ExportIdx);
+			await client.WriteRefTargetAsync(refName, target, cancellationToken);
+
+			return Ok();
+		}
+
+		/// <summary>
+		/// Uploads data to the storage service. 
+		/// </summary>
+		/// <param name="namespaceId"></param>
+		/// <param name="refName"></param>
+		/// <param name="cancellationToken"></param>
+		[HttpGet]
+		[Route("/api/v1/storage/{namespaceId}/refs/{refName}")]
+		public async Task<ActionResult<ReadRefResponse>> ReadRefAsync(NamespaceId namespaceId, RefName refName, CancellationToken cancellationToken)
+		{
+			if (!await _storageService.AuthorizeAsync(namespaceId, User, AclAction.ReadRefs, null, cancellationToken))
+			{
+				return Forbid(AclAction.ReadRefs);
+			}
+
+			IStorageClient client = await _storageService.GetClientAsync(namespaceId, cancellationToken);
+
+			RefTarget? target = await client.TryReadRefTargetAsync(refName, cancellationToken: cancellationToken);
+			if (target == null)
+			{
+				return NotFound();
+			}
+
+			return new ReadRefResponse(target);
+		}
+	}
+}
diff --git a/Engine/Source/Programs/Horde/Horde.Build/Storage/StorageService.cs b/Engine/Source/Programs/Horde/Horde.Build/Storage/StorageService.cs
--- a/Engine/Source/Programs/Horde/Horde.Build/Storage/StorageService.cs
+++ b/Engine/Source/Programs/Horde/Horde.Build/Storage/StorageService.cs
@@ -187,9 +187,9 @@ namespace Horde.Build.Storage
 			}
 		}
 
-		async ValueTask<NamespaceInfo> GetNamespaceInfo(NamespaceId namespaceId)
+		async ValueTask<NamespaceInfo> GetNamespaceInfoAsync(NamespaceId namespaceId, CancellationToken cancellationToken)
 		{
-			State state = await _cachedState.GetAsync();
+			State state = await _cachedState.GetAsync(cancellationToken);
 			if (!state.Namespaces.TryGetValue(namespaceId, out NamespaceInfo? namespaceInfo))
 			{
 				throw new StorageException($"No namespace '{namespaceId}' is configured.");
@@ -201,10 +201,11 @@ namespace Horde.Build.Storage
 		/// Gets a storage client for the given namespace
 		/// </summary>
 		/// <param name="namespaceId">Namespace identifier</param>
+		/// <param name="cancellationToken">Cancellation token for the operation</param>
 		/// <returns></returns>
-		public async ValueTask<IStorageClient> GetClientAsync(NamespaceId namespaceId)
+		public async ValueTask<IStorageClient> GetClientAsync(NamespaceId namespaceId, CancellationToken cancellationToken)
 		{
-			NamespaceInfo namespaceInfo = await GetNamespaceInfo(namespaceId);
+			NamespaceInfo namespaceInfo = await GetNamespaceInfoAsync(namespaceId, cancellationToken);
 			return namespaceInfo.Client;
 		}
 
@@ -215,10 +216,11 @@ namespace Horde.Build.Storage
 		/// <param name="user"></param>
 		/// <param name="action"></param>
 		/// <param name="cache">Global permissions cache</param>
+		/// <param name="cancellationToken">Cancellation token for the operation</param>
 		/// <returns></returns>
-		public async Task<bool> AuthorizeAsync(NamespaceId namespaceId, ClaimsPrincipal user, AclAction action, GlobalPermissionsCache? cache)
+		public async Task<bool> AuthorizeAsync(NamespaceId namespaceId, ClaimsPrincipal user, AclAction action, GlobalPermissionsCache? cache, CancellationToken cancellationToken)
 		{
-			NamespaceInfo namespaceInfo = await GetNamespaceInfo(namespaceId);
+			NamespaceInfo namespaceInfo = await GetNamespaceInfoAsync(namespaceId, cancellationToken);
 
 			bool? result = namespaceInfo.Acl?.Authorize(action, user);
 			if (result == null)
diff --git a/Engine/Source/Programs/Horde/Horde.Build/Utilities/AsyncCachedValue.cs b/Engine/Source/Programs/Horde/Horde.Build/Utilities/AsyncCachedValue.cs
--- a/Engine/Source/Programs/Horde/Horde.Build/Utilities/AsyncCachedValue.cs
+++ b/Engine/Source/Programs/Horde/Horde.Build/Utilities/AsyncCachedValue.cs
@@ -77,40 +77,36 @@ namespace Horde.Build.Utilities
 		/// <summary>
 		/// Tries to get the current value
 		/// </summary>
+		/// <param name="cancellationToken">Cancellation token for the request</param>
 		/// <returns>The cached value, if valid</returns>
-		public Task<T> GetAsync()
+		public Task<T> GetAsync(CancellationToken cancellationToken = default)
 		{
-			return GetAsync(_maxRefreshTime);
+			return GetAsync(_maxRefreshTime, cancellationToken);
 		}
 
-		Task<State> CreateOrGetStateTask(ref Task<State>? stateTask)
+		/// <summary>
+		/// Tries to get the current value
+		/// </summary>
+		/// <returns>The cached value, if valid</returns>
+		public Task<T> GetAsync(TimeSpan maxAge, CancellationToken cancellationToken = default)
 		{
-			for(; ;)
+			Task<T> task = GetInternalAsync(maxAge);
+			if (cancellationToken.CanBeCanceled)
 			{
-				Task<State>? currentStateTask = stateTask;
-				if (currentStateTask != null)
-				{
-					return currentStateTask;
-				}
-
-				Task<Task<State>> innerNewStateTask = new Task<Task<State>>(() => CreateState());
-				if (Interlocked.CompareExchange(ref stateTask, innerNewStateTask.Unwrap(), null) == null)
-				{
-					innerNewStateTask.Start();
-				}
+				// The returned task object is shared, so we don't want to cancel computation of it; just the wait for a result.
+				task = WrapCancellation(task, cancellationToken);
 			}
+			return task;
 		}
 
-		async Task<State> CreateState()
+		static async Task<T> WrapCancellation(Task<T> task, CancellationToken cancellationToken)
 		{
-			return new State(await _generator());
+			await Task.WhenAny(task, Task.Delay(-1, cancellationToken));
+			cancellationToken.ThrowIfCancellationRequested();
+			return await task;
 		}
 
-		/// <summary>
-		/// Tries to get the current value
-		/// </summary>
-		/// <returns>The cached value, if valid</returns>
-		public async Task<T> GetAsync(TimeSpan maxAge)
+		async Task<T> GetInternalAsync(TimeSpan maxAge)
 		{
 			Task<State> stateTask = CreateOrGetStateTask(ref _current);
 
@@ -131,5 +127,28 @@ namespace Horde.Build.Utilities
 
 			return state.Value;
 		}
+
+		Task<State> CreateOrGetStateTask(ref Task<State>? stateTask)
+		{
+			for(; ;)
+			{
+				Task<State>? currentStateTask = stateTask;
+				if (currentStateTask != null)
+				{
+					return currentStateTask;
+				}
+
+				Task<Task<State>> innerNewStateTask = new Task<Task<State>>(() => CreateState());
+				if (Interlocked.CompareExchange(ref stateTask, innerNewStateTask.Unwrap(), null) == null)
+				{
+					innerNewStateTask.Start();
+				}
+			}
+		}
+
+		async Task<State> CreateState()
+		{
+			return new State(await _generator());
+		}
 	}
 }
diff --git a/Engine/Source/Programs/Shared/EpicGames.Horde/Storage/Bundle.cs b/Engine/Source/Programs/Shared/EpicGames.Horde/Storage/Bundle.cs
--- a/Engine/Source/Programs/Shared/EpicGames.Horde/Storage/Bundle.cs
+++ b/Engine/Source/Programs/Shared/EpicGames.Horde/Storage/Bundle.cs
@@ -12,6 +12,7 @@ using System.Reflection.Metadata;
 using System.Text;
 using System.Threading;
 using System.Threading.Tasks;
+using static System.Net.WebRequestMethods;
 
 namespace EpicGames.Horde.Storage
 {
@@ -73,6 +74,21 @@ namespace EpicGames.Horde.Storage
 			Payload = new ReadOnlySequence<byte>(reader.ReadFixedLengthBytes(length));
 		}
 
+		/// <summary>
+		/// Reads a bundle from the given stream
+		/// </summary>
+		/// <param name="stream">Stream to read from</param>
+		/// <param name="cancellationToken">Cancellation token for the operation</param>
+		/// <returns>Bundle that was read</returns>
+		public static async Task<Bundle> FromStreamAsync(Stream stream, CancellationToken cancellationToken)
+		{
+			using (MemoryStream memoryStream = new MemoryStream())
+			{
+				await stream.CopyToAsync(memoryStream, cancellationToken);
+				return new Bundle(new MemoryReader(memoryStream.ToArray()));
+			}
+		}
+
 		/// <summary>
 		/// Serializes the bundle to a sequence of bytes
 		/// </summary>
