Horde: Add streaming methods to IStorageClient (in the form of direct blob read/write methods), so we can download bundles without having to parse them.

diff --git a/Engine/Source/Programs/Horde/Horde.Agent/Commands/Bundles/CreateCommand.cs b/Engine/Source/Programs/Horde/Horde.Agent/Commands/Bundles/CreateCommand.cs
--- a/Engine/Source/Programs/Horde/Horde.Agent/Commands/Bundles/CreateCommand.cs
+++ b/Engine/Source/Programs/Horde/Horde.Agent/Commands/Bundles/CreateCommand.cs
@@ -1,7 +1,6 @@
 // Copyright Epic Games, Inc. All Rights Reserved.
 
 using System;
-using System.Buffers;
 using System.Globalization;
 using System.IO;
 using System.Threading;
@@ -32,34 +31,27 @@ namespace Horde.Agent.Commands.Bundles
 			FileReference GetRefFile(RefName name) => FileReference.Combine(_rootDir, name.ToString() + ".ref");
 			FileReference GetBlobFile(BlobLocator id) => FileReference.Combine(_rootDir, id.Inner.ToString() + ".blob");
 
-			static async Task WriteAsync(FileReference file, ReadOnlySequence<byte> sequence, CancellationToken cancellationToken)
-			{
-				using (FileStream stream = FileReference.Open(file, FileMode.Create, FileAccess.Write, FileShare.ReadWrite))
-				{
-					foreach (ReadOnlyMemory<byte> memory in sequence)
-					{
-						await stream.WriteAsync(memory, cancellationToken);
-					}
-				}
-			}
-
 			#region Blobs
 
-			public override async Task<Bundle> ReadBundleAsync(BlobLocator id, CancellationToken cancellationToken = default)
+			public override Task<Stream> ReadBlobAsync(BlobLocator id, CancellationToken cancellationToken = default)
 			{
 				FileReference file = GetBlobFile(id);
 				_logger.LogInformation("Reading {File}", file);
-				byte[] bytes = await FileReference.ReadAllBytesAsync(file, cancellationToken);
-				return new Bundle(new MemoryReader(bytes));
+				return Task.FromResult<Stream>(FileReference.Open(file, FileMode.Open, FileAccess.Read));
 			}
 
-			public override async Task<BlobLocator> WriteBundleAsync(Bundle bundle, Utf8String prefix = default, CancellationToken cancellationToken = default)
+			public override async Task<BlobLocator> WriteBlobAsync(Stream stream, Utf8String prefix = default, CancellationToken cancellationToken = default)
 			{
 				BlobLocator id = BlobLocator.Create(HostId.Empty, prefix);
 				FileReference file = GetBlobFile(id);
 				DirectoryReference.CreateDirectory(file.Directory);
 				_logger.LogInformation("Writing {File}", file);
-				await WriteAsync(file, bundle.AsSequence(), cancellationToken);
+
+				using (FileStream fileStream = FileReference.Open(file, FileMode.Create, FileAccess.Write, FileShare.ReadWrite))
+				{
+					await stream.CopyToAsync(fileStream, cancellationToken);
+				}
+
 				return id;
 			}
 
diff --git a/Engine/Source/Programs/Horde/Horde.Build/Storage/StorageController.cs b/Engine/Source/Programs/Horde/Horde.Build/Storage/StorageController.cs
--- a/Engine/Source/Programs/Horde/Horde.Build/Storage/StorageController.cs
+++ b/Engine/Source/Programs/Horde/Horde.Build/Storage/StorageController.cs
@@ -103,8 +103,7 @@ namespace Horde.Build.Storage
 			IStorageClient client = await _storageService.GetClientAsync(namespaceId, cancellationToken);
 			using (Stream stream = file.OpenReadStream())
 			{
-				Bundle bundle = await Bundle.FromStreamAsync(stream, cancellationToken);
-				BlobLocator locator = await client.WriteBundleAsync(bundle, cancellationToken: cancellationToken);
+				BlobLocator locator = await client.WriteBlobAsync(stream, cancellationToken: cancellationToken);
 				return new WriteBundleResponse { Locator = locator };
 			}
 		}
@@ -125,10 +124,10 @@ namespace Horde.Build.Storage
 			}
 
 			IStorageClient client = await _storageService.GetClientAsync(namespaceId, cancellationToken);
-			Bundle bundle = await client.ReadBundleAsync(locator, cancellationToken);
+			Stream stream = await client.ReadBlobAsync(locator, cancellationToken);
 
 #pragma warning disable CA2000 // Dispose objects before losing scope
-			return File(new ReadOnlySequenceStream(bundle.AsSequence()), "application/x-ue-bundle");
+			return File(stream, "application/octet-stream");
 #pragma warning restore CA2000 // Dispose objects before losing scope
 		}
 
diff --git a/Engine/Source/Programs/Horde/Horde.Build/Storage/StorageService.cs b/Engine/Source/Programs/Horde/Horde.Build/Storage/StorageService.cs
--- a/Engine/Source/Programs/Horde/Horde.Build/Storage/StorageService.cs
+++ b/Engine/Source/Programs/Horde/Horde.Build/Storage/StorageService.cs
@@ -4,6 +4,7 @@ using System;
 using System.Buffers;
 using System.Collections;
 using System.Collections.Generic;
+using System.IO;
 using System.Reflection;
 using System.Security.Claims;
 using System.Text.Json;
@@ -72,30 +73,26 @@ namespace Horde.Build.Storage
 
 			#region Blobs
 
-			public override async Task<Bundle> ReadBundleAsync(BlobLocator locator, CancellationToken cancellationToken = default)
+			public override async Task<Stream> ReadBlobAsync(BlobLocator locator, CancellationToken cancellationToken = default)
 			{
 				string path = GetBlobPath(locator.BlobId);
 
-				ReadOnlyMemory<byte>? data = await _backend.ReadBytesAsync(path, cancellationToken);
-				if (data == null)
+				Stream? stream = await _backend.ReadAsync(path, cancellationToken);
+				if (stream == null)
 				{
 					throw new StorageException($"Unable to read data from {path}");
 				}
 
-				return new Bundle(new MemoryReader(data.Value));
+				return stream;
 			}
 
 			/// <inheritdoc/>
-			public override async Task<BlobLocator> WriteBundleAsync(Bundle bundle, Utf8String prefix = default, CancellationToken cancellationToken = default)
+			public override async Task<BlobLocator> WriteBlobAsync(Stream stream, Utf8String prefix = default, CancellationToken cancellationToken = default)
 			{
 				BlobId id = BlobId.CreateNew(prefix);
-				string path = GetBlobPath(id);
 
-				ReadOnlySequence<byte> sequence = bundle.AsSequence();
-				using (ReadOnlySequenceStream stream = new ReadOnlySequenceStream(sequence))
-				{
-					await _backend.WriteAsync(path, stream, cancellationToken);
-				}
+				string path = GetBlobPath(id);
+				await _backend.WriteAsync(path, stream, cancellationToken);
 
 				return new BlobLocator(HostId.Empty, id);
 			}
diff --git a/Engine/Source/Programs/Shared/EpicGames.Horde/Storage/Backends/InMemoryBlobStore.cs b/Engine/Source/Programs/Shared/EpicGames.Horde/Storage/Backends/InMemoryBlobStore.cs
--- a/Engine/Source/Programs/Shared/EpicGames.Horde/Storage/Backends/InMemoryBlobStore.cs
+++ b/Engine/Source/Programs/Shared/EpicGames.Horde/Storage/Backends/InMemoryBlobStore.cs
@@ -4,6 +4,7 @@ using System;
 using System.Buffers;
 using System.Collections.Concurrent;
 using System.Collections.Generic;
+using System.IO;
 using System.Threading;
 using System.Threading.Tasks;
 using EpicGames.Core;
@@ -42,17 +43,18 @@ namespace EpicGames.Horde.Storage.Backends
 		#region Blobs
 
 		/// <inheritdoc/>
-		public override Task<Bundle> ReadBundleAsync(BlobLocator locator, CancellationToken cancellationToken = default)
+		public override Task<Stream> ReadBlobAsync(BlobLocator locator, CancellationToken cancellationToken = default)
 		{
-			return Task.FromResult(_blobs[locator]);
+			Stream stream = new ReadOnlySequenceStream(_blobs[locator].AsSequence());
+			return Task.FromResult(stream);
 		}
 
 		/// <inheritdoc/>
-		public override Task<BlobLocator> WriteBundleAsync(Bundle bundle, Utf8String prefix = default, CancellationToken cancellationToken = default)
+		public override async Task<BlobLocator> WriteBlobAsync(Stream stream, Utf8String prefix = default, CancellationToken cancellationToken = default)
 		{
 			BlobLocator locator = BlobLocator.Create(HostId.Empty, prefix);
-			_blobs[locator] = bundle;
-			return Task.FromResult(locator);
+			_blobs[locator] = await Bundle.FromStreamAsync(stream, cancellationToken);
+			return locator;
 		}
 
 		#endregion
diff --git a/Engine/Source/Programs/Shared/EpicGames.Horde/Storage/IStorageClient.cs b/Engine/Source/Programs/Shared/EpicGames.Horde/Storage/IStorageClient.cs
--- a/Engine/Source/Programs/Shared/EpicGames.Horde/Storage/IStorageClient.cs
+++ b/Engine/Source/Programs/Shared/EpicGames.Horde/Storage/IStorageClient.cs
@@ -4,6 +4,7 @@ using System;
 using System.Buffers;
 using System.Collections.Generic;
 using System.Diagnostics;
+using System.IO;
 using System.Threading;
 using System.Threading.Tasks;
 using EpicGames.Core;
@@ -91,6 +92,27 @@ namespace EpicGames.Horde.Storage
 	{
 		#region Blobs
 
+		/// <summary>
+		/// Reads data for a blob from the store
+		/// </summary>
+		/// <param name="locator">The blob identifier</param>
+		/// <param name="cancellationToken">Cancellation token for the operation</param>
+		/// <returns></returns>
+		Task<Stream> ReadBlobAsync(BlobLocator locator, CancellationToken cancellationToken = default);
+
+		/// <summary>
+		/// Writes a new blob to the store
+		/// </summary>
+		/// <param name="stream">Blob data</param>
+		/// <param name="cancellationToken">Cancellation token for the operation</param>
+		/// <param name="prefix">Prefix for blob names. While the returned BlobId is guaranteed to be unique, this name can be used as a prefix to aid debugging.</param>
+		/// <returns>Unique identifier for the blob</returns>
+		Task<BlobLocator> WriteBlobAsync(Stream stream, Utf8String prefix = default, CancellationToken cancellationToken = default);
+
+		#endregion
+
+		#region Bundles
+
 		/// <summary>
 		/// Reads data for a bundle from the store
 		/// </summary>
@@ -201,6 +223,16 @@ namespace EpicGames.Horde.Storage
 
 			#region Blobs
 
+			/// <inheritdoc/>
+			public Task<Stream> ReadBlobAsync(BlobLocator locator, CancellationToken cancellationToken = default) => _inner.ReadBlobAsync(locator, cancellationToken);
+
+			/// <inheritdoc/>
+			public Task<BlobLocator> WriteBlobAsync(Stream stream, Utf8String prefix = default, CancellationToken cancellationToken = default) => _inner.WriteBlobAsync(stream, prefix, cancellationToken);
+
+			#endregion
+
+			#region Bundles
+
 			/// <inheritdoc/>
 			public Task<Bundle> ReadBundleAsync(BlobLocator id, CancellationToken cancellationToken = default) => _inner.ReadBundleAsync(id, cancellationToken);
 
diff --git a/Engine/Source/Programs/Shared/EpicGames.Horde/Storage/StorageClientBase.cs b/Engine/Source/Programs/Shared/EpicGames.Horde/Storage/StorageClientBase.cs
--- a/Engine/Source/Programs/Shared/EpicGames.Horde/Storage/StorageClientBase.cs
+++ b/Engine/Source/Programs/Shared/EpicGames.Horde/Storage/StorageClientBase.cs
@@ -2,6 +2,7 @@
 
 using System;
 using System.Collections.Generic;
+using System.IO;
 using System.Linq;
 using System.Runtime.CompilerServices;
 using System.Text;
@@ -31,10 +32,28 @@ namespace EpicGames.Horde.Storage
 		#region Blobs
 
 		/// <inheritdoc/>
-		public abstract Task<Bundle> ReadBundleAsync(BlobLocator locator, CancellationToken cancellationToken = default);
+		public abstract Task<Stream> ReadBlobAsync(BlobLocator locator, CancellationToken cancellationToken = default);
 
 		/// <inheritdoc/>
-		public abstract Task<BlobLocator> WriteBundleAsync(Bundle bundle, Utf8String prefix = default, CancellationToken cancellationToken = default);
+		public abstract Task<BlobLocator> WriteBlobAsync(Stream stream, Utf8String prefix = default, CancellationToken cancellationToken = default);
+
+		#endregion
+
+		#region Bundles
+
+		/// <inheritdoc/>
+		public virtual async Task<Bundle> ReadBundleAsync(BlobLocator locator, CancellationToken cancellationToken = default)
+		{
+			using Stream stream = await ReadBlobAsync(locator, cancellationToken);
+			return await Bundle.FromStreamAsync(stream, cancellationToken);
+		}
+
+		/// <inheritdoc/>
+		public virtual async Task<BlobLocator> WriteBundleAsync(Bundle bundle, Utf8String prefix = default, CancellationToken cancellationToken = default)
+		{
+			using ReadOnlySequenceStream stream = new ReadOnlySequenceStream(bundle.AsSequence());
+			return await WriteBlobAsync(stream, prefix, cancellationToken);
+		}
 
 		#endregion
 
