Horde: Add flag for setting log level in agent

diff --git a/Engine/Source/Programs/Horde/Horde.Agent/Commands/Service/RunCommand.cs b/Engine/Source/Programs/Horde/Horde.Agent/Commands/Service/RunCommand.cs
--- a/Engine/Source/Programs/Horde/Horde.Agent/Commands/Service/RunCommand.cs
+++ b/Engine/Source/Programs/Horde/Horde.Agent/Commands/Service/RunCommand.cs
@@ -18,6 +18,7 @@ using Microsoft.Extensions.Logging;
 using OpenTracing;
 using OpenTracing.Util;
 using Polly;
+using Serilog.Events;
 
 namespace Horde.Agent.Modes.Service
 {
@@ -41,12 +42,28 @@ namespace Horde.Agent.Modes.Service
 		[CommandLine("-WorkingDir=")]
 		string? WorkingDir { get; set; } = null;
 
+		/// <summary>
+		/// Log verbosity level (use normal Serilog levels such as debug, warning or info)
+		/// </summary>
+		[CommandLine("-LogLevel")]
+		public string LogLevelStr { get; set; } = "information";
+		
 		/// <summary>
 		/// Runs the service indefinitely
 		/// </summary>
 		/// <returns>Exit code</returns>
 		public override async Task<int> ExecuteAsync(ILogger logger)
 		{
+			if (Enum.TryParse(LogLevelStr, true, out LogEventLevel logEventLevel))
+			{
+				Logging.LogLevelSwitch.MinimumLevel = logEventLevel;
+			}
+			else
+			{
+				Console.WriteLine($"Unable to parse log level: {LogLevelStr}");
+				return 0;
+			}
+			
 			IHostBuilder hostBuilder = Host.CreateDefaultBuilder();
 
 			// Attempt to setup this process as a Windows service. A race condition inside Microsoft.Extensions.Hosting.WindowsServices.WindowsServiceHelpers.IsWindowsService
