Horde: Issues service now checks step outcome for recording next/previous success, it also now only considers step refs that update issues, added test case and fixed some which were incorrectly identifying last success improperly
#rnx

diff --git a/Engine/Source/Programs/Horde/Horde.Build.Tests/IssueServiceTests.cs b/Engine/Source/Programs/Horde/Horde.Build.Tests/IssueServiceTests.cs
--- a/Engine/Source/Programs/Horde/Horde.Build.Tests/IssueServiceTests.cs
+++ b/Engine/Source/Programs/Horde/Horde.Build.Tests/IssueServiceTests.cs
@@ -200,7 +200,7 @@ namespace Horde.Build.Tests
 			return node.Object;
 		}
 
-		public IJob CreateJob(StreamId streamId, int change, string name, IGraph graph, TimeSpan time = default, bool promoteByDefault = true)
+		public IJob CreateJob(StreamId streamId, int change, string name, IGraph graph, TimeSpan time = default, bool promoteByDefault = true, bool updateIssues = true)
 		{
 			JobId jobId = JobId.GenerateNewId();
 
@@ -244,6 +244,7 @@ namespace Horde.Build.Tests
 			job.SetupGet(x => x.ShowUgsBadges).Returns(promoteByDefault);
 			job.SetupGet(x => x.ShowUgsAlerts).Returns(promoteByDefault);
 			job.SetupGet(x => x.PromoteIssuesByDefault).Returns(promoteByDefault);
+			job.SetupGet(x => x.UpdateIssues).Returns(updateIssues);
 			job.SetupGet(x => x.NotificationChannel).Returns("#devtools-horde-slack-testing");
 			return job.Object;
 		}
@@ -252,11 +253,15 @@ namespace Horde.Build.Tests
 		{
 			IJobStepBatch batch = job.Batches[batchIdx];
 			IJobStep step = batch.Steps[stepIdx];
-			await IssueService.UpdateCompleteStep(job, _graph, batch.Id, step.Id);
+
+			if (job.UpdateIssues)
+			{
+				await IssueService.UpdateCompleteStep(job, _graph, batch.Id, step.Id);
+			}			
 
 			JobStepRefId jobStepRefId = new JobStepRefId(job.Id, batch.Id, step.Id);
 			string nodeName = _graph.Groups[batch.GroupIdx].Nodes[step.NodeIdx].Name;
-			await JobStepRefCollection.InsertOrReplaceAsync(jobStepRefId, "TestJob", nodeName, job.StreamId, job.TemplateId, job.Change, step.LogId, null, null, outcome, null, null, 0.0f, 0.0f, step.StartTimeUtc!.Value, step.StartTimeUtc);
+			await JobStepRefCollection.InsertOrReplaceAsync(jobStepRefId, "TestJob", nodeName, job.StreamId, job.TemplateId, job.Change, step.LogId, null, null, outcome, job.UpdateIssues, null, null, 0.0f, 0.0f, step.StartTimeUtc!.Value, step.StartTimeUtc);
 		}
 
 		async Task AddEvent(IJob job, int batchIdx, int stepIdx, object data, EventSeverity severity = EventSeverity.Error)
@@ -477,7 +482,7 @@ namespace Horde.Build.Tests
 			// Expected: No issues are created
 			{
 				IJob job = CreateJob(_mainStreamId, 105, "Test Build", _graph);
-				await UpdateCompleteStep(job, 0, 0, JobStepOutcome.Warnings);
+				await UpdateCompleteStep(job, 0, 0, JobStepOutcome.Success);
 
 				List<IIssue> issues = await IssueCollection.FindIssuesAsync();
 				Assert.AreEqual(0, issues.Count);
@@ -519,7 +524,7 @@ namespace Horde.Build.Tests
 			// Expected: No issues are created
 			{
 				IJob job = CreateJob(_mainStreamId, 105, "Test Build", _graph);
-				await UpdateCompleteStep(job, 0, 0, JobStepOutcome.Warnings);
+				await UpdateCompleteStep(job, 0, 0, JobStepOutcome.Success);
 
 				List<IIssue> issues = await IssueCollection.FindIssuesAsync();
 				Assert.AreEqual(0, issues.Count);
@@ -1435,7 +1440,7 @@ namespace Horde.Build.Tests
 				Assert.AreEqual("Hashed", issue.Fingerprints[0].Type);
 
 				IIssueSpan span = spans[0];
-				Assert.AreEqual(120, span.LastSuccess?.Change);
+				Assert.AreEqual(105, span.LastSuccess?.Change);
 				Assert.AreEqual(null, span.NextSuccess?.Change);
 			}
 		}
@@ -2057,6 +2062,81 @@ namespace Horde.Build.Tests
 
 		}
 
+		[TestMethod]
+		public async Task UpdateIssuesFlagTest()
+		{
+			int hour = 0;
+
+			// #1
+			// Scenario: Job is created that doesn't update issues at CL 225 with a successful outcome
+			// Expected: No new issue is created
+			{
+				IJob job = CreateJob(_mainStreamId, 225, "Test Build", _graph, TimeSpan.FromHours(hour++), true, false);
+				await UpdateCompleteStep(job, 0, 0, JobStepOutcome.Success);
+				List<IIssue> issues = await IssueCollection.FindIssuesAsync();
+				Assert.AreEqual(0, issues.Count);
+			}
+
+			// #2
+			// Scenario: Job is created that doesn't update issues at CL 226 with a failure outcome
+			// Expected: No new issue is created
+			{
+				IJob job = CreateJob(_mainStreamId, 226, "Test Build", _graph, TimeSpan.FromHours(hour++), true, false);
+				await AddEvent(job, 0, 0, new { level = nameof(LogLevel.Error) }, EventSeverity.Error);
+				await UpdateCompleteStep(job, 0, 0, JobStepOutcome.Failure);
+				List<IIssue> issues = await IssueCollection.FindIssuesAsync();
+				Assert.AreEqual(0, issues.Count);
+			}
+
+			// #3
+			// Scenario: Job is created at an earlier CL than in #1, with a warning outcome 
+			// Expected: Default issue is created
+			{
+				IJob job = CreateJob(_mainStreamId, 105, "Test Build", _graph, TimeSpan.FromHours(hour++));
+				await AddEvent(job, 0, 0, new { level = nameof(LogLevel.Warning) }, EventSeverity.Warning);
+				await UpdateCompleteStep(job, 0, 0, JobStepOutcome.Warnings);				
+
+				List<IIssue> issues = await IssueCollection.FindIssuesAsync();
+				Assert.AreEqual(1, issues.Count);
+				Assert.AreEqual(IssueSeverity.Warning, issues[0].Severity);
+
+				Assert.AreEqual("Warnings in Update Version Files", issues[0].Summary);
+				
+			}
+
+			// #4
+			// Scenario: Job is run which doesn't update issues, with a failure 
+			// Expected: Existing issue is not updated and remains a warning
+			{
+				IJob job = CreateJob(_mainStreamId, 225, "Test Build", _graph, TimeSpan.FromHours(hour++), true, false);
+				await AddEvent(job, 0, 0, new { level = nameof(LogLevel.Error) }, EventSeverity.Error);
+				await UpdateCompleteStep(job, 0, 0, JobStepOutcome.Failure);
+
+				List<IIssue> issues = await IssueCollection.FindIssuesAsync();
+				Assert.AreEqual(1, issues.Count);
+				Assert.AreEqual(1, issues[0].Id);
+				Assert.AreEqual(IssueSeverity.Warning, issues[0].Severity);
+
+			}
+
+			// #5
+			// Scenario: Job is run which updates issue, with a failure 
+			// Expected: Existing issue is updated and becomes an error
+			{
+				IJob job = CreateJob(_mainStreamId, 225, "Test Build", _graph, TimeSpan.FromHours(hour++));
+				await AddEvent(job, 0, 0, new { level = nameof(LogLevel.Error) }, EventSeverity.Error);
+				await UpdateCompleteStep(job, 0, 0, JobStepOutcome.Failure);
+
+				List<IIssue> issues = await IssueCollection.FindIssuesAsync();
+				Assert.AreEqual(1, issues.Count);
+				Assert.AreEqual(1, issues[0].Id);
+				Assert.AreEqual(IssueSeverity.Error, issues[0].Severity);
+			}
+
+
+		}
+
+
 
 		private async Task ParseAsync(LogId logId, string[] lines)
 		{
diff --git a/Engine/Source/Programs/Horde/Horde.Build.Tests/Stubs/Collections/JobStepRefCollectionStub.cs b/Engine/Source/Programs/Horde/Horde.Build.Tests/Stubs/Collections/JobStepRefCollectionStub.cs
--- a/Engine/Source/Programs/Horde/Horde.Build.Tests/Stubs/Collections/JobStepRefCollectionStub.cs
+++ b/Engine/Source/Programs/Horde/Horde.Build.Tests/Stubs/Collections/JobStepRefCollectionStub.cs
@@ -31,15 +31,17 @@ namespace Horde.Build.Tests.Stubs.Collections
 		public PoolId? PoolId { get; set; }
 		public AgentId? AgentId { get; set; }
 		public JobStepOutcome? Outcome { get; set; }
+		public bool UpdateIssues { get; set; } = false;
 		public int? LastSuccess { get; set; }
 		public int? LastWarning { get; set; }
 
+
 		public virtual float BatchWaitTime => throw new NotImplementedException();
 		public virtual float BatchInitTime => throw new NotImplementedException();
 		public virtual DateTime StartTimeUtc => throw new NotImplementedException();
 		public virtual DateTime? FinishTimeUtc => throw new NotImplementedException();
 
-		public JobStepRefStub(JobId jobId, SubResourceId batchId, SubResourceId stepId, string jobName, string nodeName, StreamId streamId, TemplateRefId templateId, int change, JobStepOutcome? outcome)
+		public JobStepRefStub(JobId jobId, SubResourceId batchId, SubResourceId stepId, string jobName, string nodeName, StreamId streamId, TemplateRefId templateId, int change, JobStepOutcome? outcome, bool updateIssues)
 		{
 			Id = new JobStepRefId(jobId, batchId, stepId);
 			JobName = jobName;
@@ -48,6 +50,7 @@ namespace Horde.Build.Tests.Stubs.Collections
 			TemplateId = templateId;
 			Change = change;
 			Outcome = outcome;
+			UpdateIssues = updateIssues;
 		}
 	}
 
@@ -60,12 +63,12 @@ namespace Horde.Build.Tests.Stubs.Collections
 			_refs.Add(jobStepRef);
 		}
 
-		public Task<IJobStepRef?> GetNextStepForNodeAsync(StreamId streamId, TemplateRefId templateId, string nodeName, int change)
+		public Task<IJobStepRef?> GetNextStepForNodeAsync(StreamId streamId, TemplateRefId templateId, string nodeName, int change, JobStepOutcome? outcome = null, bool? updateIssues = null)
 		{
 			IJobStepRef? nextRef = null;
 			foreach (IJobStepRef jobStepRef in _refs)
 			{
-				if (jobStepRef.StreamId == streamId && jobStepRef.TemplateId == templateId && jobStepRef.NodeName == nodeName && jobStepRef.Change > change)
+				if (jobStepRef.StreamId == streamId && jobStepRef.TemplateId == templateId && jobStepRef.NodeName == nodeName && jobStepRef.Change > change && (outcome == null ? jobStepRef.Outcome != null : jobStepRef.Outcome == outcome) && (updateIssues != null ? jobStepRef.UpdateIssues == updateIssues : true) )
 				{
 					if (nextRef == null || jobStepRef.Change < nextRef.Change)
 					{
@@ -76,12 +79,12 @@ namespace Horde.Build.Tests.Stubs.Collections
 			return Task.FromResult(nextRef);
 		}
 
-		public Task<IJobStepRef?> GetPrevStepForNodeAsync(StreamId streamId, TemplateRefId templateId, string nodeName, int change)
+		public Task<IJobStepRef?> GetPrevStepForNodeAsync(StreamId streamId, TemplateRefId templateId, string nodeName, int change, JobStepOutcome? outcome = null, bool? updateIssues = null)
 		{
 			IJobStepRef? prevRef = null;
 			foreach (IJobStepRef jobStepRef in _refs)
 			{
-				if (jobStepRef.StreamId == streamId && jobStepRef.TemplateId == templateId && jobStepRef.NodeName == nodeName && jobStepRef.Change < change)
+				if (jobStepRef.StreamId == streamId && jobStepRef.TemplateId == templateId && jobStepRef.NodeName == nodeName && jobStepRef.Change < change && (outcome == null ? jobStepRef.Outcome != null : jobStepRef.Outcome == outcome) && (updateIssues != null ? jobStepRef.UpdateIssues == updateIssues : true))
 				{
 					if (prevRef == null || jobStepRef.Change > prevRef.Change)
 					{
@@ -97,7 +100,7 @@ namespace Horde.Build.Tests.Stubs.Collections
 			throw new NotImplementedException();
 		}
 
-		Task<IJobStepRef> IJobStepRefCollection.InsertOrReplaceAsync(JobStepRefId id, string jobName, string nodeName, StreamId streamId, TemplateRefId templateId, int change, LogId? logId, PoolId? poolId, AgentId? agentId, JobStepOutcome? outcome, int? lastSuccess, int? lastWarning, float waitTime, float initTime, DateTime startTime, DateTime? finishTime)
+		Task<IJobStepRef> IJobStepRefCollection.InsertOrReplaceAsync(JobStepRefId id, string jobName, string nodeName, StreamId streamId, TemplateRefId templateId, int change, LogId? logId, PoolId? poolId, AgentId? agentId, JobStepOutcome? outcome, bool updateIssues, int? lastSuccess, int? lastWarning, float waitTime, float initTime, DateTime startTime, DateTime? finishTime)
 		{
 			throw new NotImplementedException();
 		}
diff --git a/Engine/Source/Programs/Horde/Horde.Build/Issues/IssueService.cs b/Engine/Source/Programs/Horde/Horde.Build/Issues/IssueService.cs
--- a/Engine/Source/Programs/Horde/Horde.Build/Issues/IssueService.cs
+++ b/Engine/Source/Programs/Horde/Horde.Build/Issues/IssueService.cs
@@ -811,14 +811,14 @@ namespace Horde.Build.Issues
 				// Get the span data
 				NewIssueSpanData spanData = new NewIssueSpanData(stream.Id, stream.Name, job.TemplateId, node.Name, eventGroup.Fingerprint, stepData);
 
-				IJobStepRef? prevJob = await _jobStepRefs.GetPrevStepForNodeAsync(job.StreamId, job.TemplateId, node.Name, job.Change);
+				IJobStepRef? prevJob = await _jobStepRefs.GetPrevStepForNodeAsync(job.StreamId, job.TemplateId, node.Name, job.Change, JobStepOutcome.Success, true);
 				if (prevJob != null)
 				{
 					spanData.LastSuccess = new NewIssueStepData(prevJob);
 					spanData.Suspects = await FindSuspectsForSpanAsync(stream, spanData.Fingerprint, spanData.LastSuccess.Change + 1, spanData.FirstFailure.Change);
 				}
 
-				IJobStepRef? nextJob = await _jobStepRefs.GetNextStepForNodeAsync(job.StreamId, job.TemplateId, node.Name, job.Change);
+				IJobStepRef? nextJob = await _jobStepRefs.GetNextStepForNodeAsync(job.StreamId, job.TemplateId, node.Name, job.Change, JobStepOutcome.Success, true);
 				if (nextJob != null)
 				{
 					spanData.NextSuccess = new NewIssueStepData(nextJob);
diff --git a/Engine/Source/Programs/Horde/Horde.Build/Jobs/IJobStepRef.cs b/Engine/Source/Programs/Horde/Horde.Build/Jobs/IJobStepRef.cs
--- a/Engine/Source/Programs/Horde/Horde.Build/Jobs/IJobStepRef.cs
+++ b/Engine/Source/Programs/Horde/Horde.Build/Jobs/IJobStepRef.cs
@@ -161,6 +161,11 @@ namespace Horde.Build.Jobs
 		/// </summary>
 		public JobStepOutcome? Outcome { get; }
 
+		/// <summary>
+		/// Whether this step should update issues
+		/// </summary>
+		public bool UpdateIssues { get; }
+
 		/// <summary>
 		/// The last change that succeeded. Note that this is only set when the ref is updated; it is not necessarily consistent with steps run later.
 		/// </summary>
diff --git a/Engine/Source/Programs/Horde/Horde.Build/Jobs/IJobStepRefCollection.cs b/Engine/Source/Programs/Horde/Horde.Build/Jobs/IJobStepRefCollection.cs
--- a/Engine/Source/Programs/Horde/Horde.Build/Jobs/IJobStepRefCollection.cs
+++ b/Engine/Source/Programs/Horde/Horde.Build/Jobs/IJobStepRefCollection.cs
@@ -36,13 +36,14 @@ namespace Horde.Build.Jobs
 		/// <param name="poolId">The pool id</param>
 		/// <param name="agentId">The agent id</param>
 		/// <param name="outcome">Outcome of this step, if known</param>
+		/// <param name="updateIssues">Whether this step ref is included for issue updates</param>
 		/// <param name="lastSuccess">The last change that completed with success</param>
 		/// <param name="lastWarning">The last change that completed with a warning (or success)</param>
 		/// <param name="waitTime">Time taken for the batch containing this step to start</param>
 		/// <param name="initTime">Time taken for the batch containing this step to initializer</param>
 		/// <param name="startTimeUtc">Start time</param>
 		/// <param name="finishTimeUtc">Finish time for the step, if known</param>
-		Task<IJobStepRef> InsertOrReplaceAsync(JobStepRefId id, string jobName, string stepName, StreamId streamId, TemplateId templateId, int change, LogId? logId, PoolId? poolId, AgentId? agentId, JobStepOutcome? outcome, int? lastSuccess, int? lastWarning, float waitTime, float initTime, DateTime startTimeUtc, DateTime? finishTimeUtc);
+		Task<IJobStepRef> InsertOrReplaceAsync(JobStepRefId id, string jobName, string stepName, StreamId streamId, TemplateId templateId, int change, LogId? logId, PoolId? poolId, AgentId? agentId, JobStepOutcome? outcome, bool updateIssues, int? lastSuccess, int? lastWarning, float waitTime, float initTime, DateTime startTimeUtc, DateTime? finishTimeUtc);
 
 		/// <summary>
 		/// Gets the history of a given node
@@ -63,8 +64,10 @@ namespace Horde.Build.Jobs
 		/// <param name="templateId">The template id</param>
 		/// <param name="nodeName">Name of the step to find</param>
 		/// <param name="change">The current changelist number</param>
+		/// <param name="outcome">The outcome to filter by or include all outcomes if null</param>
+		/// <param name="updateIssues">If true, constrain to steps which update issues</param>		 
 		/// <returns>The previous job, or null.</returns>
-		Task<IJobStepRef?> GetPrevStepForNodeAsync(StreamId streamId, TemplateId templateId, string nodeName, int change);
+		Task<IJobStepRef?> GetPrevStepForNodeAsync(StreamId streamId, TemplateId templateId, string nodeName, int change, JobStepOutcome? outcome = null, bool? updateIssues = null);
 
 		/// <summary>
 		/// Gets the next job that ran a given step
@@ -73,8 +76,10 @@ namespace Horde.Build.Jobs
 		/// <param name="templateId">The template id</param>
 		/// <param name="nodeName">Name of the step to find</param>
 		/// <param name="change">The current changelist number</param>
+		/// <param name="outcome">The outcome to filter by or include all outcomes if null</param>
+		/// <param name="updateIssues">If true, constrain to steps which update issues</param>
 		/// <returns>The previous job, or null.</returns>
-		Task<IJobStepRef?> GetNextStepForNodeAsync(StreamId streamId, TemplateId templateId, string nodeName, int change);
+		Task<IJobStepRef?> GetNextStepForNodeAsync(StreamId streamId, TemplateId templateId, string nodeName, int change, JobStepOutcome? outcome = null, bool? updateIssues = null);
 	}
 
 	static class JobStepRefCollectionExtensions
@@ -114,7 +119,7 @@ namespace Horde.Build.Jobs
 					outcome = JobStepOutcome.Unspecified;
 				}
 
-				await jobStepRefs.InsertOrReplaceAsync(new JobStepRefId(job.Id, batch.Id, step.Id), job.Name, nodeName, job.StreamId, job.TemplateId, job.Change, step.LogId, batch.PoolId, batch.AgentId, outcome, lastSuccess, lastWarning, waitTime, initTime, step.StartTimeUtc ?? DateTime.UtcNow, step.FinishTimeUtc);
+				await jobStepRefs.InsertOrReplaceAsync(new JobStepRefId(job.Id, batch.Id, step.Id), job.Name, nodeName, job.StreamId, job.TemplateId, job.Change, step.LogId, batch.PoolId, batch.AgentId, outcome, job.UpdateIssues, lastSuccess, lastWarning, waitTime, initTime, step.StartTimeUtc ?? DateTime.UtcNow, step.FinishTimeUtc);
 			}
 		}
 	}
diff --git a/Engine/Source/Programs/Horde/Horde.Build/Jobs/JobStepRefCollection.cs b/Engine/Source/Programs/Horde/Horde.Build/Jobs/JobStepRefCollection.cs
--- a/Engine/Source/Programs/Horde/Horde.Build/Jobs/JobStepRefCollection.cs
+++ b/Engine/Source/Programs/Horde/Horde.Build/Jobs/JobStepRefCollection.cs
@@ -64,11 +64,15 @@ namespace Horde.Build.Jobs
 			[BsonIgnoreIfNull]
 			public DateTime? FinishTimeUtc { get; set; }
 
+			[BsonIgnoreIfNull]
+			public bool? UpdateIssues { get; set; }
+
 			DateTime IJobStepRef.StartTimeUtc => StartTimeUtc ?? StartTime?.UtcDateTime ?? default;
 			DateTime? IJobStepRef.FinishTimeUtc => FinishTimeUtc ?? FinishTime?.UtcDateTime;
 			string IJobStepRef.NodeName => Name;
+			bool IJobStepRef.UpdateIssues => UpdateIssues ?? false;
 
-			public JobStepRef(JobStepRefId id, string jobName, string nodeName, StreamId streamId, TemplateId templateId, int change, LogId? logId, PoolId? poolId, AgentId? agentId, JobStepOutcome? outcome, int? lastSuccess, int? lastWarning, float batchWaitTime, float batchInitTime, DateTime startTimeUtc, DateTime? finishTimeUtc)
+			public JobStepRef(JobStepRefId id, string jobName, string nodeName, StreamId streamId, TemplateId templateId, int change, LogId? logId, PoolId? poolId, AgentId? agentId, JobStepOutcome? outcome, bool updateIssues, int? lastSuccess, int? lastWarning, float batchWaitTime, float batchInitTime, DateTime startTimeUtc, DateTime? finishTimeUtc)
 			{
 				Id = id;
 				JobName = jobName;
@@ -80,6 +84,7 @@ namespace Horde.Build.Jobs
 				PoolId = poolId;
 				AgentId = agentId;
 				Outcome = outcome;
+				UpdateIssues = updateIssues;
 				LastSuccess = lastSuccess;
 				LastWarning = lastWarning;
 				BatchWaitTime = batchWaitTime;
@@ -103,9 +108,9 @@ namespace Horde.Build.Jobs
 		}
 
 		/// <inheritdoc/>
-		public async Task<IJobStepRef> InsertOrReplaceAsync(JobStepRefId id, string jobName, string stepName, StreamId streamId, TemplateId templateId, int change, LogId? logId, PoolId? poolId, AgentId? agentId, JobStepOutcome? outcome, int? lastSuccess, int? lastWarning, float waitTime, float initTime, DateTime startTimeUtc, DateTime? finishTimeUtc)
+		public async Task<IJobStepRef> InsertOrReplaceAsync(JobStepRefId id, string jobName, string stepName, StreamId streamId, TemplateId templateId, int change, LogId? logId, PoolId? poolId, AgentId? agentId, JobStepOutcome? outcome, bool updateIssues, int? lastSuccess, int? lastWarning, float waitTime, float initTime, DateTime startTimeUtc, DateTime? finishTimeUtc)
 		{
-			JobStepRef newJobStepRef = new JobStepRef(id, jobName, stepName, streamId, templateId, change, logId, poolId, agentId, outcome, lastSuccess, lastWarning, waitTime, initTime, startTimeUtc, finishTimeUtc);
+			JobStepRef newJobStepRef = new JobStepRef(id, jobName, stepName, streamId, templateId, change, logId, poolId, agentId, outcome, updateIssues, lastSuccess, lastWarning, waitTime, initTime, startTimeUtc, finishTimeUtc);
 			await _jobStepRefs.ReplaceOneAsync(Builders<JobStepRef>.Filter.Eq(x => x.Id, newJobStepRef.Id), newJobStepRef, new ReplaceOptions { IsUpsert = true });
 			return newJobStepRef;
 		}
@@ -134,15 +139,59 @@ namespace Horde.Build.Jobs
 		}
 
 		/// <inheritdoc/>
-		public async Task<IJobStepRef?> GetPrevStepForNodeAsync(StreamId streamId, TemplateId templateId, string nodeName, int change)
+		public async Task<IJobStepRef?> GetPrevStepForNodeAsync(StreamId streamId, TemplateId templateId, string nodeName, int change, JobStepOutcome? outcome = null, bool? updateIssues = null)
 		{
-			return await _jobStepRefs.Find(x => x.StreamId == streamId && x.TemplateId == templateId && x.Name == nodeName && x.Change < change && x.Outcome != null).SortByDescending(x => x.Change).FirstOrDefaultAsync();
+			FilterDefinitionBuilder<JobStepRef> filterBuilder = Builders<JobStepRef>.Filter;
+
+			FilterDefinition<JobStepRef> filter = FilterDefinition<JobStepRef>.Empty;
+			filter &= filterBuilder.Eq(x => x.StreamId, streamId);
+			filter &= filterBuilder.Eq(x => x.TemplateId, templateId);
+			filter &= filterBuilder.Eq(x => x.Name, nodeName);
+			filter &= filterBuilder.Lt(x => x.Change, change);
+
+			if (outcome != null)
+			{
+				filter &= filterBuilder.Eq(x => x.Outcome, outcome);
+			}
+			else
+			{
+				filter &= filterBuilder.Ne(x => x.Outcome, null);
+			}
+
+			if (updateIssues != null)
+			{
+				filter &= filterBuilder.Eq(x => x.UpdateIssues, updateIssues);
+			}
+
+			return await _jobStepRefs.Find(filter).SortByDescending(x => x.Change).FirstOrDefaultAsync();
 		}
 
 		/// <inheritdoc/>
-		public async Task<IJobStepRef?> GetNextStepForNodeAsync(StreamId streamId, TemplateId templateId, string nodeName, int change)
+		public async Task<IJobStepRef?> GetNextStepForNodeAsync(StreamId streamId, TemplateId templateId, string nodeName, int change, JobStepOutcome? outcome = null, bool? updateIssues = null)
 		{
-			return await _jobStepRefs.Find(x => x.StreamId == streamId && x.TemplateId == templateId && x.Name == nodeName && x.Change > change && x.Outcome != null).SortBy(x => x.Change).FirstOrDefaultAsync();
+			FilterDefinitionBuilder<JobStepRef> filterBuilder = Builders<JobStepRef>.Filter;
+
+			FilterDefinition<JobStepRef> filter = FilterDefinition<JobStepRef>.Empty;
+			filter &= filterBuilder.Eq(x => x.StreamId, streamId);
+			filter &= filterBuilder.Eq(x => x.TemplateId, templateId);
+			filter &= filterBuilder.Eq(x => x.Name, nodeName);
+			filter &= filterBuilder.Gt(x => x.Change, change);
+			
+			if (outcome != null)
+			{
+				filter &= filterBuilder.Eq(x => x.Outcome, outcome);
+			}
+			else
+			{
+				filter &= filterBuilder.Ne(x => x.Outcome, null);
+			}
+
+			if (updateIssues != null)
+			{
+				filter &= filterBuilder.Eq(x => x.UpdateIssues, updateIssues);
+			}		
+
+			return await _jobStepRefs.Find(filter).SortBy(x => x.Change).FirstOrDefaultAsync();
 		}
 	}
 }
