Horde: Additional logging to track down cause of failures to submit.

diff --git a/Engine/Source/Programs/Horde/Horde.Build/Perforce/PerforceService.cs b/Engine/Source/Programs/Horde/Horde.Build/Perforce/PerforceService.cs
--- a/Engine/Source/Programs/Horde/Horde.Build/Perforce/PerforceService.cs
+++ b/Engine/Source/Programs/Horde/Horde.Build/Perforce/PerforceService.cs
@@ -1000,7 +1000,7 @@ namespace Horde.Build.Perforce
 					{
 						if (attempt < MaxAttempts)
 						{
-							_logger.LogWarning(ex, "Unable to submit new changelist (file: {File}, attempt: {Attempt}/{MaxAttempts}) - {Message}", depotPath, attempt, MaxAttempts, ex.Message);
+							_logger.LogWarning(ex, "Unable to submit new changelist (file: {File}, depotPath: {DepotPath}, description: \"{Description}\", attempt: {Attempt}/{MaxAttempts}) - {Message}", filePath, depotPath, description, attempt, MaxAttempts, ex.Message);
 						}
 						else
 						{
