Horde Storage - Fixed issue with peer latency test failing to resolve url

diff --git a/Engine/Source/Programs/Horde/HordeStorage/Horde.Storage/Implementation/PeerStatusService.cs b/Engine/Source/Programs/Horde/HordeStorage/Horde.Storage/Implementation/PeerStatusService.cs
--- a/Engine/Source/Programs/Horde/HordeStorage/Horde.Storage/Implementation/PeerStatusService.cs
+++ b/Engine/Source/Programs/Horde/HordeStorage/Horde.Storage/Implementation/PeerStatusService.cs
@@ -157,7 +157,7 @@ namespace Horde.Storage.Implementation
             // treat any connection that takes more then 5 seconds to establish as timed out
             client.Timeout = TimeSpan.FromSeconds(5);
 
-            Uri uri = new Uri(endpoint.Url + "/health/live");
+            Uri uri = new Uri(endpoint.Url + "health/live");
             try
             {
                 HttpResponseMessage result = await client.GetAsync(uri);
