Horde Storage - Added retry if db times out while listing all objects in a attempt to make it more likely the GC actually finshes runs in prod.

diff --git a/Engine/Source/Programs/Horde/HordeStorage/Horde.Storage/Implementation/Objects/ScyllaReferencesStore.cs b/Engine/Source/Programs/Horde/HordeStorage/Horde.Storage/Implementation/Objects/ScyllaReferencesStore.cs
--- a/Engine/Source/Programs/Horde/HordeStorage/Horde.Storage/Implementation/Objects/ScyllaReferencesStore.cs
+++ b/Engine/Source/Programs/Horde/HordeStorage/Horde.Storage/Implementation/Objects/ScyllaReferencesStore.cs
@@ -127,14 +127,33 @@ namespace Horde.Storage.Implementation
         {
             using IScope _ = Tracer.Instance.StartActive("scylla.get_records");
 
+            int retryAttempts = 0;
+            const int MaxRetryAttempts = 3;
             RowSet rowSet = await _session.ExecuteAsync(_getObjectsStatement.Bind(ns.ToString()));
             foreach (Row row in rowSet)
             {
-                if (rowSet.GetAvailableWithoutFetching() == 0)
+                try
                 {
-                    await rowSet.FetchMoreResultsAsync();
+                    if (rowSet.GetAvailableWithoutFetching() == 0)
+                    {
+                        await rowSet.FetchMoreResultsAsync();
+                    }
+                }
+                catch (ReadTimeoutException)
+                {
+                    if (retryAttempts < MaxRetryAttempts)
+                    {
+                        // wait 10 seconds and try again as the Db is under heavy load right now
+                        await Task.Delay(TimeSpan.FromSeconds(10));
+                        retryAttempts += 1;
+
+                        continue;
+                    }
+                    // we have failed to many times, rethrow the exception and abort to avoid stalling here for ever
+                    throw;
                 }
 
+                retryAttempts = 0;
                 string bucket = row.GetValue<string>("bucket");
                 string name = row.GetValue<string>("name");
                 DateTime? lastAccessTime = row.GetValue<DateTime?>("last_access_time");
