Horde Storage - Fixes for url parsing going wrong in status check tests

diff --git a/Engine/Source/Programs/Horde/HordeStorage/Horde.Storage/Implementation/PeerStatusService.cs b/Engine/Source/Programs/Horde/HordeStorage/Horde.Storage/Implementation/PeerStatusService.cs
--- a/Engine/Source/Programs/Horde/HordeStorage/Horde.Storage/Implementation/PeerStatusService.cs
+++ b/Engine/Source/Programs/Horde/HordeStorage/Horde.Storage/Implementation/PeerStatusService.cs
@@ -157,7 +157,13 @@ namespace Horde.Storage.Implementation
             // treat any connection that takes more then 5 seconds to establish as timed out
             client.Timeout = TimeSpan.FromSeconds(5);
 
-            Uri uri = new Uri(endpoint.Url + "health/live");
+            string url = endpoint.Url.ToString();
+            if (!url.EndsWith("/", StringComparison.OrdinalIgnoreCase))
+            {
+                url += "/";
+            }
+
+            Uri uri = new Uri(url + "health/live");
             try
             {
                 HttpResponseMessage result = await client.GetAsync(uri);
