Horde Storage - When attempting to fetch a blob we try to fetch it from the filesystem of other horde storage instances in the cluster before going to S3 as this is generally much faster. This means we only need to hit S3 if a blob has never before been seen in the region or has been deleted from all file caches (e.g. is really old).
This also introduces a optional list of stores to use when fetching a blob (or doing a exist call) which can be useful for debugging purposes to only check certain blob store layers.

diff --git a/Engine/Source/Programs/Horde/HordeStorage/Horde.Storage/Controllers/StorageController.cs b/Engine/Source/Programs/Horde/HordeStorage/Horde.Storage/Controllers/StorageController.cs
--- a/Engine/Source/Programs/Horde/HordeStorage/Horde.Storage/Controllers/StorageController.cs
+++ b/Engine/Source/Programs/Horde/HordeStorage/Horde.Storage/Controllers/StorageController.cs
@@ -48,7 +48,8 @@ namespace Horde.Storage.Controllers
         [ProducesDefaultResponseType]
         public async Task<IActionResult> Get(
             [Required] NamespaceId ns,
-            [Required] BlobIdentifier id)
+            [Required] BlobIdentifier id,
+            [FromQuery] List<string>? storageLayers = null)
         {
             ActionResult? result = await _requestHelper.HasAccessToNamespace(User, Request, ns, new [] { AclAction.ReadObject });
             if (result != null)
@@ -58,7 +59,7 @@ namespace Horde.Storage.Controllers
 
             try
             {
-                BlobContents blobContents = await GetImpl(ns, id);
+                BlobContents blobContents = await GetImpl(ns, id, storageLayers);
 
                 return File(blobContents.Stream, MediaTypeNames.Application.Octet, enableRangeProcessing: true);
             }
@@ -72,14 +73,15 @@ namespace Horde.Storage.Controllers
         [ProducesDefaultResponseType]
         public async Task<IActionResult> Head(
             [Required] NamespaceId ns,
-            [Required] BlobIdentifier id)
+            [Required] BlobIdentifier id,
+            [FromQuery] List<string>? storageLayers = null)
         {
             ActionResult? result = await _requestHelper.HasAccessToNamespace(User, Request, ns, new [] { AclAction.ReadObject });
             if (result != null)
             {
                 return result;
             }
-            bool exists = await _storage.Exists(ns, id);
+            bool exists = await _storage.Exists(ns, id, storageLayers);
 
             if (!exists)
             {
@@ -141,11 +143,11 @@ namespace Horde.Storage.Controllers
             return Ok(new HeadMultipleResponse { Needs = missingBlobs.ToArray()});
         }
 
-        private async Task<BlobContents> GetImpl(NamespaceId ns, BlobIdentifier blob)
+        private async Task<BlobContents> GetImpl(NamespaceId ns, BlobIdentifier blob, List<string>? storageLayers = null)
         {
             try
             {
-                return await _storage.GetObject(ns, blob);
+                return await _storage.GetObject(ns, blob, storageLayers);
             }
             catch (BlobNotFoundException)
             {
diff --git a/Engine/Source/Programs/Horde/HordeStorage/Horde.Storage/Horde.Storage.Startup.cs b/Engine/Source/Programs/Horde/HordeStorage/Horde.Storage/Horde.Storage.Startup.cs
--- a/Engine/Source/Programs/Horde/HordeStorage/Horde.Storage/Horde.Storage.Startup.cs
+++ b/Engine/Source/Programs/Horde/HordeStorage/Horde.Storage/Horde.Storage.Startup.cs
@@ -102,6 +102,7 @@ namespace Horde.Storage
             services.AddOptions<ScyllaSettings>().Configure(o => Configuration.GetSection("Scylla").Bind(o)).ValidateDataAnnotations();
 
             services.AddOptions<KubernetesLeaderElectionSettings>().Configure(o => Configuration.GetSection("Kubernetes").Bind(o)).ValidateDataAnnotations();
+            services.AddOptions<StaticHordeStorageServiceDiscoverySettings>().Configure(o => Configuration.GetSection("ServiceDiscovery").Bind(o)).ValidateDataAnnotations();
 
             services.AddSingleton(typeof(OodleCompressor), CreateOodleCompressor);
             services.AddSingleton(typeof(CompressedBufferUtils), CreateCompressedBufferUtils);
@@ -128,6 +129,7 @@ namespace Horde.Storage
             services.AddSingleton<AmazonS3Store>();
             services.AddSingleton<FileSystemStore>();
             services.AddSingleton<AzureBlobStore>();
+            services.AddSingleton<HordeStorageBlobStore>();
             services.AddSingleton<MemoryBlobStore>();
             services.AddSingleton<RelayBlobStore>();
 
@@ -190,6 +192,9 @@ namespace Horde.Storage
                 services.AddSingleton<KubernetesLeaderElection>(p => (KubernetesLeaderElection)p.GetService<ILeaderElection>()!);
                 services.AddHostedService<KubernetesLeaderElection>(p => p.GetService<KubernetesLeaderElection>()!);
             }
+
+            services.AddSingleton<BlobStoreConsistencyCheckService>();
+            services.AddSingleton(typeof(IHordeStorageServiceDiscovery), ServiceDiscoveryFactory);
         }
 
         private IBlobIndex BlobIndexFactory(IServiceProvider provider)
@@ -210,6 +215,20 @@ namespace Horde.Storage
             }
         }
 
+        private IHordeStorageServiceDiscovery ServiceDiscoveryFactory(IServiceProvider provider)
+        {
+            HordeStorageSettings settings = provider.GetService<IOptionsMonitor<HordeStorageSettings>>()!.CurrentValue!;
+            switch (settings.ServiceDiscoveryImplementation)
+            {
+                case HordeStorageSettings.ServiceDiscoveryImplementations.Kubernetes:
+                    return ActivatorUtilities.CreateInstance<KubernetesHordeStorageServiceDiscovery>(provider);
+                case HordeStorageSettings.ServiceDiscoveryImplementations.Static:
+                    return ActivatorUtilities.CreateInstance<StaticHordeStorageServiceDiscovery>(provider);
+                default:
+                    throw new NotImplementedException($"Unknown service discovery implementation: {settings.ServiceDiscoveryImplementation}");
+            }
+        }
+
         private IAmazonSecretsManager CreateAWSSecretsManager(IServiceProvider provider)
         {
             AWSCredentials awsCredentials = provider.GetService<AWSCredentials>()!;
@@ -487,6 +506,7 @@ namespace Horde.Storage
                         break;
                     case HordeStorageSettings.StorageBackendImplementations.Memory:
                     case HordeStorageSettings.StorageBackendImplementations.Relay:
+                    case HordeStorageSettings.StorageBackendImplementations.HordeStorage:
                         break;
                     default:
                         throw new ArgumentOutOfRangeException("Unhandled storage impl " + impl);
diff --git a/Engine/Source/Programs/Horde/HordeStorage/Horde.Storage/Implementation/Blob/HordeStorageBlobStore.cs b/Engine/Source/Programs/Horde/HordeStorage/Horde.Storage/Implementation/Blob/HordeStorageBlobStore.cs
new file mode 100644
--- /dev/null
+++ b/Engine/Source/Programs/Horde/HordeStorage/Horde.Storage/Implementation/Blob/HordeStorageBlobStore.cs
@@ -0,0 +1,211 @@
+// Copyright Epic Games, Inc. All Rights Reserved.
+
+using System;
+using System.Collections.Concurrent;
+using System.Collections.Generic;
+using System.ComponentModel.DataAnnotations;
+using System.IO;
+using System.Net;
+using System.Net.Http;
+using System.Net.Mime;
+using System.Threading.Tasks;
+using Dasync.Collections;
+using EpicGames.Horde.Storage;
+using Horde.Storage.Implementation.LeaderElection;
+using Jupiter.Implementation;
+using k8s;
+using k8s.Models;
+using Microsoft.Extensions.Options;
+using Serilog;
+
+namespace Horde.Storage.Implementation
+{
+    public class HordeStorageBlobStore : IBlobStore
+    {
+        private readonly IHordeStorageServiceDiscovery _serviceDiscovery;
+        private readonly IHttpClientFactory _httpClientFactory;
+        private readonly IServiceCredentials _serviceCredentials;
+        private readonly ILogger _logger = Log.ForContext<HordeStorageBlobStore>();
+        private readonly ConcurrentDictionary<string, HttpClient> _httpClients = new ConcurrentDictionary<string, HttpClient>();
+
+        public HordeStorageBlobStore(IHordeStorageServiceDiscovery serviceDiscovery, IHttpClientFactory httpClientFactory, IServiceCredentials serviceCredentials)
+        {
+            _serviceDiscovery = serviceDiscovery;
+            _httpClientFactory = httpClientFactory;
+            _serviceCredentials = serviceCredentials;
+        }
+
+        private async Task<HttpRequestMessage> BuildHttpRequest(HttpMethod method, Uri uri)
+        {
+            string? token = await _serviceCredentials.GetTokenAsync();
+            HttpRequestMessage request = new HttpRequestMessage(method, uri);
+            if (!string.IsNullOrEmpty(token))
+            {
+                request.Headers.Add("Authorization", $"{_serviceCredentials.GetAuthenticationScheme()} {token}");
+            }
+
+            return request;
+        }
+
+        public async Task<BlobContents> GetObject(NamespaceId ns, BlobIdentifier blob, LastAccessTrackingFlags flags = LastAccessTrackingFlags.DoTracking)
+        {
+            // TODO: This could be done in parallel
+            await foreach (string instance in _serviceDiscovery.FindOtherHordeStorageInstances())
+            {
+                using HttpRequestMessage getObjectRequest = await BuildHttpRequest(HttpMethod.Get, new Uri($"api/v1/blobs/{ns}/{blob}?storageLayers=Filesystem", UriKind.Relative));
+                getObjectRequest.Headers.Add("Accept", MediaTypeNames.Application.Octet);
+                HttpResponseMessage response = await GetHttpClient(instance).SendAsync(getObjectRequest);
+                if (response.StatusCode == HttpStatusCode.NotFound)
+                {
+                    continue;
+                }
+
+                response.EnsureSuccessStatusCode();
+
+                long? contentLength = response.Content.Headers.ContentLength;
+                if (contentLength == null)
+                {
+                    _logger.Warning("Content length missing in response from horde storage blob store. This is not supported, ignoring response");
+                    continue;
+                }
+                return new BlobContents(await response.Content.ReadAsStreamAsync(), contentLength.Value);
+            }
+            
+            throw new BlobNotFoundException(ns, blob);
+        }
+
+        public async Task<bool> Exists(NamespaceId ns, BlobIdentifier blob, bool forceCheck = false)
+        {
+            // TODO: This could be done in parallel
+            await foreach (string instance in _serviceDiscovery.FindOtherHordeStorageInstances())
+            {
+                using HttpRequestMessage headObjectRequest = await BuildHttpRequest(HttpMethod.Head, new Uri($"api/v1/blobs/{ns}/{blob}?storageLayers=Filesystem", UriKind.Relative));
+                HttpResponseMessage response = await GetHttpClient(instance).SendAsync(headObjectRequest);
+                if (response.StatusCode == HttpStatusCode.NotFound)
+                {
+                    continue;
+                }
+
+                response.EnsureSuccessStatusCode();
+
+                return true;
+            }
+
+            return false;
+        }
+
+        private HttpClient GetHttpClient(string instance)
+        {
+            return _httpClients.GetOrAdd(instance, s =>
+            {
+                HttpClient httpClient = _httpClientFactory.CreateClient(instance);
+                httpClient.BaseAddress = new Uri($"http://{instance}");
+                return httpClient;
+            });
+        }
+
+        public Task<BlobIdentifier> PutObject(NamespaceId ns, byte[] blob, BlobIdentifier identifier)
+        {
+            // not applicable
+            return Task.FromResult(identifier);
+        }
+
+        public Task<BlobIdentifier> PutObject(NamespaceId ns, ReadOnlyMemory<byte> blob, BlobIdentifier identifier)
+        {
+            // not applicable
+            return Task.FromResult(identifier);
+        }
+
+        public Task<BlobIdentifier> PutObject(NamespaceId ns, Stream content, BlobIdentifier identifier)
+        {
+            // not applicable
+            return Task.FromResult(identifier);
+        }
+
+        public Task DeleteObject(NamespaceId ns, BlobIdentifier blob)
+        {
+            // not applicable
+            return Task.CompletedTask;
+        }
+
+        public Task DeleteNamespace(NamespaceId ns)
+        {
+            // not applicable
+            return Task.CompletedTask;
+        }
+
+        public IAsyncEnumerable<(BlobIdentifier, DateTime)> ListObjects(NamespaceId ns)
+        {
+            // not applicable
+            return AsyncEnumerable<(BlobIdentifier, DateTime)>.Empty;
+        }
+    }
+
+    public interface IHordeStorageServiceDiscovery
+    {
+        public IAsyncEnumerable<string> FindOtherHordeStorageInstances();
+    }
+
+    public class StaticHordeStorageServiceDiscoverySettings
+    {
+        [System.Diagnostics.CodeAnalysis.SuppressMessage("Usage", "CA2227:Collection properties should be read only", Justification = "Used by serialization")]
+        public List<string> Peers { get; set; } = new List<string>();
+    }
+
+    public sealed class StaticHordeStorageServiceDiscovery : IHordeStorageServiceDiscovery
+    {
+        private readonly IOptionsMonitor<StaticHordeStorageServiceDiscoverySettings> _settings;
+
+        public StaticHordeStorageServiceDiscovery(IOptionsMonitor<StaticHordeStorageServiceDiscoverySettings> settings)
+        {
+            _settings = settings;
+        }
+
+        public async IAsyncEnumerable<string> FindOtherHordeStorageInstances()
+        {
+            await Task.CompletedTask;
+            foreach (string peer in _settings.CurrentValue.Peers)
+            {
+                yield return peer;
+            }
+        }
+    }
+
+    public sealed class KubernetesHordeStorageServiceDiscovery: IHordeStorageServiceDiscovery, IDisposable
+    {
+        private readonly IOptionsMonitor<KubernetesLeaderElectionSettings> _leaderSettings;
+        private readonly Kubernetes _client;
+
+        public KubernetesHordeStorageServiceDiscovery(IOptionsMonitor<KubernetesLeaderElectionSettings> leaderSettings)
+        {
+            _leaderSettings = leaderSettings;
+            KubernetesClientConfiguration config = KubernetesClientConfiguration.InClusterConfig();
+            _client = new Kubernetes(config);
+        }
+
+        public async IAsyncEnumerable<string> FindOtherHordeStorageInstances()
+        {
+            V1PodList podList = await _client.ListNamespacedPodAsync(_leaderSettings.CurrentValue.Namespace, labelSelector: _leaderSettings.CurrentValue.HordeStoragePodLabelSelector);
+
+            foreach (V1Pod pod in podList.Items)
+            {
+                // we typically expose the container on port 80
+                yield return $"{pod.Name()}:80";
+            }
+        }
+
+        private void Dispose(bool disposing)
+        {
+            if (disposing)
+            {
+                _client.Dispose();
+            }
+        }
+
+        public void Dispose()
+        {
+            Dispose(true);
+            GC.SuppressFinalize(this);
+        }
+    }
+}
diff --git a/Engine/Source/Programs/Horde/HordeStorage/Horde.Storage/Implementation/Blob/IBlobService.cs b/Engine/Source/Programs/Horde/HordeStorage/Horde.Storage/Implementation/Blob/IBlobService.cs
--- a/Engine/Source/Programs/Horde/HordeStorage/Horde.Storage/Implementation/Blob/IBlobService.cs
+++ b/Engine/Source/Programs/Horde/HordeStorage/Horde.Storage/Implementation/Blob/IBlobService.cs
@@ -35,11 +35,11 @@ public interface IBlobService
     Task<BlobIdentifier> PutObject(NamespaceId ns, IBufferedPayload payload, BlobIdentifier identifier);
     Task<BlobIdentifier> PutObject(NamespaceId ns, byte[] payload, BlobIdentifier identifier);
 
-    Task<BlobContents> GetObject(NamespaceId ns, BlobIdentifier blob);
+    Task<BlobContents> GetObject(NamespaceId ns, BlobIdentifier blob, List<string>? storageLayers = null);
 
     Task<BlobContents> ReplicateObject(NamespaceId ns, BlobIdentifier blob, bool force = false);
 
-    Task<bool> Exists(NamespaceId ns, BlobIdentifier blob);
+    Task<bool> Exists(NamespaceId ns, BlobIdentifier blob, List<string>? storageLayers = null);
 
     /// <summary>
     /// Checks that the blob exists in the root store, the store which is last in the list and thus is intended to have every blob in it
@@ -107,6 +107,7 @@ public class BlobService : IBlobService
             HordeStorageSettings.StorageBackendImplementations.FileSystem => provider.GetService<FileSystemStore>(),
             HordeStorageSettings.StorageBackendImplementations.Memory => provider.GetService<MemoryBlobStore>(),
             HordeStorageSettings.StorageBackendImplementations.Relay => provider.GetService<RelayBlobStore>(),
+            HordeStorageSettings.StorageBackendImplementations.HordeStorage => provider.GetService<HordeStorageBlobStore>(),
             _ => throw new NotImplementedException("Unknown blob store {store")
         };
         if (store == null)
@@ -206,7 +207,7 @@ public class BlobService : IBlobService
         return identifier;
     }
 
-    public async Task<BlobContents> GetObject(NamespaceId ns, BlobIdentifier blob)
+    public async Task<BlobContents> GetObject(NamespaceId ns, BlobIdentifier blob, List<string>? storageLayers = null)
     {
         bool seenBlobNotFound = false;
         bool seenNamespaceNotFound = false;
@@ -216,9 +217,30 @@ public class BlobService : IBlobService
 
         foreach (IBlobStore store in _blobStores)
         {
+            string blobStoreName = store.GetType().Name;
+
+            // check which storage layers to skip if we have a explicit list of storage layers to use
+            if (storageLayers != null && storageLayers.Count != 0)
+            {
+                bool found = false;
+                foreach (string storageLayer in storageLayers)
+                {
+                    if (string.Equals(storageLayer, blobStoreName, StringComparison.OrdinalIgnoreCase))
+                    {
+                        found = true;
+                        break;
+                    }
+                }
+
+                if (!found)
+                {
+                    continue;
+                }
+            }
+
             using IScope scope = Tracer.Instance.StartActive("HierarchicalStore.GetObject");
             scope.Span.ResourceName = blob.ToString();
-            scope.Span.SetTag("BlobStore", store.GetType().Name);
+            scope.Span.SetTag("BlobStore", blobStoreName);
             scope.Span.SetTag("ObjectFound", false.ToString());
             string storeName = store.GetType().Name;
             using ServerTimingMetricScoped? serverTimingScope = serverTiming?.CreateServerTimingMetricScope($"blob.get.{storeName}", $"Blob GET from: '{storeName}'");
@@ -249,7 +271,13 @@ public class BlobService : IBlobService
         {
             throw new NamespaceNotFoundException(ns);
         }
-        
+
+        // if we applied filters to the storage layers resulting in no blob found we consider it a miss
+        if (storageLayers != null && storageLayers.Count != 0 && blobContents == null)
+        {
+            throw new BlobNotFoundException(ns, blob);
+        }
+
         if (blobContents == null)
         {
             // Should not happen but exists to safeguard against the null pointer
@@ -367,7 +395,7 @@ public class BlobService : IBlobService
         return request;
     }
 
-    public async Task<bool> Exists(NamespaceId ns, BlobIdentifier blob)
+    public async Task<bool> Exists(NamespaceId ns, BlobIdentifier blob, List<string>? storageLayers = null)
     {
         bool useBlobIndex = _namespacePolicyResolver.GetPoliciesForNs(ns).UseBlobIndexForExists;
         if (useBlobIndex)
@@ -389,9 +417,28 @@ public class BlobService : IBlobService
         {
             foreach (IBlobStore store in _blobStores)
             {
+                string blobStoreName = store.GetType().Name;
+                // check which storage layers to skip if we have a explicit list of storage layers to use
+                if (storageLayers != null && storageLayers.Count != 0)
+                {
+                    bool found = false;
+                    foreach (string storageLayer in storageLayers)
+                    {
+                        if (string.Equals(storageLayer, blobStoreName, StringComparison.OrdinalIgnoreCase))
+                        {
+                            found = true;
+                            break;
+                        }
+                    }
+
+                    if (!found)
+                    {
+                        continue;
+                    }
+                }
                 using IScope scope = Tracer.Instance.StartActive("HierarchicalStore.ObjectExists");
                 scope.Span.ResourceName = blob.ToString();
-                scope.Span.SetTag("BlobStore", store.GetType().Name);
+                scope.Span.SetTag("BlobStore", blobStoreName);
                 if (await store.Exists(ns, blob))
                 {
                     scope.Span.SetTag("ObjectFound", true.ToString());
diff --git a/Engine/Source/Programs/Horde/HordeStorage/Horde.Storage/Implementation/LeaderElection/KubernetesLeaderElection.cs b/Engine/Source/Programs/Horde/HordeStorage/Horde.Storage/Implementation/LeaderElection/KubernetesLeaderElection.cs
--- a/Engine/Source/Programs/Horde/HordeStorage/Horde.Storage/Implementation/LeaderElection/KubernetesLeaderElection.cs
+++ b/Engine/Source/Programs/Horde/HordeStorage/Horde.Storage/Implementation/LeaderElection/KubernetesLeaderElection.cs
@@ -26,6 +26,8 @@ namespace Horde.Storage.Implementation.LeaderElection
 
         [Required]
         public string Namespace { get; set; } = "";
+
+        public string HordeStoragePodLabelSelector { get; set; } = "app.kubernetes.io/name: horde-storage";
     }
 
     public class KubernetesLeaderElectionState
diff --git a/Engine/Source/Programs/Horde/HordeStorage/Horde.Storage/Settings.cs b/Engine/Source/Programs/Horde/HordeStorage/Horde.Storage/Settings.cs
--- a/Engine/Source/Programs/Horde/HordeStorage/Horde.Storage/Settings.cs
+++ b/Engine/Source/Programs/Horde/HordeStorage/Horde.Storage/Settings.cs
@@ -24,7 +24,8 @@ namespace Horde.Storage
             Azure,
             FileSystem,
             Memory,
-            Relay
+            Relay,
+            HordeStorage
         }
 
         public enum ReferencesDbImplementations
@@ -57,6 +58,12 @@ namespace Horde.Storage
             Kubernetes
         }
 
+        public enum ServiceDiscoveryImplementations
+        {
+            Static, 
+            Kubernetes
+        }
+
         private sealed class ValidStorageBackend : ValidationAttribute
         {
             public override string FormatErrorMessage(string name)
@@ -98,6 +105,7 @@ namespace Horde.Storage
         public LeaderElectionImplementations LeaderElectionImplementation { get; set; } = LeaderElectionImplementations.Static;
         public ContentIdStoreImplementations ContentIdStoreImplementation { get; set; } = ContentIdStoreImplementations.Memory;
         public BlobIndexImplementations BlobIndexImplementation { get; set; } = BlobIndexImplementations.Memory;
+        public ServiceDiscoveryImplementations ServiceDiscoveryImplementation { get; set; } = ServiceDiscoveryImplementations.Static;
 
         public int? MaxSingleBlobSize { get; set; } = null; // disable blob partitioning
 
diff --git a/Engine/Source/Programs/Horde/HordeStorage/Horde.Storage/Tests/Functional/Blob/HordeStorageBlobTest.cs b/Engine/Source/Programs/Horde/HordeStorage/Horde.Storage/Tests/Functional/Blob/HordeStorageBlobTest.cs
new file mode 100644
--- /dev/null
+++ b/Engine/Source/Programs/Horde/HordeStorage/Horde.Storage/Tests/Functional/Blob/HordeStorageBlobTest.cs
@@ -0,0 +1,188 @@
+// Copyright Epic Games, Inc. All Rights Reserved.
+
+using System;
+using System.Collections.Generic;
+using System.Linq;
+using System.Net;
+using System.Net.Http;
+using System.Net.Http.Headers;
+using System.Text;
+using System.Threading.Tasks;
+using Amazon.S3;
+using Amazon.S3.Model;
+using EpicGames.Horde.Storage;
+using Horde.Storage.FunctionalTests.Storage;
+using Jupiter.Implementation;
+using Microsoft.AspNetCore.Hosting;
+using Microsoft.AspNetCore.TestHost;
+using Microsoft.Extensions.Configuration;
+using Microsoft.Extensions.DependencyInjection;
+using Microsoft.VisualStudio.TestTools.UnitTesting;
+using Moq;
+using Moq.Contrib.HttpClient;
+using Serilog;
+using Serilog.Core;
+
+namespace Horde.Storage.FunctionalTests.Storage
+{
+    [TestClass]
+    public class HordeStorageBlobTest
+    {
+        protected const string S3FileContent = "Foo bar content that goes into s3";
+        protected const string OtherHordeStorageContent = "This content will be faked from another horde storage";
+
+        protected BlobIdentifier S3FileContentHash { get; } = BlobIdentifier.FromBlob(Encoding.ASCII.GetBytes(S3FileContent));
+        protected BlobIdentifier OtherHordeStorageContentHash { get; } = BlobIdentifier.FromBlob(Encoding.ASCII.GetBytes(OtherHordeStorageContent));
+
+        private IAmazonS3? _s3;
+        protected NamespaceId TestNamespaceName { get; } = new NamespaceId("testbucket");
+        protected TestServer? Server { get; set; }
+
+        private HttpClient? _httpClient;
+        private Mock<HttpMessageHandler>? _otherHordeStorageMock;
+
+        [TestInitialize]
+        public async Task Setup()
+        {
+            IConfigurationRoot configuration = new ConfigurationBuilder()
+                // we are not reading the base appSettings here as we want exact control over what runs in the tests
+                .AddJsonFile("appsettings.Testing.json", true)
+                .AddInMemoryCollection(GetSettings())
+                .AddEnvironmentVariables()
+                .Build();
+
+            Logger logger = new LoggerConfiguration()
+                .ReadFrom.Configuration(configuration)
+                .CreateLogger();
+
+            _otherHordeStorageMock = new Mock<HttpMessageHandler>();
+
+            TestServer server = new TestServer(new WebHostBuilder()
+                .UseConfiguration(configuration)
+                .UseEnvironment("Testing")
+                .UseSerilog(logger)
+                .UseStartup<HordeStorageStartup>()
+                .ConfigureTestServices(collection =>
+                {
+                    // configure a faked remote horde storage instance that can serve this blob
+                    _otherHordeStorageMock.SetupRequest($"http://other-horde-storage-instance.com/api/v1/blobs/{TestNamespaceName}/{S3FileContentHash}?storageLayers=Filesystem").ReturnsResponse(HttpStatusCode.OK, S3FileContent).Verifiable();
+                    _otherHordeStorageMock.SetupRequest($"http://other-horde-storage-instance.com/api/v1/blobs/{TestNamespaceName}/{OtherHordeStorageContentHash}?storageLayers=Filesystem").ReturnsResponse(HttpStatusCode.OK, OtherHordeStorageContent).Verifiable();
+
+                    IHttpClientFactory mockFactory = _otherHordeStorageMock.CreateClientFactory();
+
+                    Mock.Get(mockFactory).Setup(x => x.CreateClient("other-horde-storage-instance.com"))
+                    .Returns(() =>
+                    {
+                        HttpClient? client = _otherHordeStorageMock.CreateClient();
+                        client.BaseAddress = new Uri("http://other-horde-storage-instance.com");
+                        return client;
+                    });
+
+                    collection.AddSingleton<IHttpClientFactory>(mockFactory);
+                })
+
+            );
+            _httpClient = server.CreateClient();
+            Server = server;
+
+            // Seed storage
+            await Seed(Server.Services);
+        }
+
+        [TestCleanup]
+        public async Task MyTeardown()
+        {
+            await Teardown(Server!.Services);
+        }
+        protected IEnumerable<KeyValuePair<string, string>> GetSettings()
+        {
+            return new[]
+            {
+                new KeyValuePair<string, string>("Horde_Storage:StorageImplementations:0", HordeStorageSettings.StorageBackendImplementations.HordeStorage.ToString()),
+                new KeyValuePair<string, string>("Horde_Storage:StorageImplementations:1", HordeStorageSettings.StorageBackendImplementations.S3.ToString()),
+                new KeyValuePair<string, string>("ServiceDiscovery:Peers:0", "other-horde-storage-instance.com"),
+                new KeyValuePair<string, string>("S3:BucketName", $"tests-{TestNamespaceName}")
+            };
+        }
+
+        protected async Task Seed(IServiceProvider provider)
+        {
+            string S3BucketName = $"tests-{TestNamespaceName}";
+
+            _s3 = provider.GetService<IAmazonS3>();
+            Assert.IsNotNull(_s3);
+            if (await _s3!.DoesS3BucketExistAsync(S3BucketName))
+            {
+                // if we have failed to run the cleanup for some reason we run it now
+                await Teardown(provider);
+            }
+        }
+
+        protected async Task Teardown(IServiceProvider provider)
+        {
+            string S3BucketName = $"tests-{TestNamespaceName}";
+            ListObjectsResponse response = await _s3!.ListObjectsAsync(S3BucketName);
+            List<KeyVersion> objectKeys = response.S3Objects.Select(o => new KeyVersion { Key = o.Key }).ToList();
+            await _s3.DeleteObjectsAsync(new DeleteObjectsRequest { BucketName = S3BucketName, Objects = objectKeys });
+
+            await _s3.DeleteBucketAsync(S3BucketName);
+        }
+
+        [TestMethod]
+        public async Task ForceFetchBlobFromS3()
+        {
+            byte[] s3ContentBytes = Encoding.ASCII.GetBytes(S3FileContent);
+            using ByteArrayContent requestContent = new ByteArrayContent(s3ContentBytes);
+            requestContent.Headers.ContentType = new MediaTypeHeaderValue("application/octet-stream");
+            requestContent.Headers.ContentLength = S3FileContent.Length;
+            HttpResponseMessage putResponse = await _httpClient!.PutAsync(new Uri($"api/v1/s/{TestNamespaceName}/{S3FileContentHash}", UriKind.Relative), requestContent);
+            putResponse.EnsureSuccessStatusCode();
+            InsertResponse response = await putResponse.Content.ReadAsAsync<InsertResponse>();
+            Assert.AreEqual(S3FileContentHash, response.Identifier);
+
+            HttpResponseMessage getResponse = await _httpClient.GetAsync(new Uri($"api/v1/s/{TestNamespaceName}/{S3FileContentHash}?storageLayers=AmazonS3Store", UriKind.Relative));
+            getResponse.EnsureSuccessStatusCode();
+            CollectionAssert.AreEqual(s3ContentBytes, actual: await getResponse.Content.ReadAsByteArrayAsync());
+
+            _otherHordeStorageMock!.VerifyNoOtherCalls();
+
+        }
+
+        
+        [TestMethod]
+        public async Task ForceFetchBlobFromNoneExistentLayer()
+        {
+            byte[] s3ContentBytes = Encoding.ASCII.GetBytes(S3FileContent);
+            using ByteArrayContent requestContent = new ByteArrayContent(s3ContentBytes);
+            requestContent.Headers.ContentType = new MediaTypeHeaderValue("application/octet-stream");
+            requestContent.Headers.ContentLength = S3FileContent.Length;
+            HttpResponseMessage putResponse = await _httpClient!.PutAsync(new Uri($"api/v1/s/{TestNamespaceName}/{S3FileContentHash}", UriKind.Relative), requestContent);
+            putResponse.EnsureSuccessStatusCode();
+            InsertResponse response = await putResponse.Content.ReadAsAsync<InsertResponse>();
+            Assert.AreEqual(S3FileContentHash, response.Identifier);
+
+            HttpResponseMessage getResponse = await _httpClient.GetAsync(new Uri($"api/v1/s/{TestNamespaceName}/{S3FileContentHash}?storageLayers=ThisImplemenationDoesNotExist", UriKind.Relative));
+            Assert.AreEqual(HttpStatusCode.NotFound, getResponse.StatusCode);
+        }
+
+        [TestMethod]
+        public async Task FetchBlobFromOtherInstance()
+        {
+            byte[] payload = Encoding.ASCII.GetBytes(OtherHordeStorageContent);
+            using ByteArrayContent requestContent = new ByteArrayContent(payload);
+            requestContent.Headers.ContentType = new MediaTypeHeaderValue("application/octet-stream");
+            requestContent.Headers.ContentLength = payload.Length;
+            HttpResponseMessage putResponse = await _httpClient!.PutAsync(new Uri($"api/v1/s/{TestNamespaceName}/{OtherHordeStorageContentHash}", UriKind.Relative), requestContent);
+            putResponse.EnsureSuccessStatusCode();
+            InsertResponse response = await putResponse.Content.ReadAsAsync<InsertResponse>();
+            Assert.AreEqual(OtherHordeStorageContentHash, response.Identifier);
+
+            HttpResponseMessage getResponse = await _httpClient.GetAsync(new Uri($"api/v1/s/{TestNamespaceName}/{OtherHordeStorageContentHash}", UriKind.Relative));
+            getResponse.EnsureSuccessStatusCode();
+            CollectionAssert.AreEqual(payload, actual: await getResponse.Content.ReadAsByteArrayAsync());
+
+            _otherHordeStorageMock!.VerifyRequest(HttpMethod.Get, new Uri($"http://other-horde-storage-instance.com/api/v1/blobs/{TestNamespaceName}/{OtherHordeStorageContentHash}?storageLayers=Filesystem"));
+            _otherHordeStorageMock!.VerifyNoOtherCalls();
+        }
+    }
+}
