Horde Storage - Added stacktrace (if present) to returned cbobject to better debug inconsistent error with batch operations.

diff --git a/Engine/Source/Programs/Horde/HordeStorage/Horde.Storage/Controllers/ReferencesController.cs b/Engine/Source/Programs/Horde/HordeStorage/Horde.Storage/Controllers/ReferencesController.cs
--- a/Engine/Source/Programs/Horde/HordeStorage/Horde.Storage/Controllers/ReferencesController.cs
+++ b/Engine/Source/Programs/Horde/HordeStorage/Horde.Storage/Controllers/ReferencesController.cs
@@ -971,6 +971,10 @@ namespace Horde.Storage.Controllers
             writer.BeginObject();
             writer.WriteString("title", e.Message);
             writer.WriteInteger("status", (int)statusCode);
+            if (e.StackTrace != null)
+            {
+                writer.WriteString("stackTrace", e.StackTrace);
+            }
             writer.EndObject();
             return (writer.ToObject(), statusCode);
         }
