Horde Storage - Fixed issue with memory replication log (only used for tests) were it could get conflicting keys for the replication log (as time uuid is globally unique but multiple entries sorts into the same place by definition). Treating it as a guid should resolve this (allowing the full set of the unique value to be used for compare).

diff --git a/Engine/Source/Programs/Horde/HordeStorage/Horde.Storage/Implementation/TransactionLog/MemoryReplicationLog.cs b/Engine/Source/Programs/Horde/HordeStorage/Horde.Storage/Implementation/TransactionLog/MemoryReplicationLog.cs
--- a/Engine/Source/Programs/Horde/HordeStorage/Horde.Storage/Implementation/TransactionLog/MemoryReplicationLog.cs
+++ b/Engine/Source/Programs/Horde/HordeStorage/Horde.Storage/Implementation/TransactionLog/MemoryReplicationLog.cs
@@ -14,7 +14,7 @@ namespace Horde.Storage.Implementation
 {
     internal class MemoryReplicationLog : IReplicationLog
     {
-        private readonly ConcurrentDictionary<NamespaceId, SortedList<string, SortedList<TimeUuid, ReplicationLogEvent>>> _replicationEvents = new();
+        private readonly ConcurrentDictionary<NamespaceId, SortedList<string, SortedList<Guid, ReplicationLogEvent>>> _replicationEvents = new();
         private readonly ConcurrentDictionary<NamespaceId, List<SnapshotInfo>>  _snapshots = new();
 
         private readonly ConcurrentDictionary<NamespaceId, ConcurrentDictionary<string, ReplicatorState>>  _replicatorState = new();
@@ -42,13 +42,13 @@ namespace Horde.Storage.Implementation
 
                 _replicationEvents.AddOrUpdate(ns, _ =>
                 {
-                    SortedList<string, SortedList<TimeUuid, ReplicationLogEvent>> l = new() { { bucketId, new () { {eventId, logEvent} } } };
+                    SortedList<string, SortedList<Guid, ReplicationLogEvent>> l = new() { { bucketId, new () { {eventId, logEvent} } } };
                     return l;
                 }, (_, buckets) =>
                 {
                     lock (buckets)
                     {
-                        if (buckets.TryGetValue(bucketId, out SortedList<TimeUuid, ReplicationLogEvent>? events))
+                        if (buckets.TryGetValue(bucketId, out SortedList<Guid, ReplicationLogEvent>? events))
                         {
                             lock (events)
                             {
@@ -76,7 +76,7 @@ namespace Horde.Storage.Implementation
         public async IAsyncEnumerable<ReplicationLogEvent> Get(NamespaceId ns, string? lastBucket, Guid? lastEvent)
         {
             await Task.CompletedTask;
-            if (!_replicationEvents.TryGetValue(ns, out SortedList<string, SortedList<TimeUuid, ReplicationLogEvent>>? buckets))
+            if (!_replicationEvents.TryGetValue(ns, out SortedList<string, SortedList<Guid, ReplicationLogEvent>>? buckets))
             {
                 throw new NamespaceNotFoundException(ns);
             }
@@ -105,7 +105,7 @@ namespace Horde.Storage.Implementation
             {
                 string bucket = buckets.Keys[i];
 
-                if (!buckets.TryGetValue(bucket, out SortedList<TimeUuid, ReplicationLogEvent>? logEvents))
+                if (!buckets.TryGetValue(bucket, out SortedList<Guid, ReplicationLogEvent>? logEvents))
                 {
                     yield break;
                 }
@@ -123,7 +123,7 @@ namespace Horde.Storage.Implementation
                     skipCount = eventIndex + 1;
                 }
 
-                foreach ((TimeUuid _, ReplicationLogEvent value) in logEvents.Skip(skipCount))
+                foreach ((Guid _, ReplicationLogEvent value) in logEvents.Skip(skipCount))
                 {
                     yield return value;
                 }
