Horde Storage - Another attempt at fixing sorting issues with the memory replication log to resolve issues with our tests intermittently failing. Avoid using sorted list as it will not correctly handle TimeUuids that sort to the same value (based on time) but we still need them to be sorted by time when looping over these (so switched to a list and manually sorted them when we need it sorted). Also see change 22203479.

diff --git a/Engine/Source/Programs/Horde/HordeStorage/Horde.Storage/Implementation/TransactionLog/MemoryReplicationLog.cs b/Engine/Source/Programs/Horde/HordeStorage/Horde.Storage/Implementation/TransactionLog/MemoryReplicationLog.cs
--- a/Engine/Source/Programs/Horde/HordeStorage/Horde.Storage/Implementation/TransactionLog/MemoryReplicationLog.cs
+++ b/Engine/Source/Programs/Horde/HordeStorage/Horde.Storage/Implementation/TransactionLog/MemoryReplicationLog.cs
@@ -7,6 +7,7 @@ using System.Linq;
 using System.Threading.Tasks;
 using Cassandra;
 using Dasync.Collections;
+using EpicGames.Core;
 using EpicGames.Horde.Storage;
 using Jupiter.Implementation;
 
@@ -14,7 +15,7 @@ namespace Horde.Storage.Implementation
 {
     internal class MemoryReplicationLog : IReplicationLog
     {
-        private readonly ConcurrentDictionary<NamespaceId, SortedList<string, SortedList<Guid, ReplicationLogEvent>>> _replicationEvents = new();
+        private readonly ConcurrentDictionary<NamespaceId, SortedList<string, List<(TimeUuid, ReplicationLogEvent)>>> _replicationEvents = new();
         private readonly ConcurrentDictionary<NamespaceId, List<SnapshotInfo>>  _snapshots = new();
 
         private readonly ConcurrentDictionary<NamespaceId, ConcurrentDictionary<string, ReplicatorState>>  _replicatorState = new();
@@ -42,22 +43,22 @@ namespace Horde.Storage.Implementation
 
                 _replicationEvents.AddOrUpdate(ns, _ =>
                 {
-                    SortedList<string, SortedList<Guid, ReplicationLogEvent>> l = new() { { bucketId, new () { {eventId, logEvent} } } };
+                    SortedList<string, List<(TimeUuid, ReplicationLogEvent)>> l = new() { { bucketId, new () { (eventId, logEvent) } } };
                     return l;
                 }, (_, buckets) =>
                 {
                     lock (buckets)
                     {
-                        if (buckets.TryGetValue(bucketId, out SortedList<Guid, ReplicationLogEvent>? events))
+                        if (buckets.TryGetValue(bucketId, out List<(TimeUuid, ReplicationLogEvent)>? events))
                         {
                             lock (events)
                             {
-                                events.Add(eventId, logEvent);
+                                events.Add((eventId, logEvent));
                             }
                         }
                         else
                         {
-                            buckets.Add(bucketId, new () {{eventId, logEvent} });
+                            buckets.Add(bucketId, new () {(eventId, logEvent) });
                         }
                     }
 
@@ -76,7 +77,7 @@ namespace Horde.Storage.Implementation
         public async IAsyncEnumerable<ReplicationLogEvent> Get(NamespaceId ns, string? lastBucket, Guid? lastEvent)
         {
             await Task.CompletedTask;
-            if (!_replicationEvents.TryGetValue(ns, out SortedList<string, SortedList<Guid, ReplicationLogEvent>>? buckets))
+            if (!_replicationEvents.TryGetValue(ns, out SortedList<string, List<(TimeUuid, ReplicationLogEvent)>>? buckets))
             {
                 throw new NamespaceNotFoundException(ns);
             }
@@ -105,27 +106,35 @@ namespace Horde.Storage.Implementation
             {
                 string bucket = buckets.Keys[i];
 
-                if (!buckets.TryGetValue(bucket, out SortedList<Guid, ReplicationLogEvent>? logEvents))
+                if (!buckets.TryGetValue(bucket, out List<(TimeUuid, ReplicationLogEvent)>? logEvents))
                 {
                     yield break;
                 }
-
-                int skipCount = 0;
-                if (lastEvent.HasValue)
+                logEvents.SortBy(tuple => tuple.Item1);
+                bool shouldSkipEvents = lastEvent.HasValue;
+                bool eventFound = false;
+                foreach ((TimeUuid eventId, ReplicationLogEvent value) in logEvents)
                 {
-                    int eventIndex = logEvents.IndexOfKey(lastEvent.Value);
-                    // if the event to resume from does not exist we need to use a snapshot instead
-                    if (eventIndex == -1)
+                    if (shouldSkipEvents)
                     {
-                        throw new IncrementalLogNotAvailableException();
+                        if (lastEvent!.Value.Equals(eventId))
+                        {
+                            // start reading after this event
+                            shouldSkipEvents = false;
+                            eventFound = true;
+                            continue;
+                        }
+                        // continue skipping
+                        continue;
                     }
 
-                    skipCount = eventIndex + 1;
+                    eventFound = true;
+                    yield return value;
                 }
 
-                foreach ((Guid _, ReplicationLogEvent value) in logEvents.Skip(skipCount))
+                if (!lastEvent.HasValue && !eventFound)
                 {
-                    yield return value;
+                    throw new IncrementalLogNotAvailableException();
                 }
 
                 // once we have enumerated the bucket we were working on we should now return all events from the remaining ones
