Horde Storage - Fixes for kubernetes service discovery used in the horde storage blob store

diff --git a/Engine/Source/Programs/Horde/HordeStorage/Horde.Storage/Implementation/LeaderElection/KubernetesLeaderElection.cs b/Engine/Source/Programs/Horde/HordeStorage/Horde.Storage/Implementation/LeaderElection/KubernetesLeaderElection.cs
--- a/Engine/Source/Programs/Horde/HordeStorage/Horde.Storage/Implementation/LeaderElection/KubernetesLeaderElection.cs
+++ b/Engine/Source/Programs/Horde/HordeStorage/Horde.Storage/Implementation/LeaderElection/KubernetesLeaderElection.cs
@@ -27,7 +27,7 @@ namespace Horde.Storage.Implementation.LeaderElection
         [Required]
         public string Namespace { get; set; } = "";
 
-        public string HordeStoragePodLabelSelector { get; set; } = "app.kubernetes.io/name: horde-storage";
+        public string HordeStoragePodLabelSelector { get; set; } = "app.kubernetes.io/name=horde-storage";
     }
 
     public class KubernetesLeaderElectionState
