Horde Storage - Locate other pods not by name but by ip, this seems to be more consistently resolvable within the kubernetes cluster

diff --git a/Engine/Source/Programs/Horde/HordeStorage/Horde.Storage/Implementation/Blob/HordeStorageBlobStore.cs b/Engine/Source/Programs/Horde/HordeStorage/Horde.Storage/Implementation/Blob/HordeStorageBlobStore.cs
--- a/Engine/Source/Programs/Horde/HordeStorage/Horde.Storage/Implementation/Blob/HordeStorageBlobStore.cs
+++ b/Engine/Source/Programs/Horde/HordeStorage/Horde.Storage/Implementation/Blob/HordeStorageBlobStore.cs
@@ -189,8 +189,13 @@ namespace Horde.Storage.Implementation
 
             foreach (V1Pod pod in podList.Items)
             {
+                string ip = pod.Status.PodIP;
+                if (string.IsNullOrEmpty(ip))
+                {
+                    continue;
+                }
                 // we typically expose the container on port 80
-                yield return $"{pod.Name()}:80";
+                yield return $"{ip}:80";
             }
         }
 
