Horde: Instrument background work in LogFileService

diff --git a/Engine/Source/Programs/Horde/Horde.Build/Logs/LogFileService.cs b/Engine/Source/Programs/Horde/Horde.Build/Logs/LogFileService.cs
--- a/Engine/Source/Programs/Horde/Horde.Build/Logs/LogFileService.cs
+++ b/Engine/Source/Programs/Horde/Horde.Build/Logs/LogFileService.cs
@@ -988,6 +988,8 @@ namespace Horde.Build.Logs
 		/// <param name="stoppingToken">Cancellation token</param>
 		async ValueTask TickAsync(CancellationToken stoppingToken)
 		{
+			using IScope scope = GlobalTracer.Instance.BuildSpan("LogFileService.TickAsync").StartActive();
+			
 			lock (_writeLock)
 			{
 				try
@@ -1008,10 +1010,14 @@ namespace Horde.Build.Logs
 		/// <returns>Async task</returns>
 		private async Task IncrementalFlush()
 		{
+			using IScope scope = GlobalTracer.Instance.BuildSpan("LogFileService.IncrementalFlush").StartActive();
+			
 			// Get all the chunks older than 20 minutes
 			List<(LogId, long)> flushChunks = await _builder.TouchChunksAsync(TimeSpan.FromMinutes(10.0));
 			_logger.LogDebug("Performing incremental flush of log builder ({NumChunks} chunks)", flushChunks.Count);
 
+			scope.Span.SetTag("numChunks", flushChunks.Count);
+
 			// Mark them all as complete
 			foreach ((LogId logId, long offset) in flushChunks)
 			{
@@ -1028,6 +1034,7 @@ namespace Horde.Build.Logs
 		/// <returns>Async task</returns>
 		public async Task FlushAsync()
 		{
+			using IScope scope = GlobalTracer.Instance.BuildSpan("LogFileService.FlushAsync").StartActive();
 			_logger.LogInformation("Forcing flush of pending log chunks...");
 
 			// Mark everything in the cache as complete
@@ -1071,6 +1078,8 @@ namespace Horde.Build.Logs
 		/// <param name="bCreateIndex">Create an index for the log</param>
 		private void WriteCompleteChunks(List<(LogId, long)> chunksToWrite, bool bCreateIndex)
 		{
+			using IScope scope = GlobalTracer.Instance.BuildSpan("LogFileService.WriteCompleteChunks").StartActive();
+			
 			foreach (IGrouping<LogId, long> group in chunksToWrite.GroupBy(x => x.Item1, x => x.Item2))
 			{
 				LogId logId = group.Key;
