Horde Storage - Fixed issue were horde storage blob store would always return 404 for blobs (as it was not correctly specifying the name of the filesystem store)

diff --git a/Engine/Source/Programs/Horde/HordeStorage/Horde.Storage/Implementation/Blob/HordeStorageBlobStore.cs b/Engine/Source/Programs/Horde/HordeStorage/Horde.Storage/Implementation/Blob/HordeStorageBlobStore.cs
--- a/Engine/Source/Programs/Horde/HordeStorage/Horde.Storage/Implementation/Blob/HordeStorageBlobStore.cs
+++ b/Engine/Source/Programs/Horde/HordeStorage/Horde.Storage/Implementation/Blob/HordeStorageBlobStore.cs
@@ -52,7 +52,8 @@ namespace Horde.Storage.Implementation
             // TODO: This could be done in parallel
             await foreach (string instance in _serviceDiscovery.FindOtherHordeStorageInstances())
             {
-                using HttpRequestMessage getObjectRequest = await BuildHttpRequest(HttpMethod.Get, new Uri($"api/v1/blobs/{ns}/{blob}?storageLayers=Filesystem", UriKind.Relative));
+                string filesystemLayerName = nameof(FileSystemStore);
+                using HttpRequestMessage getObjectRequest = await BuildHttpRequest(HttpMethod.Get, new Uri($"api/v1/blobs/{ns}/{blob}?storageLayers={filesystemLayerName}", UriKind.Relative));
                 getObjectRequest.Headers.Add("Accept", MediaTypeNames.Application.Octet);
                 HttpResponseMessage response = await GetHttpClient(instance).SendAsync(getObjectRequest);
                 if (response.StatusCode == HttpStatusCode.NotFound)
@@ -79,7 +80,8 @@ namespace Horde.Storage.Implementation
             // TODO: This could be done in parallel
             await foreach (string instance in _serviceDiscovery.FindOtherHordeStorageInstances())
             {
-                using HttpRequestMessage headObjectRequest = await BuildHttpRequest(HttpMethod.Head, new Uri($"api/v1/blobs/{ns}/{blob}?storageLayers=Filesystem", UriKind.Relative));
+                string filesystemLayerName = nameof(FileSystemStore);
+                using HttpRequestMessage headObjectRequest = await BuildHttpRequest(HttpMethod.Head, new Uri($"api/v1/blobs/{ns}/{blob}?storageLayers={filesystemLayerName}", UriKind.Relative));
                 HttpResponseMessage response = await GetHttpClient(instance).SendAsync(headObjectRequest);
                 if (response.StatusCode == HttpStatusCode.NotFound)
                 {
