Horde Storage - Fixed up failing tests caused by change 22212879

diff --git a/Engine/Source/Programs/Horde/HordeStorage/Horde.Storage/Tests/Functional/Blob/HordeStorageBlobTest.cs b/Engine/Source/Programs/Horde/HordeStorage/Horde.Storage/Tests/Functional/Blob/HordeStorageBlobTest.cs
--- a/Engine/Source/Programs/Horde/HordeStorage/Horde.Storage/Tests/Functional/Blob/HordeStorageBlobTest.cs
+++ b/Engine/Source/Programs/Horde/HordeStorage/Horde.Storage/Tests/Functional/Blob/HordeStorageBlobTest.cs
@@ -12,6 +12,7 @@ using Amazon.S3;
 using Amazon.S3.Model;
 using EpicGames.Horde.Storage;
 using Horde.Storage.FunctionalTests.Storage;
+using Horde.Storage.Implementation;
 using Jupiter.Implementation;
 using Microsoft.AspNetCore.Hosting;
 using Microsoft.AspNetCore.TestHost;
@@ -64,9 +65,10 @@ namespace Horde.Storage.FunctionalTests.Storage
                 .UseStartup<HordeStorageStartup>()
                 .ConfigureTestServices(collection =>
                 {
+                    string storageLayerName = nameof(FileSystemStore);
                     // configure a faked remote horde storage instance that can serve this blob
-                    _otherHordeStorageMock.SetupRequest($"http://other-horde-storage-instance.com/api/v1/blobs/{TestNamespaceName}/{S3FileContentHash}?storageLayers=Filesystem").ReturnsResponse(HttpStatusCode.OK, S3FileContent).Verifiable();
-                    _otherHordeStorageMock.SetupRequest($"http://other-horde-storage-instance.com/api/v1/blobs/{TestNamespaceName}/{OtherHordeStorageContentHash}?storageLayers=Filesystem").ReturnsResponse(HttpStatusCode.OK, OtherHordeStorageContent).Verifiable();
+                    _otherHordeStorageMock.SetupRequest($"http://other-horde-storage-instance.com/api/v1/blobs/{TestNamespaceName}/{S3FileContentHash}?storageLayers={storageLayerName}").ReturnsResponse(HttpStatusCode.OK, S3FileContent).Verifiable();
+                    _otherHordeStorageMock.SetupRequest($"http://other-horde-storage-instance.com/api/v1/blobs/{TestNamespaceName}/{OtherHordeStorageContentHash}?storageLayers={storageLayerName}").ReturnsResponse(HttpStatusCode.OK, OtherHordeStorageContent).Verifiable();
 
                     IHttpClientFactory mockFactory = _otherHordeStorageMock.CreateClientFactory();
 
@@ -181,7 +183,8 @@ namespace Horde.Storage.FunctionalTests.Storage
             getResponse.EnsureSuccessStatusCode();
             CollectionAssert.AreEqual(payload, actual: await getResponse.Content.ReadAsByteArrayAsync());
 
-            _otherHordeStorageMock!.VerifyRequest(HttpMethod.Get, new Uri($"http://other-horde-storage-instance.com/api/v1/blobs/{TestNamespaceName}/{OtherHordeStorageContentHash}?storageLayers=Filesystem"));
+            string storageLayerName = nameof(FileSystemStore);
+            _otherHordeStorageMock!.VerifyRequest(HttpMethod.Get, new Uri($"http://other-horde-storage-instance.com/api/v1/blobs/{TestNamespaceName}/{OtherHordeStorageContentHash}?storageLayers={storageLayerName}"));
             _otherHordeStorageMock!.VerifyNoOtherCalls();
         }
     }
