Horde: Add more logging to LogFileService to debug performance problems

diff --git a/Engine/Source/Programs/Horde/Horde.Build/Logs/LogFileService.cs b/Engine/Source/Programs/Horde/Horde.Build/Logs/LogFileService.cs
--- a/Engine/Source/Programs/Horde/Horde.Build/Logs/LogFileService.cs
+++ b/Engine/Source/Programs/Horde/Horde.Build/Logs/LogFileService.cs
@@ -1014,7 +1014,7 @@ namespace Horde.Build.Logs
 			
 			// Get all the chunks older than 20 minutes
 			List<(LogId, long)> flushChunks = await _builder.TouchChunksAsync(TimeSpan.FromMinutes(10.0));
-			_logger.LogDebug("Performing incremental flush of log builder ({NumChunks} chunks)", flushChunks.Count);
+			_logger.LogInformation("Performing incremental flush of log builder ({NumChunks} chunks)", flushChunks.Count);
 
 			scope.Span.SetTag("numChunks", flushChunks.Count);
 
@@ -1079,6 +1079,7 @@ namespace Horde.Build.Logs
 		private void WriteCompleteChunks(List<(LogId, long)> chunksToWrite, bool bCreateIndex)
 		{
 			using IScope scope = GlobalTracer.Instance.BuildSpan("LogFileService.WriteCompleteChunks").StartActive();
+			int numTasksCreated = 0;
 			
 			foreach (IGrouping<LogId, long> group in chunksToWrite.GroupBy(x => x.Item1, x => x.Item2))
 			{
@@ -1101,12 +1102,16 @@ namespace Horde.Build.Logs
 				if (offsets.Count > 0)
 				{
 					Task task = Task.Run(() => WriteCompleteChunksForLogAsync(logId, offsets, bCreateIndex));
+					numTasksCreated++;
 					lock (_writeLock)
 					{
 						_writeTasks.Add(task);
 					}
 				}
 			}
+
+			scope.Span.SetTag("numWriteTasksCreated", numTasksCreated);
+			_logger.LogInformation("{NumWriteTasksCreated} write tasks created", numTasksCreated);
 		}
 
 		/// <summary>
