Horde Storage - Resolved issues running the docker compose for local deployment - some missing configuration but also fixes to allow scylla enough time to start (by adding retries to the connect). Also made us default to allowing all namespaces when auth is disabled (can be overriden if needed to).

diff --git a/Engine/Source/Programs/Horde/HordeStorage/Horde.Storage/Horde.Storage.Startup.cs b/Engine/Source/Programs/Horde/HordeStorage/Horde.Storage/Horde.Storage.Startup.cs
--- a/Engine/Source/Programs/Horde/HordeStorage/Horde.Storage/Horde.Storage.Startup.cs
+++ b/Engine/Source/Programs/Horde/HordeStorage/Horde.Storage/Horde.Storage.Startup.cs
@@ -8,6 +8,7 @@ using System.Net;
 using System.Net.Security;
 using System.Runtime.InteropServices;
 using System.Security.Authentication;
+using System.Threading;
 using Amazon.Extensions.NETCore.Setup;
 using Amazon.Runtime;
 using Amazon.S3;
@@ -352,8 +353,30 @@ namespace Horde.Storage
                 replicationStrategy = settings.KeyspaceReplicationStrategy;
                 Logger.Information("Scylla Replication strategy for replicated keyspace is set to {@ReplicationStrategy}", replicationStrategy);
             }
-            ISession replicatedSession = cluster.ConnectAndCreateDefaultKeyspaceIfNotExists(replicationStrategy);
 
+            ISession replicatedSession;
+            const int MaxRetryAttempts = 10;
+            int countOfAttempts = 0;
+            while(true)
+            {
+                try
+                {
+                    countOfAttempts += 1;
+                    replicatedSession = cluster.ConnectAndCreateDefaultKeyspaceIfNotExists(replicationStrategy);
+                    break;
+                }
+                catch (NoHostAvailableException)
+                {
+                    Logger.Warning("Failed to connect to scylla, waiting a while and will then retry. Attempt {0} of {1}", countOfAttempts, MaxRetryAttempts);
+                    // wait for a few seconds before attempt to connect again
+                    Thread.Sleep(5000);
+                    if (countOfAttempts >= MaxRetryAttempts)
+                    {
+                        Logger.Error("Failed to connect to Scylla after {0} attempts, giving up.", countOfAttempts);
+                        throw;
+                    }
+                }
+            }
             replicatedSession.Execute(new SimpleStatement("CREATE TYPE IF NOT EXISTS blob_identifier (hash blob)"));
             replicatedSession.UserDefinedTypes.Define(UdtMap.For<ScyllaBlobIdentifier>("blob_identifier", DefaultKeyspaceName));
 
diff --git a/Engine/Source/Programs/Horde/HordeStorage/Jupiter.Common/Implementation/NamespacePolicyResolver.cs b/Engine/Source/Programs/Horde/HordeStorage/Jupiter.Common/Implementation/NamespacePolicyResolver.cs
--- a/Engine/Source/Programs/Horde/HordeStorage/Jupiter.Common/Implementation/NamespacePolicyResolver.cs
+++ b/Engine/Source/Programs/Horde/HordeStorage/Jupiter.Common/Implementation/NamespacePolicyResolver.cs
@@ -1,6 +1,7 @@
 // Copyright Epic Games, Inc. All Rights Reserved.
 
 using System;
+using System.Collections.Generic;
 using EpicGames.Horde.Storage;
 using Microsoft.Extensions.Options;
 
@@ -17,13 +18,19 @@ namespace Jupiter.Common
         private readonly IOptionsMonitor<NamespaceSettings> _namespaceSettings;
         private readonly NamespacePolicy _internalNamespacePolicy;
 
-        public NamespacePolicyResolver(IOptionsMonitor<NamespaceSettings> namespaceSettings)
+        public NamespacePolicyResolver(IOptionsMonitor<NamespaceSettings> namespaceSettings, IOptionsMonitor<AuthSettings> authSettings)
         {
             _namespaceSettings = namespaceSettings;
             _internalNamespacePolicy = new NamespacePolicy()
             {
                 StoragePool = "persistent"
             };
+
+            // if auth is disabled add a default namespace policy that allows access to everything unless something else exists to override it
+            if (!authSettings.CurrentValue.Enabled)
+            {
+                _namespaceSettings.CurrentValue.Policies.TryAdd("*", new NamespacePolicy() { Acls = new List<AclEntry> { new() { Claims = new List<string> {"*"} } } });
+            }
         }
 
         public NamespacePolicy GetPoliciesForNs(NamespaceId ns)
