Horde: Missing file.

diff --git a/Engine/Source/Programs/Horde/Horde.Build/Perforce/ReplicationService.cs b/Engine/Source/Programs/Horde/Horde.Build/Perforce/ReplicationService.cs
--- a/Engine/Source/Programs/Horde/Horde.Build/Perforce/ReplicationService.cs
+++ b/Engine/Source/Programs/Horde/Horde.Build/Perforce/ReplicationService.cs
@@ -543,7 +543,7 @@ namespace Horde.Build.Perforce
 				UpdateHash();
 				if (_entry.Target != null)
 				{
-					await writer.WriteAsync(_entry.Target, cancellationToken);
+					await writer.WriteAsync(_entry, cancellationToken);
 				}
 			}
 
@@ -851,7 +851,7 @@ namespace Horde.Build.Perforce
 					}
 				}
 			}
-			await writer.WriteAsync(syncNode, cancellationToken);
+//			await writer.WriteAsync(syncNode, cancellationToken);
 
 			clientInfo.Change = change;
 
