Horde: Fix issue with older job step refs not having valid update issues flag
#rnx

diff --git a/Engine/Source/Programs/Horde/Horde.Build/Jobs/JobStepRefCollection.cs b/Engine/Source/Programs/Horde/Horde.Build/Jobs/JobStepRefCollection.cs
--- a/Engine/Source/Programs/Horde/Horde.Build/Jobs/JobStepRefCollection.cs
+++ b/Engine/Source/Programs/Horde/Horde.Build/Jobs/JobStepRefCollection.cs
@@ -160,7 +160,7 @@ namespace Horde.Build.Jobs
 
 			if (updateIssues != null)
 			{
-				filter &= filterBuilder.Eq(x => x.UpdateIssues, updateIssues);
+				filter &= filterBuilder.Ne(x => x.UpdateIssues, false);
 			}
 
 			return await _jobStepRefs.Find(filter).SortByDescending(x => x.Change).FirstOrDefaultAsync();
@@ -188,7 +188,7 @@ namespace Horde.Build.Jobs
 
 			if (updateIssues != null)
 			{
-				filter &= filterBuilder.Eq(x => x.UpdateIssues, updateIssues);
+				filter &= filterBuilder.Ne(x => x.UpdateIssues, false);
 			}		
 
 			return await _jobStepRefs.Find(filter).SortBy(x => x.Change).FirstOrDefaultAsync();
