Horde: Implement ranged reads for storage backends. Remove code to emulate ranged reads from StorageClientBase.

diff --git a/Engine/Source/Programs/Horde/Horde.Agent/Commands/Bundles/CreateCommand.cs b/Engine/Source/Programs/Horde/Horde.Agent/Commands/Bundles/CreateCommand.cs
--- a/Engine/Source/Programs/Horde/Horde.Agent/Commands/Bundles/CreateCommand.cs
+++ b/Engine/Source/Programs/Horde/Horde.Agent/Commands/Bundles/CreateCommand.cs
@@ -6,7 +6,6 @@ using System.Globalization;
 using System.IO;
 using System.Threading;
 using System.Threading.Tasks;
-using System.Timers;
 using EpicGames.Core;
 using EpicGames.Horde.Storage;
 using EpicGames.Horde.Storage.Nodes;
@@ -43,6 +42,13 @@ namespace Horde.Agent.Commands.Bundles
 				return Task.FromResult<Stream>(FileReference.Open(file, FileMode.Open, FileAccess.Read, FileShare.Read));
 			}
 
+			public override async Task<Stream> ReadBlobRangeAsync(BlobLocator id, int offset, int length, CancellationToken cancellationToken = default)
+			{
+				Stream stream = await ReadBlobAsync(id, cancellationToken);
+				stream.Seek(offset, SeekOrigin.Begin);
+				return stream;
+			}
+
 			public override async Task<BlobLocator> WriteBlobAsync(Stream stream, Utf8String prefix = default, CancellationToken cancellationToken = default)
 			{
 				BlobLocator id = BlobLocator.Create(HostId.Empty, prefix);
diff --git a/Engine/Source/Programs/Horde/Horde.Build.Tests/BlobStoreTests.cs b/Engine/Source/Programs/Horde/Horde.Build.Tests/BlobStoreTests.cs
--- a/Engine/Source/Programs/Horde/Horde.Build.Tests/BlobStoreTests.cs
+++ b/Engine/Source/Programs/Horde/Horde.Build.Tests/BlobStoreTests.cs
@@ -19,12 +19,12 @@ namespace Horde.Build.Tests
 	public class BlobStoreTests : TestSetup
 	{
 		IMemoryCache? _cache;
-		InMemoryBlobStore? _store;
+		MemoryBlobStore? _store;
 
 		IStorageClient CreateBlobStore()
 		{
 			_cache ??= new MemoryCache(new MemoryCacheOptions());
-			_store ??= new InMemoryBlobStore(_cache, NullLogger.Instance);
+			_store ??= new MemoryBlobStore(_cache, NullLogger.Instance);
 			return _store;
 		}
 
diff --git a/Engine/Source/Programs/Horde/Horde.Build.Tests/LogFileServiceTests.cs b/Engine/Source/Programs/Horde/Horde.Build.Tests/LogFileServiceTests.cs
--- a/Engine/Source/Programs/Horde/Horde.Build.Tests/LogFileServiceTests.cs
+++ b/Engine/Source/Programs/Horde/Horde.Build.Tests/LogFileServiceTests.cs
@@ -43,7 +43,7 @@ namespace Horde.Build.Tests
             ILogger<LogFileService> logger = _loggerFactory.CreateLogger<LogFileService>();
 
 			ILogBuilder logBuilder = new RedisLogBuilder(GetRedisServiceSingleton().ConnectionPool, NullLogger.Instance);
-			_logStorage = new PersistentLogStorage(new TransientStorageBackend().ForType<PersistentLogStorage>(), NullLogger<PersistentLogStorage>.Instance);
+			_logStorage = new PersistentLogStorage(new MemoryStorageBackend().ForType<PersistentLogStorage>(), NullLogger<PersistentLogStorage>.Instance);
 			_clock = new FakeClock();
 			_logFileService = new LogFileService(logFileCollection, null!, logBuilder, _logStorage, _clock, logger);
         }
diff --git a/Engine/Source/Programs/Horde/Horde.Build.Tests/TestSetup.cs b/Engine/Source/Programs/Horde/Horde.Build.Tests/TestSetup.cs
--- a/Engine/Source/Programs/Horde/Horde.Build.Tests/TestSetup.cs
+++ b/Engine/Source/Programs/Horde/Horde.Build.Tests/TestSetup.cs
@@ -239,9 +239,9 @@ namespace Horde.Build.Tests
 			services.AddSingleton<ConformTaskSource>();
 			services.AddSingleton<ICommitService, CommitService>();
 
-			services.AddSingleton<IStorageBackend<PersistentLogStorage>>(sp => new TransientStorageBackend().ForType<PersistentLogStorage>());
-			services.AddSingleton<IStorageBackend<ArtifactCollection>>(sp => new TransientStorageBackend().ForType<ArtifactCollection>());
-			services.AddSingleton<IStorageBackend<BasicStorageClient>>(sp => new TransientStorageBackend().ForType<BasicStorageClient>());
+			services.AddSingleton<IStorageBackend<PersistentLogStorage>>(sp => new MemoryStorageBackend().ForType<PersistentLogStorage>());
+			services.AddSingleton<IStorageBackend<ArtifactCollection>>(sp => new MemoryStorageBackend().ForType<ArtifactCollection>());
+			services.AddSingleton<IStorageBackend<BasicStorageClient>>(sp => new MemoryStorageBackend().ForType<BasicStorageClient>());
 
 			services.AddSingleton<ILegacyStorageClient, BasicStorageClient>();
 
diff --git a/Engine/Source/Programs/Horde/Horde.Build/Jobs/Artifacts/ArtifactCollection.cs b/Engine/Source/Programs/Horde/Horde.Build/Jobs/Artifacts/ArtifactCollection.cs
--- a/Engine/Source/Programs/Horde/Horde.Build/Jobs/Artifacts/ArtifactCollection.cs
+++ b/Engine/Source/Programs/Horde/Horde.Build/Jobs/Artifacts/ArtifactCollection.cs
@@ -199,7 +199,7 @@ namespace Horde.Build.Jobs.Artifacts
 		/// <returns>The chunk data</returns>
 		public async Task<System.IO.Stream> OpenArtifactReadStreamAsync(IArtifact artifact)
 		{
-			System.IO.Stream? stream = await _storageBackend.ReadAsync(GetPath(artifact.JobId, artifact.StepId, artifact.Name));
+			System.IO.Stream? stream = await _storageBackend.TryReadAsync(GetPath(artifact.JobId, artifact.StepId, artifact.Name));
 			if (stream == null)
 			{
 				throw new Exception($"Unable to get artifact {artifact.Id}");
diff --git a/Engine/Source/Programs/Horde/Horde.Build/ServerSettings.cs b/Engine/Source/Programs/Horde/Horde.Build/ServerSettings.cs
--- a/Engine/Source/Programs/Horde/Horde.Build/ServerSettings.cs
+++ b/Engine/Source/Programs/Horde/Horde.Build/ServerSettings.cs
@@ -30,18 +30,13 @@ namespace Horde.Build
 		/// <summary>
 		/// In-memory only (for testing)
 		/// </summary>
-		Transient,
-
-		/// <summary>
-		/// Relay to another server (useful for testing against prod)
-		/// </summary>
-		Relay,
+		Memory,
 	};
 
 	/// <summary>
 	/// Common settings for different storage backends
 	/// </summary>
-	public interface IStorageBackendOptions : IFileSystemStorageOptions, IAwsStorageOptions, IRelayStorageOptions
+	public interface IStorageBackendOptions : IFileSystemStorageOptions, IAwsStorageOptions
 	{
 		/// <summary>
 		/// The type of storage backend to use
@@ -77,12 +72,6 @@ namespace Horde.Build
 
 		/// <inheritdoc/>
 		public string? AwsRegion { get; set; }
-
-		/// <inheritdoc/>
-		public string? RelayServer { get; set; }
-
-		/// <inheritdoc/>
-		public string? RelayToken { get; set; }
 	}
 
 	/// <summary>
diff --git a/Engine/Source/Programs/Horde/Horde.Build/Startup.cs b/Engine/Source/Programs/Horde/Horde.Build/Startup.cs
--- a/Engine/Source/Programs/Horde/Horde.Build/Startup.cs
+++ b/Engine/Source/Programs/Horde/Horde.Build/Startup.cs
@@ -268,10 +268,8 @@ namespace Horde.Build
 					return new FileSystemStorageBackend(options);
 				case StorageBackendType.Aws:
 					return new AwsStorageBackend(sp.GetRequiredService<IConfiguration>(), options, sp.GetRequiredService<ILogger<AwsStorageBackend>>());
-				case StorageBackendType.Transient:
-					return new TransientStorageBackend();
-				case StorageBackendType.Relay:
-					return new RelayStorageBackend(options);
+				case StorageBackendType.Memory:
+					return new MemoryStorageBackend();
 				default:
 					throw new NotImplementedException();
 			}
diff --git a/Engine/Source/Programs/Horde/Horde.Build/Storage/Backends/AwsStorageBackend.cs b/Engine/Source/Programs/Horde/Horde.Build/Storage/Backends/AwsStorageBackend.cs
--- a/Engine/Source/Programs/Horde/Horde.Build/Storage/Backends/AwsStorageBackend.cs
+++ b/Engine/Source/Programs/Horde/Horde.Build/Storage/Backends/AwsStorageBackend.cs
@@ -257,7 +257,18 @@ namespace Horde.Build.Storage.Backends
 		}
 
 		/// <inheritdoc/>
-		public async Task<Stream?> ReadAsync(string path, CancellationToken cancellationToken)
+		public Task<Stream?> TryReadAsync(string path, CancellationToken cancellationToken)
+		{
+			return TryReadAsync(path, null, cancellationToken);
+		}
+
+		/// <inheritdoc/>
+		public Task<Stream?> TryReadAsync(string path, int offset, int length, CancellationToken cancellationToken)
+		{
+			return TryReadAsync(path, new ByteRange(offset, offset + length), cancellationToken);
+		}
+
+		async Task<Stream?> TryReadAsync(string path, ByteRange? byteRange, CancellationToken cancellationToken)
 		{
 			using IScope scope = GlobalTracer.Instance.BuildSpan("AwsStorageBackend.ReadAsync").StartActive();
 			scope.Span.SetTag("Path", path);
@@ -273,6 +284,7 @@ namespace Horde.Build.Storage.Backends
 				GetObjectRequest newGetRequest = new GetObjectRequest();
 				newGetRequest.BucketName = _options.AwsBucketName;
 				newGetRequest.Key = fullPath;
+				newGetRequest.ByteRange = byteRange;
 
 				response = await _client.GetObjectAsync(newGetRequest, cancellationToken);
 
diff --git a/Engine/Source/Programs/Horde/Horde.Build/Storage/Backends/FileStorageBackend.cs b/Engine/Source/Programs/Horde/Horde.Build/Storage/Backends/FileStorageBackend.cs
new file mode 100644
--- /dev/null
+++ b/Engine/Source/Programs/Horde/Horde.Build/Storage/Backends/FileStorageBackend.cs
@@ -0,0 +1,152 @@
+// Copyright Epic Games, Inc. All Rights Reserved.
+
+using System;
+using System.IO;
+using System.Threading;
+using System.Threading.Tasks;
+using EpicGames.Core;
+using OpenTracing;
+using OpenTracing.Util;
+
+namespace Horde.Build.Storage.Backends
+{
+	/// <summary>
+	/// Options for the filesystem backend
+	/// </summary>
+	public interface IFileSystemStorageOptions
+	{
+		/// <summary>
+		/// Base directory for storing files
+		/// </summary>
+		public string? BaseDir { get; }
+	}
+
+	/// <summary>
+	/// Storage backend using the filesystem
+	/// </summary>
+	public class FileSystemStorageBackend : IStorageBackend
+	{
+		/// <summary>
+		/// Base directory for log files
+		/// </summary>
+		private readonly DirectoryReference _baseDir;
+
+		/// <summary>
+		/// Unique identifier for this instance
+		/// </summary>
+		private readonly string _instanceId;
+
+		/// <summary>
+		/// Unique id for each write
+		/// </summary>
+		private int _uniqueId;
+
+		/// <summary>
+		/// Constructor
+		/// </summary>
+		/// <param name="options">Current Horde Settings</param>
+		public FileSystemStorageBackend(IFileSystemStorageOptions options)
+		{
+			_baseDir = DirectoryReference.Combine(Program.DataDir, options.BaseDir ?? "Storage");
+			_instanceId = Guid.NewGuid().ToString("N");
+			DirectoryReference.CreateDirectory(_baseDir);
+		}
+
+		/// <inheritdoc/>
+		public Task<Stream?> TryReadAsync(string path, CancellationToken cancellationToken)
+		{
+			using IScope scope = GlobalTracer.Instance.BuildSpan("FileSystemStorageBackend.ReadAsync").StartActive();
+			scope.Span.SetTag("Path", path);
+			
+			FileReference location = FileReference.Combine(_baseDir, path);
+			if (!FileReference.Exists(location))
+			{
+				return Task.FromResult<Stream?>(null);
+			}
+
+			try
+			{
+				return Task.FromResult<Stream?>(FileReference.Open(location, FileMode.Open, FileAccess.Read, FileShare.Read));
+			}
+			catch (DirectoryNotFoundException)
+			{
+				return Task.FromResult<Stream?>(null);
+			}
+			catch (FileNotFoundException)
+			{
+				return Task.FromResult<Stream?>(null);
+			}
+		}
+
+		/// <inheritdoc/>
+		public async Task<Stream?> TryReadAsync(string path, int offset, int length, CancellationToken cancellationToken)
+		{
+			Stream? stream = await TryReadAsync(path, cancellationToken);
+			if (stream != null)
+			{
+				stream.Seek(offset, SeekOrigin.Begin);
+			}
+			return stream;
+		}
+
+		/// <inheritdoc/>
+		public async Task WriteAsync(string path, Stream stream, CancellationToken cancellationToken)
+		{
+			using IScope scope = GlobalTracer.Instance.BuildSpan("FileSystemStorageBackend.WriteAsync").StartActive();
+			scope.Span.SetTag("Path", path);
+			
+			FileReference finalLocation = FileReference.Combine(_baseDir, path);
+			if (!FileReference.Exists(finalLocation))
+			{
+				// Write to a temp file first
+				int newUniqueId = Interlocked.Increment(ref _uniqueId);
+
+				DirectoryReference.CreateDirectory(finalLocation.Directory);
+				FileReference tempLocation = new FileReference($"{finalLocation}.{_instanceId}.{newUniqueId:x8}");
+
+				using (Stream outputStream = FileReference.Open(tempLocation, FileMode.Create, FileAccess.Write, FileShare.Read))
+				{
+					await stream.CopyToAsync(outputStream, cancellationToken);
+				}
+
+				// Move the temp file into place
+				try
+				{
+					FileReference.Move(tempLocation, finalLocation, true);
+				}
+				catch (IOException) // Already exists
+				{
+					if (FileReference.Exists(finalLocation))
+					{
+						FileReference.Delete(tempLocation);
+					}
+					else
+					{
+						throw;
+					}
+				}
+			}
+		}
+
+		/// <inheritdoc/>
+		public Task<bool> ExistsAsync(string path, CancellationToken cancellationToken)
+		{
+			using IScope scope = GlobalTracer.Instance.BuildSpan("FileSystemStorageBackend.ExistsAsync").StartActive();
+			scope.Span.SetTag("Path", path);
+			
+			FileReference location = FileReference.Combine(_baseDir, path);
+			return Task.FromResult(FileReference.Exists(location));
+		}
+
+		/// <inheritdoc/>
+		public Task DeleteAsync(string path, CancellationToken cancellationToken)
+		{
+			using IScope scope = GlobalTracer.Instance.BuildSpan("FileSystemStorageBackend.DeleteAsync").StartActive();
+			scope.Span.SetTag("Path", path);
+			
+			FileReference location = FileReference.Combine(_baseDir, path);
+			FileReference.Delete(location);
+			return Task.CompletedTask;
+		}
+	}
+}
diff --git a/Engine/Source/Programs/Horde/Horde.Build/Storage/Backends/FileSystemStorageBackend.cs b/Engine/Source/Programs/Horde/Horde.Build/Storage/Backends/FileSystemStorageBackend.cs
deleted file mode 100644
diff --git a/Engine/Source/Programs/Horde/Horde.Build/Storage/Backends/MemoryStorageBackend.cs b/Engine/Source/Programs/Horde/Horde.Build/Storage/Backends/MemoryStorageBackend.cs
new file mode 100644
--- /dev/null
+++ b/Engine/Source/Programs/Horde/Horde.Build/Storage/Backends/MemoryStorageBackend.cs
@@ -0,0 +1,68 @@
+// Copyright Epic Games, Inc. All Rights Reserved.
+
+using System.Collections.Concurrent;
+using System.IO;
+using System.Threading;
+using System.Threading.Tasks;
+
+namespace Horde.Build.Storage.Backends
+{
+	/// <summary>
+	/// In-memory implementation of ILogFileStorage
+	/// </summary>
+	public class MemoryStorageBackend : IStorageBackend
+	{
+		/// <summary>
+		/// Data storage
+		/// </summary>
+		readonly ConcurrentDictionary<string, byte[]> _pathToData = new ConcurrentDictionary<string, byte[]>();
+
+		/// <inheritdoc/>
+		public Task<Stream?> TryReadAsync(string path, CancellationToken cancellationToken)
+		{
+			byte[]? data;
+			if (_pathToData.TryGetValue(path, out data))
+			{
+				return Task.FromResult<Stream?>(new MemoryStream(data, false));
+			}
+			else
+			{
+				return Task.FromResult<Stream?>(null);
+			}
+		}
+
+		/// <inheritdoc/>
+		public async Task<Stream?> TryReadAsync(string path, int offset, int length, CancellationToken cancellationToken)
+		{
+			Stream? stream = await TryReadAsync(path, cancellationToken);
+			if (stream != null)
+			{
+				stream.Seek(offset, SeekOrigin.Begin);
+			}
+			return stream;
+		}
+
+		/// <inheritdoc/>
+		public async Task WriteAsync(string path, Stream stream, CancellationToken cancellationToken)
+		{
+			using (MemoryStream buffer = new MemoryStream())
+			{
+				await stream.CopyToAsync(buffer, cancellationToken);
+				_pathToData[path] = buffer.ToArray();
+			}
+		}
+
+		/// <inheritdoc/>
+		public Task<bool> ExistsAsync(string path, CancellationToken cancellationToken)
+		{
+			return Task.FromResult(_pathToData.ContainsKey(path));
+		}
+
+		/// <inheritdoc/>
+		public Task DeleteAsync(string path, CancellationToken cancellationToken)
+		{
+			_pathToData.TryRemove(path, out _);
+			return Task.CompletedTask;
+		}
+	}
+}
diff --git a/Engine/Source/Programs/Horde/Horde.Build/Storage/Backends/RelayStorageBackend.cs b/Engine/Source/Programs/Horde/Horde.Build/Storage/Backends/RelayStorageBackend.cs
deleted file mode 100644
diff --git a/Engine/Source/Programs/Horde/Horde.Build/Storage/Backends/TransientStorageBackend.cs b/Engine/Source/Programs/Horde/Horde.Build/Storage/Backends/TransientStorageBackend.cs
deleted file mode 100644
diff --git a/Engine/Source/Programs/Horde/Horde.Build/Storage/IStorageBackend.cs b/Engine/Source/Programs/Horde/Horde.Build/Storage/IStorageBackend.cs
--- a/Engine/Source/Programs/Horde/Horde.Build/Storage/IStorageBackend.cs
+++ b/Engine/Source/Programs/Horde/Horde.Build/Storage/IStorageBackend.cs
@@ -14,12 +14,22 @@ namespace Horde.Build.Storage
 	public interface IStorageBackend
 	{
 		/// <summary>
-		/// Opens a read stream for the given path.
+		/// Attempts to open a read stream for the given path.
 		/// </summary>
 		/// <param name="path">Relative path within the bucket</param>
 		/// <param name="cancellationToken">Cancellation token for the operation</param>
 		/// <returns></returns>
-		Task<Stream?> ReadAsync(string path, CancellationToken cancellationToken = default);
+		Task<Stream?> TryReadAsync(string path, CancellationToken cancellationToken = default);
+
+		/// <summary>
+		/// Attempts to open a read stream for the given path.
+		/// </summary>
+		/// <param name="path">Relative path within the bucket</param>
+		/// <param name="offset">Offset to start reading from</param>
+		/// <param name="length">Length of data to read</param>
+		/// <param name="cancellationToken">Cancellation token for the operation</param>
+		/// <returns></returns>
+		Task<Stream?> TryReadAsync(string path, int offset, int length, CancellationToken cancellationToken = default);
 
 		/// <summary>
 		/// Writes a stream to the given path. If the stream throws an exception during read, the write will be aborted.
@@ -74,7 +84,10 @@ namespace Horde.Build.Storage
 			public StorageBackend(IStorageBackend inner) => _inner = inner;
 
 			/// <inheritdoc/>
-			public Task<Stream?> ReadAsync(string path, CancellationToken cancellationToken) => _inner.ReadAsync(path, cancellationToken);
+			public Task<Stream?> TryReadAsync(string path, CancellationToken cancellationToken) => _inner.TryReadAsync(path, cancellationToken);
+
+			/// <inheritdoc/>
+			public Task<Stream?> TryReadAsync(string path, int offset, int length, CancellationToken cancellationToken) => _inner.TryReadAsync(path, offset, length, cancellationToken);
 
 			/// <inheritdoc/>
 			public Task WriteAsync(string path, Stream stream, CancellationToken cancellationToken) => _inner.WriteAsync(path, stream, cancellationToken);
@@ -106,7 +119,7 @@ namespace Horde.Build.Storage
 		/// <returns></returns>
 		public static async Task<ReadOnlyMemory<byte>?> ReadBytesAsync(this IStorageBackend storageBackend, string path, CancellationToken cancellationToken = default)
 		{
-			using (Stream? inputStream = await storageBackend.ReadAsync(path, cancellationToken))
+			using (Stream? inputStream = await storageBackend.TryReadAsync(path, cancellationToken))
 			{
 				if (inputStream == null)
 				{
diff --git a/Engine/Source/Programs/Horde/Horde.Build/Storage/Services/BasicStorageClient.cs b/Engine/Source/Programs/Horde/Horde.Build/Storage/Services/BasicStorageClient.cs
--- a/Engine/Source/Programs/Horde/Horde.Build/Storage/Services/BasicStorageClient.cs
+++ b/Engine/Source/Programs/Horde/Horde.Build/Storage/Services/BasicStorageClient.cs
@@ -62,7 +62,7 @@ namespace Horde.Build.Storage.Services
 		{
 			string path = GetFullBlobPath(namespaceId, hash);
 
-			Stream? stream = await _storageBackend.ReadAsync(path, cancellationToken);
+			Stream? stream = await _storageBackend.TryReadAsync(path, cancellationToken);
 			if(stream == null)
 			{
 				throw new BlobNotFoundException(namespaceId, hash);
diff --git a/Engine/Source/Programs/Horde/Horde.Build/Storage/StorageConfig.cs b/Engine/Source/Programs/Horde/Horde.Build/Storage/StorageConfig.cs
--- a/Engine/Source/Programs/Horde/Horde.Build/Storage/StorageConfig.cs
+++ b/Engine/Source/Programs/Horde/Horde.Build/Storage/StorageConfig.cs
@@ -43,36 +43,10 @@ namespace Horde.Build.Storage
 		public List<NamespaceConfig> Namespaces { get; set; } = new List<NamespaceConfig>();
 	}
 
-	/// <summary>
-	/// Types of storage backend to use
-	/// </summary>
-	public enum BackendType
-	{
-		/// <summary>
-		/// Local filesystem
-		/// </summary>
-		FileSystem,
-
-		/// <summary>
-		/// AWS S3
-		/// </summary>
-		Aws,
-
-		/// <summary>
-		/// In-memory only (for testing)
-		/// </summary>
-		Transient,
-
-		/// <summary>
-		/// Relay to another server (useful for testing against prod)
-		/// </summary>
-		Relay,
-	};
-
 	/// <summary>
 	/// Common settings object for different providers
 	/// </summary>
-	public class BackendConfig : IFileSystemStorageOptions, IAwsStorageOptions, IRelayStorageOptions
+	public class BackendConfig : IStorageBackendOptions
 	{
 		/// <summary>
 		/// The storage backend ID
@@ -86,7 +60,7 @@ namespace Horde.Build.Storage
 		public BackendId Base { get; set; }
 
 		/// <inheritdoc/>
-		public BackendType? Type { get; set; }
+		public StorageBackendType? Type { get; set; }
 
 		/// <inheritdoc/>
 		public string? BaseDir { get; set; }
diff --git a/Engine/Source/Programs/Horde/Horde.Build/Storage/StorageService.cs b/Engine/Source/Programs/Horde/Horde.Build/Storage/StorageService.cs
--- a/Engine/Source/Programs/Horde/Horde.Build/Storage/StorageService.cs
+++ b/Engine/Source/Programs/Horde/Horde.Build/Storage/StorageService.cs
@@ -73,11 +73,26 @@ namespace Horde.Build.Storage
 
 			#region Blobs
 
+			/// <inheritdoc/>
 			public override async Task<Stream> ReadBlobAsync(BlobLocator locator, CancellationToken cancellationToken = default)
 			{
 				string path = GetBlobPath(locator.BlobId);
 
-				Stream? stream = await _backend.ReadAsync(path, cancellationToken);
+				Stream? stream = await _backend.TryReadAsync(path, cancellationToken);
+				if (stream == null)
+				{
+					throw new StorageException($"Unable to read data from {path}");
+				}
+
+				return stream;
+			}
+
+			/// <inheritdoc/>
+			public override async Task<Stream> ReadBlobRangeAsync(BlobLocator locator, int offset, int length, CancellationToken cancellationToken = default)
+			{
+				string path = GetBlobPath(locator.BlobId);
+
+				Stream? stream = await _backend.TryReadAsync(path, offset, length, cancellationToken);
 				if (stream == null)
 				{
 					throw new StorageException($"Unable to read data from {path}");
@@ -290,16 +305,14 @@ namespace Horde.Build.Storage
 		/// <returns>New storage backend instance</returns>
 		IStorageBackend CreateStorageBackend(BackendConfig config)
 		{
-			switch (config.Type ?? BackendType.FileSystem)
+			switch (config.Type ?? StorageBackendType.FileSystem)
 			{
-				case BackendType.FileSystem:
+				case StorageBackendType.FileSystem:
 					return new FileSystemStorageBackend(config);
-				case BackendType.Aws:
+				case StorageBackendType.Aws:
 					return new AwsStorageBackend(_serviceProvider.GetRequiredService<IConfiguration>(), config, _serviceProvider.GetRequiredService<ILogger<AwsStorageBackend>>());
-				case BackendType.Transient:
-					return new TransientStorageBackend();
-				case BackendType.Relay:
-					return new RelayStorageBackend(config);
+				case StorageBackendType.Memory:
+					return new MemoryStorageBackend();
 				default:
 					throw new NotImplementedException();
 			}
diff --git a/Engine/Source/Programs/Shared/EpicGames.Horde.Tests/BundleStoreTests.cs b/Engine/Source/Programs/Shared/EpicGames.Horde.Tests/BundleStoreTests.cs
--- a/Engine/Source/Programs/Shared/EpicGames.Horde.Tests/BundleStoreTests.cs
+++ b/Engine/Source/Programs/Shared/EpicGames.Horde.Tests/BundleStoreTests.cs
@@ -27,7 +27,7 @@ namespace EpicGames.Horde.Tests
 		{
 			using IMemoryCache cache = new MemoryCache(new MemoryCacheOptions());
 
-			InMemoryBlobStore blobStore = new InMemoryBlobStore(cache, NullLogger.Instance);
+			MemoryBlobStore blobStore = new MemoryBlobStore(cache, NullLogger.Instance);
 			await TestTreeAsync(blobStore, new TreeOptions { MaxBlobSize = 1024 * 1024 });
 
 			Assert.AreEqual(1, blobStore.Blobs.Count);
@@ -39,7 +39,7 @@ namespace EpicGames.Horde.Tests
 		{
 			using IMemoryCache cache = new MemoryCache(new MemoryCacheOptions());
 
-			InMemoryBlobStore blobStore = new InMemoryBlobStore(cache, NullLogger.Instance);
+			MemoryBlobStore blobStore = new MemoryBlobStore(cache, NullLogger.Instance);
 			await TestTreeAsync(blobStore, new TreeOptions { MaxBlobSize = 1 });
 
 			Assert.AreEqual(5, blobStore.Blobs.Count);
@@ -70,7 +70,7 @@ namespace EpicGames.Horde.Tests
 			}
 		}
 
-		static async Task TestTreeAsync(InMemoryBlobStore store, TreeOptions options)
+		static async Task TestTreeAsync(MemoryBlobStore store, TreeOptions options)
 		{
 			// Generate a tree
 			{
@@ -132,7 +132,7 @@ namespace EpicGames.Horde.Tests
 		public async Task DirectoryNodesAsync()
 		{
 			using IMemoryCache cache = new MemoryCache(new MemoryCacheOptions());
-			InMemoryBlobStore store = new InMemoryBlobStore(cache, NullLogger.Instance);
+			MemoryBlobStore store = new MemoryBlobStore(cache, NullLogger.Instance);
 
 			// Generate a tree
 			{
@@ -168,7 +168,7 @@ namespace EpicGames.Horde.Tests
 		public async Task FileNodesAsync()
 		{
 			using IMemoryCache cache = new MemoryCache(new MemoryCacheOptions());
-			InMemoryBlobStore store = new InMemoryBlobStore(cache, NullLogger.Instance);
+			MemoryBlobStore store = new MemoryBlobStore(cache, NullLogger.Instance);
 
 			// Generate a tree
 			{
@@ -213,7 +213,7 @@ namespace EpicGames.Horde.Tests
 		{
 			using IMemoryCache cache = new MemoryCache(new MemoryCacheOptions());
 
-			InMemoryBlobStore store = new InMemoryBlobStore(cache, NullLogger.Instance);
+			MemoryBlobStore store = new MemoryBlobStore(cache, NullLogger.Instance);
 			TreeWriter writer = new TreeWriter(store, new TreeOptions());
 
 			byte[] chunk = RandomNumberGenerator.GetBytes(256);
@@ -289,7 +289,7 @@ namespace EpicGames.Horde.Tests
 		public async Task LargeFileTestAsync()
 		{
 			using IMemoryCache cache = new MemoryCache(new MemoryCacheOptions());
-			InMemoryBlobStore store = new InMemoryBlobStore(cache, NullLogger.Instance);
+			MemoryBlobStore store = new MemoryBlobStore(cache, NullLogger.Instance);
 
 			const int length = 1024;
 			const int copies = 4096;
diff --git a/Engine/Source/Programs/Shared/EpicGames.Horde.Tests/BundleTests.cs b/Engine/Source/Programs/Shared/EpicGames.Horde.Tests/BundleTests.cs
--- a/Engine/Source/Programs/Shared/EpicGames.Horde.Tests/BundleTests.cs
+++ b/Engine/Source/Programs/Shared/EpicGames.Horde.Tests/BundleTests.cs
@@ -20,12 +20,12 @@ namespace EpicGames.Horde.Tests
 	public sealed class BundleTests : IDisposable
 	{
 		readonly IMemoryCache _cache;
-		readonly InMemoryBlobStore _storage;
+		readonly MemoryBlobStore _storage;
 
 		public BundleTests()
 		{
 			_cache = new MemoryCache(new MemoryCacheOptions());
-			_storage = new InMemoryBlobStore(_cache, NullLogger.Instance);
+			_storage = new MemoryBlobStore(_cache, NullLogger.Instance);
 		}
 
 		public void Dispose()
@@ -166,7 +166,7 @@ namespace EpicGames.Horde.Tests
 		[TestMethod]
 		public async Task BasicTestDirectory()
 		{
-			InMemoryBlobStore store = _storage;
+			MemoryBlobStore store = _storage;
 
 			DirectoryNode root = new DirectoryNode(DirectoryFlags.None);
 			DirectoryNode node = root.AddDirectory("hello");
diff --git a/Engine/Source/Programs/Shared/EpicGames.Horde/Storage/Backends/InMemoryBlobStore.cs b/Engine/Source/Programs/Shared/EpicGames.Horde/Storage/Backends/InMemoryBlobStore.cs
deleted file mode 100644
diff --git a/Engine/Source/Programs/Shared/EpicGames.Horde/Storage/Backends/MemoryBlobStore.cs b/Engine/Source/Programs/Shared/EpicGames.Horde/Storage/Backends/MemoryBlobStore.cs
new file mode 100644
--- /dev/null
+++ b/Engine/Source/Programs/Shared/EpicGames.Horde/Storage/Backends/MemoryBlobStore.cs
@@ -0,0 +1,93 @@
+// Copyright Epic Games, Inc. All Rights Reserved.
+
+using System;
+using System.Buffers;
+using System.Collections.Concurrent;
+using System.Collections.Generic;
+using System.IO;
+using System.Threading;
+using System.Threading.Tasks;
+using EpicGames.Core;
+using Microsoft.Extensions.Caching.Memory;
+using Microsoft.Extensions.Logging;
+
+namespace EpicGames.Horde.Storage.Backends
+{
+	/// <summary>
+	/// Implementation of <see cref="IStorageClient"/> which stores data in memory. Not intended for production use.
+	/// </summary>
+	public class MemoryBlobStore : StorageClientBase
+	{
+		/// <summary>
+		/// Map of blob id to blob data
+		/// </summary>
+		readonly ConcurrentDictionary<BlobLocator, Bundle> _blobs = new ConcurrentDictionary<BlobLocator, Bundle>();
+
+		/// <summary>
+		/// Map of ref name to ref data
+		/// </summary>
+		readonly ConcurrentDictionary<RefName, NodeLocator> _refs = new ConcurrentDictionary<RefName, NodeLocator>();
+
+		/// <inheritdoc cref="_blobs"/>
+		public IReadOnlyDictionary<BlobLocator, Bundle> Blobs => _blobs;
+
+		/// <inheritdoc cref="_refs"/>
+		public IReadOnlyDictionary<RefName, NodeLocator> Refs => _refs;
+
+		/// <summary>
+		/// Constructor
+		/// </summary>
+		public MemoryBlobStore(IMemoryCache cache, ILogger logger) 
+			: base(cache, logger)
+		{
+		}
+
+		#region Blobs
+
+		/// <inheritdoc/>
+		public override Task<Stream> ReadBlobAsync(BlobLocator locator, CancellationToken cancellationToken = default)
+		{
+			Stream stream = new ReadOnlySequenceStream(_blobs[locator].AsSequence());
+			return Task.FromResult(stream);
+		}
+
+		/// <inheritdoc/>
+		public override Task<Stream> ReadBlobRangeAsync(BlobLocator locator, int offset, int length, CancellationToken cancellationToken = default)
+		{
+			Stream stream = new ReadOnlySequenceStream(_blobs[locator].AsSequence().Slice(offset));
+			return Task.FromResult(stream);
+		}
+
+		/// <inheritdoc/>
+		public override async Task<BlobLocator> WriteBlobAsync(Stream stream, Utf8String prefix = default, CancellationToken cancellationToken = default)
+		{
+			BlobLocator locator = BlobLocator.Create(HostId.Empty, prefix);
+			_blobs[locator] = await Bundle.FromStreamAsync(stream, cancellationToken);
+			return locator;
+		}
+
+		#endregion
+
+		#region Refs
+
+		/// <inheritdoc/>
+		public override Task DeleteRefAsync(RefName name, CancellationToken cancellationToken) => Task.FromResult(_refs.TryRemove(name, out _));
+
+		/// <inheritdoc/>
+		public override Task<NodeLocator> TryReadRefTargetAsync(RefName name, DateTime cacheTime = default, CancellationToken cancellationToken = default)
+		{
+			NodeLocator refTarget;
+			_refs.TryGetValue(name, out refTarget);
+			return Task.FromResult(refTarget);
+		}
+
+		/// <inheritdoc/>
+		public override Task WriteRefTargetAsync(RefName name, NodeLocator target, CancellationToken cancellationToken = default)
+		{
+			_refs[name] = target;
+			return Task.CompletedTask;
+		}
+
+		#endregion
+	}
+}
diff --git a/Engine/Source/Programs/Shared/EpicGames.Horde/Storage/StorageClientBase.cs b/Engine/Source/Programs/Shared/EpicGames.Horde/Storage/StorageClientBase.cs
--- a/Engine/Source/Programs/Shared/EpicGames.Horde/Storage/StorageClientBase.cs
+++ b/Engine/Source/Programs/Shared/EpicGames.Horde/Storage/StorageClientBase.cs
@@ -9,7 +9,6 @@ using System.Threading;
 using System.Threading.Tasks;
 using EpicGames.Core;
 using Microsoft.CodeAnalysis;
-using Microsoft.CodeAnalysis.CSharp.Syntax;
 using Microsoft.Extensions.Caching.Memory;
 using Microsoft.Extensions.Logging;
 
@@ -20,85 +19,6 @@ namespace EpicGames.Horde.Storage
 	/// </summary>
 	public abstract class StorageClientBase : IStorageClient
 	{
-		/// <summary>
-		/// Wraps a stream and returns a portion of it
-		/// </summary>
-		class RangedReadStream : Stream, IDisposable
-		{
-			readonly Stream _inner;
-			readonly int _length;
-
-			int _position;
-
-			public RangedReadStream(Stream inner, int length)
-			{
-				_inner = inner;
-				_length = length;
-			}
-
-			/// <inheritdoc/>
-			public override bool CanRead => true;
-
-			/// <inheritdoc/>
-			public override bool CanSeek => false;
-
-			/// <inheritdoc/>
-			public override bool CanWrite => false;
-
-			/// <inheritdoc/>
-			public override long Length => throw new NotSupportedException();
-
-			/// <inheritdoc/>
-			public override long Position { get => _position; set => throw new NotSupportedException(); }
-
-			/// <inheritdoc/>
-			protected override void Dispose(bool disposing)
-			{
-				base.Dispose(disposing);
-
-				if (disposing)
-				{
-					_inner.Dispose();
-				}
-			}
-
-			/// <inheritdoc/>
-			public override void Flush() => _inner.Flush();
-
-			/// <inheritdoc/>
-			public override int Read(Span<byte> buffer)
-			{
-				buffer = buffer.Slice(0, Math.Min(buffer.Length, _length - _position));
-				int read = _inner.Read(buffer);
-				_position += read;
-				return read;
-			}
-
-			/// <inheritdoc/>
-			public override int Read(byte[] buffer, int offset, int count) => Read(buffer.AsSpan(offset, count));
-
-			/// <inheritdoc/>
-			public override async ValueTask<int> ReadAsync(Memory<byte> buffer, CancellationToken cancellationToken = default)
-			{
-				buffer = buffer.Slice(0, Math.Min(buffer.Length, _length - _position));
-				int read = await _inner.ReadAsync(buffer, cancellationToken);
-				_position += read;
-				return read;
-			}
-
-			/// <inheritdoc/>
-			public override async Task<int> ReadAsync(byte[] buffer, int offset, int count, CancellationToken cancellationToken) => await ReadAsync(buffer.AsMemory(offset, count), cancellationToken);
-
-			/// <inheritdoc/>
-			public override long Seek(long offset, SeekOrigin origin) => throw new NotSupportedException();
-
-			/// <inheritdoc/>
-			public override void SetLength(long value) => throw new NotSupportedException();
-
-			/// <inheritdoc/>
-			public override void Write(byte[] buffer, int offset, int count) => throw new NotSupportedException();
-		}
-
 		/// <summary>
 		/// Describes a bundle export
 		/// </summary>
@@ -222,14 +142,16 @@ namespace EpicGames.Horde.Storage
 		public abstract Task<Stream> ReadBlobAsync(BlobLocator locator, CancellationToken cancellationToken = default);
 
 		/// <inheritdoc/>
-		public virtual async Task<Stream> ReadBlobRangeAsync(BlobLocator locator, int offset, int length, CancellationToken cancellationToken = default)
-		{
-			_logger.LogInformation("Reading {Locator} [{Offset},{Length}]", locator, offset, length);
-			Stream stream = await ReadBlobAsync(locator, cancellationToken);
-			stream.Seek(offset, SeekOrigin.Begin);
-			return new RangedReadStream(stream, Math.Min((int)(stream.Length - offset), length));
-		}
+		public abstract Task<Stream> ReadBlobRangeAsync(BlobLocator locator, int offset, int length, CancellationToken cancellationToken = default);
 
+		/// <summary>
+		/// Utility method to read a blob into a buffer
+		/// </summary>
+		/// <param name="locator">Blob location</param>
+		/// <param name="offset">Offset within the blob</param>
+		/// <param name="memory">Buffer to read into</param>
+		/// <param name="cancellationToken">Cancellation token for the operation</param>
+		/// <returns>The data that was read</returns>
 		async Task<Memory<byte>> ReadBlobRangeAsync(BlobLocator locator, int offset, Memory<byte> memory, CancellationToken cancellationToken = default)
 		{
 			using (Stream stream = await ReadBlobRangeAsync(locator, offset, memory.Length, cancellationToken))
