Horde: Fix typo in issue collection
#rnx

diff --git a/Engine/Source/Programs/Horde/Horde.Build/Issues/IssueCollection.cs b/Engine/Source/Programs/Horde/Horde.Build/Issues/IssueCollection.cs
--- a/Engine/Source/Programs/Horde/Horde.Build/Issues/IssueCollection.cs
+++ b/Engine/Source/Programs/Horde/Horde.Build/Issues/IssueCollection.cs
@@ -557,7 +557,7 @@ namespace Horde.Build.Issues
 				}
 				else
 				{
-					issueLogger.LogInformation("Issue was acknowledged by {UserId}", oldIssue.OwnerId);
+					issueLogger.LogInformation("Issue was acknowledged by {UserId}", newIssue.OwnerId);
 				}
 			}
 			if (newIssue.FixChange != oldIssue.FixChange)
