Horde: Add support for instance type to use when expanding pools

diff --git a/Engine/Source/Programs/Horde/Horde.Build.Tests/Fleet/AwsReuseFleetManagerTest.cs b/Engine/Source/Programs/Horde/Horde.Build.Tests/Fleet/AwsReuseFleetManagerTest.cs
new file mode 100644
--- /dev/null
+++ b/Engine/Source/Programs/Horde/Horde.Build.Tests/Fleet/AwsReuseFleetManagerTest.cs
@@ -0,0 +1,51 @@
+// Copyright Epic Games, Inc. All Rights Reserved.
+
+using System.Collections.Generic;
+using System.Threading;
+using System.Threading.Tasks;
+using Amazon.EC2;
+using Amazon.EC2.Model;
+using Horde.Build.Agents.Pools;
+using Horde.Build.Agents.Fleet;
+using Microsoft.VisualStudio.TestTools.UnitTesting;
+using Horde.Build.Agents;
+using Horde.Build.Agents.Fleet.Providers;
+using Microsoft.Extensions.DependencyInjection;
+using Microsoft.Extensions.Logging;
+
+namespace Horde.Build.Tests.Fleet
+{
+	[TestClass]
+	public class AwsReuseFleetManagerTest : TestSetup
+	{
+		[TestMethod]
+		public async Task ExpandOneAgent()
+		{
+			FakeAmazonEc2 fakeAmazonEc2 = new ();
+			Instance i = fakeAmazonEc2.AddInstance(FakeAmazonEc2.StateStopped, InstanceType.M5Large);
+
+			await ExpandPoolAsync(fakeAmazonEc2.Get(), 1, new());
+			Assert.AreEqual(InstanceType.M5Large, fakeAmazonEc2.Instances[i.InstanceId].InstanceType);
+			Assert.AreEqual(FakeAmazonEc2.StatePending, fakeAmazonEc2.Instances[i.InstanceId].State);
+		}
+		
+		[TestMethod]
+		public async Task ExpandWithInstanceTypeChange2()
+		{
+			FakeAmazonEc2 fakeAmazonEc2 = new ();
+			Instance i = fakeAmazonEc2.AddInstance(FakeAmazonEc2.StateStopped, InstanceType.M5Large);
+
+			await ExpandPoolAsync(fakeAmazonEc2.Get(), 1, new(new List<string> { InstanceType.M54xlarge }));
+			Assert.AreEqual(InstanceType.M54xlarge, fakeAmazonEc2.Instances[i.InstanceId].InstanceType);
+			Assert.AreEqual(FakeAmazonEc2.StatePending, fakeAmazonEc2.Instances[i.InstanceId].State);
+		}
+
+		private async Task ExpandPoolAsync(IAmazonEC2 ec2, int numRequestedInstances, AwsReuseFleetManagerSettings settings)
+		{
+			ILogger<AwsReuseFleetManager> logger = ServiceProvider.GetRequiredService<ILogger<AwsReuseFleetManager>>();
+			IPool pool = await PoolService.CreatePoolAsync("testPool", null, true, 0, 0, sizeStrategy: PoolSizeStrategy.NoOp);
+			AwsReuseFleetManager manager = new (ec2, AgentCollection, settings, logger);
+			await manager.ExpandPoolAsync(pool, new List<IAgent>(), numRequestedInstances, CancellationToken.None);
+		}
+	}
+}
\ No newline at end of file
diff --git a/Engine/Source/Programs/Horde/Horde.Build.Tests/Fleet/FakeAmazonEc2.cs b/Engine/Source/Programs/Horde/Horde.Build.Tests/Fleet/FakeAmazonEc2.cs
new file mode 100644
--- /dev/null
+++ b/Engine/Source/Programs/Horde/Horde.Build.Tests/Fleet/FakeAmazonEc2.cs
@@ -0,0 +1,117 @@
+// Copyright Epic Games, Inc. All Rights Reserved.
+
+using System;
+using System.Collections.Generic;
+using System.Linq;
+using System.Net;
+using System.Threading;
+using System.Threading.Tasks;
+using Amazon.EC2;
+using Amazon.EC2.Model;
+using Moq;
+
+namespace Horde.Build.Tests.Fleet;
+
+/// <summary>
+/// Fake implementation of the IAmazonEC2 interface
+/// Simplifies testing of external EC2 API calls.
+/// Uses a mock to avoid implementing the entire interface. Even if shimmed, it becomes many lines of no-op code.
+/// </summary>
+public class FakeAmazonEc2
+{
+	public static readonly InstanceState StatePending = new() { Code = 0, Name = "pending" };
+	public static readonly InstanceState StateStopped = new() { Code = 80, Name = "stopped" };
+	
+	private readonly Mock<IAmazonEC2> _mock;
+	private readonly Dictionary<string, Instance> _instances = new();
+	private int _instanceIdCounter;
+	
+	internal IReadOnlyDictionary<string, Instance> Instances => _instances;
+
+	public FakeAmazonEc2()
+	{
+		_mock = new Mock<IAmazonEC2>(MockBehavior.Strict);
+		_mock
+			.Setup(x => x.StartInstancesAsync(It.IsAny<StartInstancesRequest>(), It.IsAny<CancellationToken>()))
+			.Returns(StartInstancesAsync);
+		
+		_mock
+			.Setup(x => x.DescribeInstancesAsync(It.IsAny<DescribeInstancesRequest>(), It.IsAny<CancellationToken>()))
+			.Returns(DescribeInstancesAsync);
+		
+		_mock
+			.Setup(x => x.ModifyInstanceAttributeAsync(It.IsAny<ModifyInstanceAttributeRequest>(), It.IsAny<CancellationToken>()))
+			.Returns(ModifyInstanceAttributeAsync);
+	}
+
+	public IAmazonEC2 Get()
+	{
+		return _mock.Object;
+	}
+
+	public Instance AddInstance(InstanceState state, InstanceType type)
+	{
+		Instance i = new() { InstanceId = "bogus-instance-" + _instanceIdCounter++, State = state, InstanceType = type };
+		_instances.Add(i.InstanceId, i);
+		return i;
+	}
+	
+	public Instance? GetInstance(string instanceId)
+	{
+		return _instances[instanceId];
+	}
+
+	private Task<DescribeInstancesResponse> DescribeInstancesAsync(DescribeInstancesRequest request, CancellationToken cancellationToken)
+	{
+		return Task.FromResult(new DescribeInstancesResponse
+		{
+			Reservations = new () { new Reservation { Instances = _instances.Values.ToList() } },
+			HttpStatusCode = HttpStatusCode.OK
+		});
+	}
+	
+	private Task<StartInstancesResponse> StartInstancesAsync(StartInstancesRequest request, CancellationToken cancellationToken)
+	{
+		List<InstanceStateChange> stateChanges = new();
+		foreach (string instanceId in request.InstanceIds)
+		{
+			if (!_instances.ContainsKey(instanceId))
+			{
+				throw new ArgumentException("Unknown instance ID " + instanceId);
+			}
+
+			if (_instances[instanceId].State.Name == StateStopped.Name)
+			{
+				stateChanges.Add(new()
+				{
+					InstanceId = instanceId,
+					CurrentState = StatePending,
+					PreviousState = _instances[instanceId].State
+				});
+				_instances[instanceId].State = StatePending;
+			}
+		}
+		
+		return Task.FromResult(new StartInstancesResponse
+		{
+			StartingInstances = stateChanges,
+			HttpStatusCode = HttpStatusCode.OK
+		});
+	}
+	
+	private Task<ModifyInstanceAttributeResponse> ModifyInstanceAttributeAsync(ModifyInstanceAttributeRequest request, CancellationToken cancellationToken)
+	{
+		if (!_instances.ContainsKey(request.InstanceId))
+		{
+			throw new ArgumentException("Unknown instance ID " + request.InstanceId);			
+		}
+
+		_instances[request.InstanceId].InstanceType = request.InstanceType;
+		
+		return Task.FromResult(new ModifyInstanceAttributeResponse()
+		{
+			ContentLength = 0,
+			HttpStatusCode = HttpStatusCode.OK
+		});
+	}
+}
\ No newline at end of file
diff --git a/Engine/Source/Programs/Horde/Horde.Build/Agents/Fleet/AutoscaleServiceV2.cs b/Engine/Source/Programs/Horde/Horde.Build/Agents/Fleet/AutoscaleServiceV2.cs
--- a/Engine/Source/Programs/Horde/Horde.Build/Agents/Fleet/AutoscaleServiceV2.cs
+++ b/Engine/Source/Programs/Horde/Horde.Build/Agents/Fleet/AutoscaleServiceV2.cs
@@ -186,9 +186,11 @@ namespace Horde.Build.Agents.Fleet
 			int desiredAgentCount = poolSizeData.DesiredAgentCount.Value;
 			int deltaAgentCount = desiredAgentCount - currentAgentCount;
 
-			_logger.LogInformation("{PoolName,-48} Current={Current,4} Target={Target,4} Delta={Delta,4} Status={Status}", pool.Name, currentAgentCount, desiredAgentCount, deltaAgentCount, poolSizeData.StatusMessage);
-
 			IFleetManager fleetManager = _fleetManagerFactory(poolSizeData.Pool);
+
+			_logger.LogInformation("{PoolName,-48} Current={Current,4} Target={Target,4} Delta={Delta,4} Status={Status} FleetManager={FleetManager}",
+				pool.Name, currentAgentCount, desiredAgentCount, deltaAgentCount, poolSizeData.StatusMessage, fleetManager.GetType().Name);
+			
 			try
 			{
 				using IScope scope = GlobalTracer.Instance.BuildSpan("ScalingPool").StartActive();
@@ -306,7 +308,7 @@ namespace Horde.Build.Agents.Fleet
 				case FleetManagerType.Aws:
 					return new AwsFleetManager(_ec2, _agentCollection, DeserializeSettings<AwsFleetManagerSettings>(config), _loggerFactory.CreateLogger<AwsFleetManager>());
 				case FleetManagerType.AwsReuse:
-					return new AwsReuseFleetManager(_ec2, _agentCollection, _loggerFactory.CreateLogger<AwsReuseFleetManager>());
+					return new AwsReuseFleetManager(_ec2, _agentCollection, DeserializeSettings<AwsReuseFleetManagerSettings>(config), _loggerFactory.CreateLogger<AwsReuseFleetManager>());
 				case FleetManagerType.AwsAsg:
 					return new AwsAsgFleetManager(_awsAutoScaling, DeserializeSettings<AwsAsgSettings>(config), _loggerFactory.CreateLogger<AwsAsgFleetManager>());	
 				default:
diff --git a/Engine/Source/Programs/Horde/Horde.Build/Agents/Fleet/Providers/AwsReuseFleetManager.cs b/Engine/Source/Programs/Horde/Horde.Build/Agents/Fleet/Providers/AwsReuseFleetManager.cs
--- a/Engine/Source/Programs/Horde/Horde.Build/Agents/Fleet/Providers/AwsReuseFleetManager.cs
+++ b/Engine/Source/Programs/Horde/Horde.Build/Agents/Fleet/Providers/AwsReuseFleetManager.cs
@@ -14,6 +14,27 @@ using OpenTracing.Util;
 
 namespace Horde.Build.Agents.Fleet.Providers
 {
+	/// <summary>
+	/// Settings for AWS fleet manager that reuses instances
+	/// </summary>
+	public class AwsReuseFleetManagerSettings
+	{
+		/// <summary>
+		/// Prioritized list of instance types to use waking up an instance
+		/// If set to null, the instance will be awakened with its last instance type.
+		/// </summary>
+		public List<string>? InstanceTypes { get; }
+
+		/// <summary>
+		/// Constructor
+		/// </summary>
+		/// <param name="instanceTypes"></param>
+		public AwsReuseFleetManagerSettings(List<string>? instanceTypes = null)
+		{
+			InstanceTypes = instanceTypes;
+		}
+	}
+	
 	/// <summary>
 	/// Fleet manager for handling AWS EC2 instances
 	/// Will start already existing but stopped instances to reuse existing EBS disks.
@@ -23,31 +44,33 @@ namespace Horde.Build.Agents.Fleet.Providers
 	{
 		private readonly IAmazonEC2 _ec2;
 		private readonly IAgentCollection _agentCollection;
+		private readonly AwsReuseFleetManagerSettings _settings;
 		private readonly ILogger _logger;
 
 		/// <summary>
 		/// Constructor
 		/// </summary>
-		public AwsReuseFleetManager(IAmazonEC2 ec2, IAgentCollection agentCollection, ILogger<AwsReuseFleetManager> logger)
+		public AwsReuseFleetManager(IAmazonEC2 ec2, IAgentCollection agentCollection, AwsReuseFleetManagerSettings settings, ILogger<AwsReuseFleetManager> logger)
 		{
 			_ec2 = ec2;
 			_agentCollection = agentCollection;
+			_settings = settings;
 			_logger = logger;
 		}
 
 		/// <inheritdoc/>
-		public async Task ExpandPoolAsync(IPool pool, IReadOnlyList<IAgent> agents, int count, CancellationToken cancellationToken)
+		public async Task ExpandPoolAsync(IPool pool, IReadOnlyList<IAgent> agents, int requestedInstancesCount, CancellationToken cancellationToken)
 		{
 			using IScope scope = GlobalTracer.Instance.BuildSpan("ExpandPool").StartActive();
 			scope.Span.SetTag("poolName", pool.Name);
 			scope.Span.SetTag("numAgents", agents.Count);
-			scope.Span.SetTag("count", count);
+			scope.Span.SetTag("count", requestedInstancesCount);
 
 			DescribeInstancesResponse describeResponse;
 			using (IScope describeScope = GlobalTracer.Instance.BuildSpan("DescribeInstances").StartActive())
 			{
 				// Find stopped instances in the correct pool
-				DescribeInstancesRequest describeRequest = new DescribeInstancesRequest();
+				DescribeInstancesRequest describeRequest = new ();
 				describeRequest.Filters = new List<Filter>();
 				describeRequest.Filters.Add(new Filter("instance-state-name", new List<string> { InstanceStateName.Stopped.Value }));
 				describeRequest.Filters.Add(new Filter("tag:" + AwsFleetManager.PoolTagName, new List<string> { pool.Name }));
@@ -56,13 +79,38 @@ namespace Horde.Build.Agents.Fleet.Providers
 				describeScope.Span.SetTag("res.numReservations", describeResponse.Reservations.Count);
 			}
 
+			IEnumerable<Instance> stoppedInstances = describeResponse.Reservations.SelectMany(x => x.Instances);
+
+			if (_settings.InstanceTypes is { Count: >= 1 })
+			{
+				InstanceType newInstanceType = InstanceType.FindValue(_settings.InstanceTypes[0]);
+				foreach (Instance stoppedInstance in stoppedInstances)
+				{
+					if (stoppedInstance.InstanceType == newInstanceType) { continue; }
+					using IScope modifyScope = GlobalTracer.Instance.BuildSpan("ModifyInstanceAttribute").StartActive();
+
+					ModifyInstanceAttributeRequest request = new () { InstanceId = stoppedInstance.InstanceId, InstanceType = newInstanceType };
+					ModifyInstanceAttributeResponse response = await _ec2.ModifyInstanceAttributeAsync(request, cancellationToken);
+					if ((int)response.HttpStatusCode >= 200 && (int)response.HttpStatusCode <= 299)
+					{
+						_logger.LogInformation("Instance type for {InstanceId} in pool {PoolId} modified ({PrevInstanceType} -> {CurrentInstanceType})", stoppedInstance.InstanceId, pool.Id, stoppedInstance.InstanceType, newInstanceType);
+					}
+					else
+					{
+						_logger.LogError("Failed to modify instance type for {InstanceId} in pool {PoolId}. Wanted {InstanceType}. Status code {StatusCode}", stoppedInstance.InstanceId, pool.Id, newInstanceType, response.HttpStatusCode);
+					}
+				}
+			}
+
 			using (IScope startScope = GlobalTracer.Instance.BuildSpan("StartInstances").StartActive())
 			{
 				// Try to start the given instances
-				StartInstancesRequest startRequest = new StartInstancesRequest();
-				startRequest.InstanceIds.AddRange(describeResponse.Reservations.SelectMany(x => x.Instances).Select(x => x.InstanceId).Take(count));
+				StartInstancesRequest startRequest = new ();
+				startRequest.InstanceIds.AddRange(describeResponse.Reservations.SelectMany(x => x.Instances).Select(x => x.InstanceId).Take(requestedInstancesCount));
+				int stoppedInstancesCount = startRequest.InstanceIds.Count;
 				
 				startScope.Span.SetTag("req.instanceIds", String.Join(",", startRequest.InstanceIds));
+				int startedInstancesCount = 0;
 				if (startRequest.InstanceIds.Count > 0)
 				{
 					StartInstancesResponse startResponse = await _ec2.StartInstancesAsync(startRequest, cancellationToken);
@@ -72,15 +120,24 @@ namespace Horde.Build.Agents.Fleet.Providers
 					{
 						foreach (InstanceStateChange instanceChange in startResponse.StartingInstances)
 						{
-							_logger.LogInformation("Starting instance {InstanceId} for pool {PoolId} (prev state {PrevState}, current state {CurrentState}", instanceChange.InstanceId, pool.Id, instanceChange.PreviousState, instanceChange.CurrentState);
+							startedInstancesCount++;
+							_logger.LogInformation("Starting instance {InstanceId} for pool {PoolId} (prev state {PrevState}, current state {CurrentState}", instanceChange.InstanceId, pool.Id, instanceChange.PreviousState.Name, instanceChange.CurrentState.Name);
 						}
 					}
 				}
 
-				if (startRequest.InstanceIds.Count < count)
+				if (startRequest.InstanceIds.Count < requestedInstancesCount)
 				{
-					_logger.LogInformation("Unable to expand pool {PoolName} with the requested number of instances. " +
-					                      "Num requested instances to add {RequestedCount}. Actual instances started {InstancesStarted}", pool.Name, count, startRequest.InstanceIds.Count);
+					string reason = "";
+					if (stoppedInstancesCount < requestedInstancesCount) { reason += " Not enough stopped instances to accommodate the pool expansion."; }
+					if (startedInstancesCount < stoppedInstancesCount) { reason += " Not all instances were able to start."; }
+
+					_logger.LogInformation("Unable to expand pool {PoolName}.\n" + 
+					                       "Reason={Reason}\n" +  
+					                       "RequestedInstancesCount={RequestedCount}\n" +
+					                       "StoppedInstancesCount={StoppedInstancesCount}\n" +
+					                       "StartedInstancesCount={StartedInstancesCount}",
+						pool.Name, reason.Trim(), requestedInstancesCount, stoppedInstancesCount, startedInstancesCount);
 				}
 			}
 		}
