Horde: Fix incorrect mapping between Slack users and email addresses. Was returning user documents with the 'hidden' flag set. Also fix users self-assigning issues showing up as assignment by "default".

diff --git a/Engine/Source/Programs/Horde/Horde.Build/Notifications/Sinks/SlackNotificationSink.cs b/Engine/Source/Programs/Horde/Horde.Build/Notifications/Sinks/SlackNotificationSink.cs
--- a/Engine/Source/Programs/Horde/Horde.Build/Notifications/Sinks/SlackNotificationSink.cs
+++ b/Engine/Source/Programs/Horde/Horde.Build/Notifications/Sinks/SlackNotificationSink.cs
@@ -859,6 +859,8 @@ namespace Horde.Build.Notifications.Sinks
 				// Permalink to the summary text so we link inside the thread rather than just to the original message
 				string? permalink = await _slackClient.GetPermalinkAsync(summaryId);
 				await UpdateMessageStateAsync(state.Id, state.MessageId, permalink);
+
+				_issueService.Collection.GetLogger(issue.Id).LogInformation("Created Slack thread: {SlackLink}", permalink);
 			}
 
 			// Post a message containing the controls and status
@@ -2417,7 +2419,7 @@ namespace Horde.Build.Notifications.Sinks
 			}
 			else if (String.Equals(verb, "accept", StringComparison.Ordinal))
 			{
-				await _issueService.UpdateIssueAsync(issueId, ownerId: userId, acknowledged: true);
+				await _issueService.UpdateIssueAsync(issueId, ownerId: userId, nominatedById: userId, acknowledged: true);
 			}
 			else if (String.Equals(verb, "decline", StringComparison.Ordinal))
 			{
@@ -2483,7 +2485,7 @@ namespace Horde.Build.Notifications.Sinks
 					}
 				}
 
-				await _issueService.UpdateIssueAsync(issueId, acknowledged: true, ownerId: user.Id);
+				await _issueService.UpdateIssueAsync(issueId, acknowledged: true, ownerId: user.Id, nominatedById: user.Id);
 			}
 			else if (String.Equals(verb, "decline", StringComparison.Ordinal))
 			{
diff --git a/Engine/Source/Programs/Horde/Horde.Build/Users/UserCollectionV2.cs b/Engine/Source/Programs/Horde/Horde.Build/Users/UserCollectionV2.cs
--- a/Engine/Source/Programs/Horde/Horde.Build/Users/UserCollectionV2.cs
+++ b/Engine/Source/Programs/Horde/Horde.Build/Users/UserCollectionV2.cs
@@ -241,14 +241,16 @@ namespace Horde.Build.Users
 		public async Task<IUser?> FindUserByLoginAsync(string login)
 		{
 			string loginUpper = login.ToUpperInvariant();
-			return await _users.Find(x => x.LoginUpper == loginUpper).FirstOrDefaultAsync();
+			FilterDefinition<UserDocument> filter = Builders<UserDocument>.Filter.Eq(x => x.LoginUpper, loginUpper) & Builders<UserDocument>.Filter.Ne(x => x.Hidden, true);
+			return await _users.Find(filter).FirstOrDefaultAsync();
 		}
 
 		/// <inheritdoc/>
 		public async Task<IUser?> FindUserByEmailAsync(string email)
 		{
 			string emailUpper = email.ToUpperInvariant();
-			return await _users.Find(x => x.EmailUpper == emailUpper).FirstOrDefaultAsync();
+			FilterDefinition<UserDocument> filter = Builders<UserDocument>.Filter.Eq(x => x.EmailUpper, emailUpper) & Builders<UserDocument>.Filter.Ne(x => x.Hidden, true);
+			return await _users.Find(filter).FirstOrDefaultAsync();
 		}
 
 		/// <inheritdoc/>
