Horde: Implement an HTTP and filesystem storage client.

diff --git a/Engine/Source/Programs/Horde/Horde.Agent/Commands/Bundles/CreateCommand.cs b/Engine/Source/Programs/Horde/Horde.Agent/Commands/Bundles/CreateCommand.cs
--- a/Engine/Source/Programs/Horde/Horde.Agent/Commands/Bundles/CreateCommand.cs
+++ b/Engine/Source/Programs/Horde/Horde.Agent/Commands/Bundles/CreateCommand.cs
@@ -2,12 +2,12 @@
 
 using System;
 using System.Diagnostics;
-using System.Globalization;
-using System.IO;
+using System.Net.Http;
 using System.Threading;
 using System.Threading.Tasks;
 using EpicGames.Core;
 using EpicGames.Horde.Storage;
+using EpicGames.Horde.Storage.Backends;
 using EpicGames.Horde.Storage.Nodes;
 using Microsoft.Extensions.Caching.Memory;
 using Microsoft.Extensions.Logging;
@@ -16,100 +16,72 @@ namespace Horde.Agent.Commands.Bundles
 {
 	abstract class BundleCommandBase : Command
 	{
-		class FileBlobStore : StorageClientBase
+		protected interface IStorageClientOwner : IDisposable
 		{
-			readonly DirectoryReference _rootDir;
-			readonly ILogger _logger;
-
-			public FileBlobStore(DirectoryReference rootDir, IMemoryCache cache, ILogger logger)
-				: base(cache, logger)
-			{
-				_rootDir = rootDir;
-				_logger = logger;
-
-				DirectoryReference.CreateDirectory(_rootDir);
-			}
-
-			FileReference GetRefFile(RefName name) => FileReference.Combine(_rootDir, name.ToString() + ".ref");
-			FileReference GetBlobFile(BlobLocator id) => FileReference.Combine(_rootDir, id.Inner.ToString() + ".blob");
-
-			#region Blobs
+			public IStorageClient Store { get; }
+		}
 
-			public override Task<Stream> ReadBlobAsync(BlobLocator id, CancellationToken cancellationToken = default)
-			{
-				FileReference file = GetBlobFile(id);
-				_logger.LogInformation("Reading {File}", file);
-				return Task.FromResult<Stream>(FileReference.Open(file, FileMode.Open, FileAccess.Read, FileShare.Read));
-			}
+		class FileStorageClientOwner : IStorageClientOwner
+		{
+			readonly IMemoryCache _cache;
+			public IStorageClient Store { get; }
 
-			public override async Task<Stream> ReadBlobRangeAsync(BlobLocator id, int offset, int length, CancellationToken cancellationToken = default)
+			public FileStorageClientOwner(DirectoryReference storageDir, ILogger logger)
 			{
-				Stream stream = await ReadBlobAsync(id, cancellationToken);
-				stream.Seek(offset, SeekOrigin.Begin);
-				return stream;
+				_cache = new MemoryCache(new MemoryCacheOptions());
+				Store = new FileStorageClient(storageDir, _cache, logger);
 			}
 
-			public override async Task<BlobLocator> WriteBlobAsync(Stream stream, Utf8String prefix = default, CancellationToken cancellationToken = default)
+			public void Dispose()
 			{
-				BlobLocator id = BlobLocator.Create(HostId.Empty, prefix);
-				FileReference file = GetBlobFile(id);
-				DirectoryReference.CreateDirectory(file.Directory);
-				_logger.LogInformation("Writing {File}", file);
-
-				using (FileStream fileStream = FileReference.Open(file, FileMode.Create, FileAccess.Write, FileShare.ReadWrite))
-				{
-					await stream.CopyToAsync(fileStream, cancellationToken);
-				}
-
-				return id;
+				_cache.Dispose();
 			}
+		}
 
-			#endregion
-
-			#region Refs
-
-			public override Task DeleteRefAsync(RefName name, CancellationToken cancellationToken = default)
-			{
-				throw new NotImplementedException();
-			}
+		class HttpStorageClientOwner : IStorageClientOwner
+		{
+			readonly IMemoryCache _cache;
+			readonly HttpClient _httpClient;
+			public IStorageClient Store { get; }
 
-			public override async Task<NodeLocator> TryReadRefTargetAsync(RefName name, DateTime cacheTime = default, CancellationToken cancellationToken = default)
+			public HttpStorageClientOwner(NamespaceId namespaceId, Uri baseUri, ILogger logger)
 			{
-				FileReference file = GetRefFile(name);
-				if (!FileReference.Exists(file))
-				{
-					return default;
-				}
-
-				_logger.LogInformation("Reading {File}", file);
-				string text = await FileReference.ReadAllTextAsync(file, cancellationToken);
-
-				int hashIdx = text.IndexOf('#', StringComparison.Ordinal);
-				BlobLocator locator = new BlobLocator(text.Substring(0, hashIdx));
-				int exportIdx = Int32.Parse(text.Substring(hashIdx + 1), CultureInfo.InvariantCulture);
-
-				return new NodeLocator(locator, exportIdx);
+				_cache = new MemoryCache(new MemoryCacheOptions());
+				_httpClient = new HttpClient { BaseAddress = baseUri };
+				Store = new HttpStorageClient(namespaceId, _httpClient, _cache, logger);
 			}
 
-			public override async Task WriteRefTargetAsync(RefName name, NodeLocator target, CancellationToken cancellationToken = default)
+			public void Dispose()
 			{
-				FileReference file = GetRefFile(name);
-				DirectoryReference.CreateDirectory(file.Directory);
-				_logger.LogInformation("Writing {File}", file);
-				await FileReference.WriteAllTextAsync(file, $"{target.Blob}#{target.ExportIdx}");
+				_httpClient.Dispose();
+				_cache.Dispose();
 			}
-
-			#endregion
 		}
 
 		public static RefName DefaultRefName { get; } = new RefName("default-ref");
 
+		[CommandLine("-Http")]
+		public bool Http { get; set; }
+
+		[CommandLine("-Server=", Description = "Server to read from")]
+		public string Server { get; set; } = "https://localhost:5001";
+
+		[CommandLine("-Namespace=", Description = "Namespace to use for storage")]
+		public NamespaceId NamespaceId { get; set; } = new NamespaceId("default");
+
 		[CommandLine("-StorageDir=", Description = "Overrides the default storage server with a local directory")]
 		public DirectoryReference StorageDir { get; set; } = DirectoryReference.Combine(Program.AppDir, "bundles");
 
-		protected IStorageClient CreateStorageClient(IMemoryCache cache, ILogger logger)
+		protected IStorageClientOwner CreateStorageClient(ILogger logger)
 		{
-			return new FileBlobStore(StorageDir, cache, logger);
+			if (Http)
+			{
+				return new HttpStorageClientOwner(NamespaceId, new Uri(Server), logger);
+			}
+			else
+			{
+				return new FileStorageClientOwner(DirectoryReference.Combine(StorageDir, NamespaceId.ToString()), logger);
+			}
 		}
 	}
 
@@ -124,8 +96,8 @@ namespace Horde.Agent.Commands.Bundles
 
 		public override async Task<int> ExecuteAsync(ILogger logger)
 		{
-			using IMemoryCache cache = new MemoryCache(new MemoryCacheOptions());
-			IStorageClient store = CreateStorageClient(cache, logger);
+			using IStorageClientOwner storeOwner = CreateStorageClient(logger);
+			IStorageClient store = storeOwner.Store;
 
 			TreeWriter writer = new TreeWriter(store, prefix: RefName.Text);
 
diff --git a/Engine/Source/Programs/Horde/Horde.Agent/Commands/Bundles/DumpCommand.cs b/Engine/Source/Programs/Horde/Horde.Agent/Commands/Bundles/DumpCommand.cs
--- a/Engine/Source/Programs/Horde/Horde.Agent/Commands/Bundles/DumpCommand.cs
+++ b/Engine/Source/Programs/Horde/Horde.Agent/Commands/Bundles/DumpCommand.cs
@@ -22,8 +22,8 @@ namespace Horde.Agent.Commands.Bundles
 
 		public override async Task<int> ExecuteAsync(ILogger logger)
 		{
-			using IMemoryCache cache = new MemoryCache(new MemoryCacheOptions());
-			IStorageClient store = CreateStorageClient(cache, logger);
+			using IStorageClientOwner storeOwner = CreateStorageClient(logger);
+			IStorageClient store = storeOwner.Store;
 
 			if (BlobId == null)
 			{
diff --git a/Engine/Source/Programs/Horde/Horde.Agent/Commands/Bundles/ExtractCommand.cs b/Engine/Source/Programs/Horde/Horde.Agent/Commands/Bundles/ExtractCommand.cs
--- a/Engine/Source/Programs/Horde/Horde.Agent/Commands/Bundles/ExtractCommand.cs
+++ b/Engine/Source/Programs/Horde/Horde.Agent/Commands/Bundles/ExtractCommand.cs
@@ -6,7 +6,6 @@ using System.Threading.Tasks;
 using EpicGames.Core;
 using EpicGames.Horde.Storage;
 using EpicGames.Horde.Storage.Nodes;
-using Microsoft.Extensions.Caching.Memory;
 using Microsoft.Extensions.Logging;
 
 namespace Horde.Agent.Commands.Bundles
@@ -22,15 +21,15 @@ namespace Horde.Agent.Commands.Bundles
 
 		public override async Task<int> ExecuteAsync(ILogger logger)
 		{
-			using IMemoryCache cache = new MemoryCache(new MemoryCacheOptions());
-			IStorageClient store = CreateStorageClient(cache, logger);
+			using IStorageClientOwner storeOwner = CreateStorageClient(logger);
+			IStorageClient store = storeOwner.Store;
 
 			Stopwatch timer = Stopwatch.StartNew();
 
 			DirectoryNode node = await store.ReadNodeAsync<DirectoryNode>(RefName);
 			await node.CopyToDirectoryAsync(OutputDir.ToDirectoryInfo(), logger, CancellationToken.None);
 
-			logger.LogInformation("Elapsed: {Time}s", timer.Elapsed.Seconds);
+			logger.LogInformation("Elapsed: {Time}s", timer.Elapsed.TotalSeconds);
 			return 0;
 		}
 	}
diff --git a/Engine/Source/Programs/Horde/Horde.Build/Startup.cs b/Engine/Source/Programs/Horde/Horde.Build/Startup.cs
--- a/Engine/Source/Programs/Horde/Horde.Build/Startup.cs
+++ b/Engine/Source/Programs/Horde/Horde.Build/Startup.cs
@@ -843,6 +843,21 @@ namespace Horde.Build
 			}
 		}
 
+		public sealed class NamespaceIdBsonSerializer : SerializerBase<NamespaceId>
+		{
+			/// <inheritdoc/>
+			public override NamespaceId Deserialize(BsonDeserializationContext context, BsonDeserializationArgs args)
+			{
+				return new NamespaceId(context.Reader.ReadString());
+			}
+
+			/// <inheritdoc/>
+			public override void Serialize(BsonSerializationContext context, BsonSerializationArgs args, NamespaceId value)
+			{
+				context.Writer.WriteString(value.ToString());
+			}
+		}
+
 		static int s_haveConfiguredMongoDb = 0;
 
 		public static void ConfigureMongoDbClient()
@@ -859,6 +874,7 @@ namespace Horde.Build
 				BsonSerializer.RegisterSerializer(new BlobIdBsonSerializer());
 				BsonSerializer.RegisterSerializer(new RefIdBsonSerializer());
 				BsonSerializer.RegisterSerializer(new RefNameBsonSerializer());
+				BsonSerializer.RegisterSerializer(new NamespaceIdBsonSerializer());
 				BsonSerializer.RegisterSerializer(new ConditionSerializer());
 				BsonSerializer.RegisterSerializationProvider(new BsonSerializationProvider());
 				BsonSerializer.RegisterSerializationProvider(new StringIdSerializationProvider());
diff --git a/Engine/Source/Programs/Horde/Horde.Build/Storage/RefCollection.cs b/Engine/Source/Programs/Horde/Horde.Build/Storage/RefCollection.cs
--- a/Engine/Source/Programs/Horde/Horde.Build/Storage/RefCollection.cs
+++ b/Engine/Source/Programs/Horde/Horde.Build/Storage/RefCollection.cs
@@ -37,6 +37,7 @@ namespace Horde.Build.Storage
 
 		class RefDocument
 		{
+			[BsonIgnoreIfDefault]
 			public ObjectId Id { get; set; }
 
 			[BsonElement("ns")]
diff --git a/Engine/Source/Programs/Horde/Horde.Build/Storage/StorageController.cs b/Engine/Source/Programs/Horde/Horde.Build/Storage/StorageController.cs
--- a/Engine/Source/Programs/Horde/Horde.Build/Storage/StorageController.cs
+++ b/Engine/Source/Programs/Horde/Horde.Build/Storage/StorageController.cs
@@ -1,6 +1,9 @@
 // Copyright Epic Games, Inc. All Rights Reserved.
 
+using System;
 using System.IO;
+using System.Linq;
+using System.Text.RegularExpressions;
 using System.Threading;
 using System.Threading.Tasks;
 using Amazon.EC2.Model;
@@ -11,13 +14,14 @@ using Horde.Build.Utilities;
 using Microsoft.AspNetCore.Authorization;
 using Microsoft.AspNetCore.Http;
 using Microsoft.AspNetCore.Mvc;
+using Microsoft.Extensions.Primitives;
 
 namespace Horde.Build.Storage
 {
 	/// <summary>
 	/// Response from uploading a bundle
 	/// </summary>
-	public class WriteBundleResponse
+	public class WriteBlobResponse
 	{
 		/// <summary>
 		/// Locator for the uploaded bundle
@@ -92,8 +96,8 @@ namespace Horde.Build.Storage
 		/// <param name="file">Data to be uploaded</param>
 		/// <param name="cancellationToken">Cancellation token for the operation</param>
 		[HttpPost]
-		[Route("/api/v1/storage/{namespaceId}/bundles")]
-		public async Task<ActionResult<WriteBundleResponse>> WriteBundleAsync(NamespaceId namespaceId, IFormFile file, CancellationToken cancellationToken)
+		[Route("/api/v1/storage/{namespaceId}/blobs")]
+		public async Task<ActionResult<WriteBlobResponse>> WriteBlobAsync(NamespaceId namespaceId, IFormFile file, CancellationToken cancellationToken)
 		{
 			if (!await _storageService.AuthorizeAsync(namespaceId, User, AclAction.WriteBlobs, null, cancellationToken))
 			{
@@ -104,7 +108,7 @@ namespace Horde.Build.Storage
 			using (Stream stream = file.OpenReadStream())
 			{
 				BlobLocator locator = await client.WriteBlobAsync(stream, cancellationToken: cancellationToken);
-				return new WriteBundleResponse { Locator = locator };
+				return new WriteBlobResponse { Locator = locator };
 			}
 		}
 
@@ -113,10 +117,12 @@ namespace Horde.Build.Storage
 		/// </summary>
 		/// <param name="namespaceId">Namespace to fetch from</param>
 		/// <param name="locator">Bundle to retrieve</param>
+		/// <param name="offset">Offset of the data.</param>
+		/// <param name="length">Length of the data to return.</param>
 		/// <param name="cancellationToken">Cancellation token for the operation</param>
 		[HttpGet]
-		[Route("/api/v1/storage/{namespaceId}/bundles/{locator}")]
-		public async Task<ActionResult> ReadBundleAsync(NamespaceId namespaceId, BlobLocator locator, CancellationToken cancellationToken)
+		[Route("/api/v1/storage/{namespaceId}/blobs/{*locator}")]
+		public async Task<ActionResult> ReadBlobAsync(NamespaceId namespaceId, BlobLocator locator, [FromQuery] int? offset = null, [FromQuery] int? length = null, CancellationToken cancellationToken = default)
 		{
 			if (!await _storageService.AuthorizeAsync(namespaceId, User, AclAction.ReadBlobs, null, cancellationToken))
 			{
@@ -124,9 +130,22 @@ namespace Horde.Build.Storage
 			}
 
 			IStorageClient client = await _storageService.GetClientAsync(namespaceId, cancellationToken);
-			Stream stream = await client.ReadBlobAsync(locator, cancellationToken);
 
 #pragma warning disable CA2000 // Dispose objects before losing scope
+			// TODO: would be better to use the range header here, but seems to require a lot of plumbing to convert unseekable AWS streams into a format that works with range processing.
+			Stream stream;
+			if (offset == null && length == null)
+			{
+				stream = await client.ReadBlobAsync(locator, cancellationToken);
+			}
+			else if (offset != null && length != null)
+			{
+				stream = await client.ReadBlobRangeAsync(locator, offset.Value, length.Value, cancellationToken);
+			}
+			else
+			{
+				return BadRequest("Offset and length must both be specified as query parameters for ranged reads");
+			}
 			return File(stream, "application/octet-stream");
 #pragma warning restore CA2000 // Dispose objects before losing scope
 		}
@@ -139,7 +158,7 @@ namespace Horde.Build.Storage
 		/// <param name="request">Request for the ref to write</param>
 		/// <param name="cancellationToken">Cancellation token for the operation</param>
 		[HttpPut]
-		[Route("/api/v1/storage/{namespaceId}/refs/{refName}")]
+		[Route("/api/v1/storage/{namespaceId}/refs/{*refName}")]
 		public async Task<ActionResult> WriteRefAsync(NamespaceId namespaceId, RefName refName, [FromBody] WriteRefRequest request, CancellationToken cancellationToken)
 		{
 			if (!await _storageService.AuthorizeAsync(namespaceId, User, AclAction.WriteRefs, null, cancellationToken))
@@ -161,7 +180,7 @@ namespace Horde.Build.Storage
 		/// <param name="refName"></param>
 		/// <param name="cancellationToken"></param>
 		[HttpGet]
-		[Route("/api/v1/storage/{namespaceId}/refs/{refName}")]
+		[Route("/api/v1/storage/{namespaceId}/refs/{*refName}")]
 		public async Task<ActionResult<ReadRefResponse>> ReadRefAsync(NamespaceId namespaceId, RefName refName, CancellationToken cancellationToken)
 		{
 			if (!await _storageService.AuthorizeAsync(namespaceId, User, AclAction.ReadRefs, null, cancellationToken))
diff --git a/Engine/Source/Programs/Shared/EpicGames.Horde/Storage/Backends/FileStorageClient.cs b/Engine/Source/Programs/Shared/EpicGames.Horde/Storage/Backends/FileStorageClient.cs
new file mode 100644
--- /dev/null
+++ b/Engine/Source/Programs/Shared/EpicGames.Horde/Storage/Backends/FileStorageClient.cs
@@ -0,0 +1,116 @@
+// Copyright Epic Games, Inc. All Rights Reserved.
+
+using System;
+using System.Globalization;
+using System.IO;
+using System.Threading;
+using System.Threading.Tasks;
+using EpicGames.Core;
+using Microsoft.Extensions.Caching.Memory;
+using Microsoft.Extensions.Logging;
+
+namespace EpicGames.Horde.Storage.Backends
+{
+	/// <summary>
+	/// Implementation of <see cref="IStorageClient"/> which writes data to files on disk.
+	/// </summary>
+	public class FileStorageClient : StorageClientBase
+	{
+		readonly DirectoryReference _rootDir;
+		readonly ILogger _logger;
+
+		/// <summary>
+		/// Constructor
+		/// </summary>
+		/// <param name="rootDir">Root directory for storing blobs</param>
+		/// <param name="cache">Cache for blob data</param>
+		/// <param name="logger">Logger interface</param>
+		public FileStorageClient(DirectoryReference rootDir, IMemoryCache cache, ILogger logger)
+			: base(cache, logger)
+		{
+			_rootDir = rootDir;
+			_logger = logger;
+
+			DirectoryReference.CreateDirectory(_rootDir);
+		}
+
+		FileReference GetRefFile(RefName name) => FileReference.Combine(_rootDir, name.ToString() + ".ref");
+		FileReference GetBlobFile(BlobLocator id) => FileReference.Combine(_rootDir, id.Inner.ToString() + ".blob");
+
+		#region Blobs
+
+		/// <inheritdoc/>
+		public override Task<Stream> ReadBlobAsync(BlobLocator id, CancellationToken cancellationToken = default)
+		{
+			FileReference file = GetBlobFile(id);
+			_logger.LogInformation("Reading {File}", file);
+			return Task.FromResult<Stream>(FileReference.Open(file, FileMode.Open, FileAccess.Read, FileShare.Read));
+		}
+
+		/// <inheritdoc/>
+		public override async Task<Stream> ReadBlobRangeAsync(BlobLocator id, int offset, int length, CancellationToken cancellationToken = default)
+		{
+			Stream stream = await ReadBlobAsync(id, cancellationToken);
+			stream.Seek(offset, SeekOrigin.Begin);
+			return stream;
+		}
+
+		/// <inheritdoc/>
+		public override async Task<BlobLocator> WriteBlobAsync(Stream stream, Utf8String prefix = default, CancellationToken cancellationToken = default)
+		{
+			BlobLocator id = BlobLocator.Create(HostId.Empty, prefix);
+			FileReference file = GetBlobFile(id);
+			DirectoryReference.CreateDirectory(file.Directory);
+			_logger.LogInformation("Writing {File}", file);
+
+			using (FileStream fileStream = FileReference.Open(file, FileMode.Create, FileAccess.Write, FileShare.ReadWrite))
+			{
+				await stream.CopyToAsync(fileStream, cancellationToken);
+			}
+
+			return id;
+		}
+
+		#endregion
+
+		#region Refs
+
+		/// <inheritdoc/>
+		public override Task DeleteRefAsync(RefName name, CancellationToken cancellationToken = default)
+		{
+			FileReference file = GetRefFile(name);
+			FileReference.Delete(file);
+			return Task.CompletedTask;
+		}
+
+		/// <inheritdoc/>
+		public override async Task<NodeLocator> TryReadRefTargetAsync(RefName name, DateTime cacheTime = default, CancellationToken cancellationToken = default)
+		{
+			FileReference file = GetRefFile(name);
+			if (!FileReference.Exists(file))
+			{
+				return default;
+			}
+
+			_logger.LogInformation("Reading {File}", file);
+			string text = await FileReference.ReadAllTextAsync(file, cancellationToken);
+
+			int hashIdx = text.IndexOf('#', StringComparison.Ordinal);
+			BlobLocator locator = new BlobLocator(text.Substring(0, hashIdx));
+			int exportIdx = Int32.Parse(text.Substring(hashIdx + 1), CultureInfo.InvariantCulture);
+
+			return new NodeLocator(locator, exportIdx);
+		}
+
+		/// <inheritdoc/>
+		public override async Task WriteRefTargetAsync(RefName name, NodeLocator target, CancellationToken cancellationToken = default)
+		{
+			FileReference file = GetRefFile(name);
+			DirectoryReference.CreateDirectory(file.Directory);
+			_logger.LogInformation("Writing {File}", file);
+			await FileReference.WriteAllTextAsync(file, $"{target.Blob}#{target.ExportIdx}");
+		}
+
+		#endregion
+	}
+}
diff --git a/Engine/Source/Programs/Shared/EpicGames.Horde/Storage/Backends/HttpStorageClient.cs b/Engine/Source/Programs/Shared/EpicGames.Horde/Storage/Backends/HttpStorageClient.cs
new file mode 100644
--- /dev/null
+++ b/Engine/Source/Programs/Shared/EpicGames.Horde/Storage/Backends/HttpStorageClient.cs
@@ -0,0 +1,130 @@
+// Copyright Epic Games, Inc. All Rights Reserved.
+
+using System;
+using System.IO;
+using System.Net.Http;
+using System.Net.Http.Headers;
+using System.Net.Http.Json;
+using System.Threading;
+using System.Threading.Tasks;
+using EpicGames.Core;
+using Microsoft.Extensions.Caching.Memory;
+using Microsoft.Extensions.Logging;
+
+namespace EpicGames.Horde.Storage.Backends
+{
+	/// <summary>
+	/// Implementation of <see cref="IStorageClient"/> which communicates with an upstream Horde instance via HTTP.
+	/// </summary>
+	public class HttpStorageClient : StorageClientBase
+	{
+		class WriteBlobResponse
+		{
+			public BlobLocator Locator { get; set; }
+		}
+
+		class ReadRefResponse
+		{
+			public BlobLocator Blob { get; set; }
+			public int ExportIdx { get; set; }
+		}
+
+		readonly NamespaceId _namespaceId;
+		readonly HttpClient _httpClient;
+
+		/// <summary>
+		/// Constructor
+		/// </summary>
+		public HttpStorageClient(NamespaceId namespaceId, HttpClient httpClient, IMemoryCache cache, ILogger logger) 
+			: base(cache, logger)
+		{
+			_namespaceId = namespaceId;
+			_httpClient = httpClient;
+		}
+
+		#region Blobs
+
+		/// <inheritdoc/>
+		public override async Task<Stream> ReadBlobAsync(BlobLocator locator, CancellationToken cancellationToken = default)
+		{
+			using (HttpRequestMessage request = new HttpRequestMessage(HttpMethod.Get, $"api/v1/storage/{_namespaceId}/blobs/{locator}"))
+			{
+				HttpResponseMessage response = await _httpClient.SendAsync(request, cancellationToken);
+				response.EnsureSuccessStatusCode();
+				return await response.Content.ReadAsStreamAsync(cancellationToken);
+			}
+		}
+
+		/// <inheritdoc/>
+		public override async Task<Stream> ReadBlobRangeAsync(BlobLocator locator, int offset, int length, CancellationToken cancellationToken = default)
+		{
+			using (HttpRequestMessage request = new HttpRequestMessage(HttpMethod.Get, $"api/v1/storage/{_namespaceId}/blobs/{locator}?offset={offset}&length={length}"))
+			{
+				HttpResponseMessage response = await _httpClient.SendAsync(request, cancellationToken);
+				response.EnsureSuccessStatusCode();
+				return await response.Content.ReadAsStreamAsync(cancellationToken);
+			}
+		}
+
+		/// <inheritdoc/>
+		public override async Task<BlobLocator> WriteBlobAsync(Stream stream, Utf8String prefix = default, CancellationToken cancellationToken = default)
+		{
+			using (HttpRequestMessage request = new HttpRequestMessage(HttpMethod.Post, $"api/v1/storage/{_namespaceId}/blobs"))
+			{
+				using StreamContent streamContent = new StreamContent(stream);
+
+				MultipartFormDataContent form = new MultipartFormDataContent();
+				form.Add(streamContent, "file", "filename");
+
+				request.Content = form;
+				using (HttpResponseMessage response = await _httpClient.SendAsync(request, cancellationToken))
+				{
+					response.EnsureSuccessStatusCode();
+					WriteBlobResponse? data = await response.Content.ReadFromJsonAsync<WriteBlobResponse>(cancellationToken: cancellationToken);
+					return data!.Locator;
+				}
+			}
+		}
+
+		#endregion
+
+		#region Refs
+
+		/// <inheritdoc/>
+		public override async Task DeleteRefAsync(RefName name, CancellationToken cancellationToken)
+		{
+			using (HttpRequestMessage request = new HttpRequestMessage(HttpMethod.Delete, $"api/v1/storage/{_namespaceId}/refs/{name}"))
+			{
+				using (HttpResponseMessage response = await _httpClient.SendAsync(request, cancellationToken))
+				{
+					response.EnsureSuccessStatusCode();
+				}
+			}
+		}
+
+		/// <inheritdoc/>
+		public override async Task<NodeLocator> TryReadRefTargetAsync(RefName name, DateTime cacheTime = default, CancellationToken cancellationToken = default)
+		{
+			using (HttpRequestMessage request = new HttpRequestMessage(HttpMethod.Get, $"api/v1/storage/{_namespaceId}/refs/{name}"))
+			{
+				using (HttpResponseMessage response = await _httpClient.SendAsync(request, cancellationToken))
+				{
+					response.EnsureSuccessStatusCode();
+					ReadRefResponse? data = await response.Content.ReadFromJsonAsync<ReadRefResponse>(cancellationToken: cancellationToken);
+					return new NodeLocator(data!.Blob, data!.ExportIdx);
+				}
+			}
+		}
+
+		/// <inheritdoc/>
+		public override async Task WriteRefTargetAsync(RefName name, NodeLocator target, CancellationToken cancellationToken = default)
+		{
+			using (HttpResponseMessage response = await _httpClient.PutAsync($"api/v1/storage/{_namespaceId}/refs/{name}", new { locator = target.Blob, exportIdx = target.ExportIdx }, cancellationToken))
+			{
+				response.EnsureSuccessStatusCode();
+			}
+		}
+
+		#endregion
+	}
+}
diff --git a/Engine/Source/Programs/Shared/EpicGames.Horde/Storage/Backends/MemoryBlobStore.cs b/Engine/Source/Programs/Shared/EpicGames.Horde/Storage/Backends/MemoryBlobStore.cs
deleted file mode 100644
diff --git a/Engine/Source/Programs/Shared/EpicGames.Horde/Storage/Backends/MemoryStorageClient.cs b/Engine/Source/Programs/Shared/EpicGames.Horde/Storage/Backends/MemoryStorageClient.cs
new file mode 100644
--- /dev/null
+++ b/Engine/Source/Programs/Shared/EpicGames.Horde/Storage/Backends/MemoryStorageClient.cs
@@ -0,0 +1,93 @@
+// Copyright Epic Games, Inc. All Rights Reserved.
+
+using System;
+using System.Buffers;
+using System.Collections.Concurrent;
+using System.Collections.Generic;
+using System.IO;
+using System.Threading;
+using System.Threading.Tasks;
+using EpicGames.Core;
+using Microsoft.Extensions.Caching.Memory;
+using Microsoft.Extensions.Logging;
+
+namespace EpicGames.Horde.Storage.Backends
+{
+	/// <summary>
+	/// Implementation of <see cref="IStorageClient"/> which stores data in memory. Not intended for production use.
+	/// </summary>
+	public class MemoryBlobStore : StorageClientBase
+	{
+		/// <summary>
+		/// Map of blob id to blob data
+		/// </summary>
+		readonly ConcurrentDictionary<BlobLocator, Bundle> _blobs = new ConcurrentDictionary<BlobLocator, Bundle>();
+
+		/// <summary>
+		/// Map of ref name to ref data
+		/// </summary>
+		readonly ConcurrentDictionary<RefName, NodeLocator> _refs = new ConcurrentDictionary<RefName, NodeLocator>();
+
+		/// <inheritdoc cref="_blobs"/>
+		public IReadOnlyDictionary<BlobLocator, Bundle> Blobs => _blobs;
+
+		/// <inheritdoc cref="_refs"/>
+		public IReadOnlyDictionary<RefName, NodeLocator> Refs => _refs;
+
+		/// <summary>
+		/// Constructor
+		/// </summary>
+		public MemoryBlobStore(IMemoryCache cache, ILogger logger) 
+			: base(cache, logger)
+		{
+		}
+
+		#region Blobs
+
+		/// <inheritdoc/>
+		public override Task<Stream> ReadBlobAsync(BlobLocator locator, CancellationToken cancellationToken = default)
+		{
+			Stream stream = new ReadOnlySequenceStream(_blobs[locator].AsSequence());
+			return Task.FromResult(stream);
+		}
+
+		/// <inheritdoc/>
+		public override Task<Stream> ReadBlobRangeAsync(BlobLocator locator, int offset, int length, CancellationToken cancellationToken = default)
+		{
+			Stream stream = new ReadOnlySequenceStream(_blobs[locator].AsSequence().Slice(offset));
+			return Task.FromResult(stream);
+		}
+
+		/// <inheritdoc/>
+		public override async Task<BlobLocator> WriteBlobAsync(Stream stream, Utf8String prefix = default, CancellationToken cancellationToken = default)
+		{
+			BlobLocator locator = BlobLocator.Create(HostId.Empty, prefix);
+			_blobs[locator] = await Bundle.FromStreamAsync(stream, cancellationToken);
+			return locator;
+		}
+
+		#endregion
+
+		#region Refs
+
+		/// <inheritdoc/>
+		public override Task DeleteRefAsync(RefName name, CancellationToken cancellationToken) => Task.FromResult(_refs.TryRemove(name, out _));
+
+		/// <inheritdoc/>
+		public override Task<NodeLocator> TryReadRefTargetAsync(RefName name, DateTime cacheTime = default, CancellationToken cancellationToken = default)
+		{
+			NodeLocator refTarget;
+			_refs.TryGetValue(name, out refTarget);
+			return Task.FromResult(refTarget);
+		}
+
+		/// <inheritdoc/>
+		public override Task WriteRefTargetAsync(RefName name, NodeLocator target, CancellationToken cancellationToken = default)
+		{
+			_refs[name] = target;
+			return Task.CompletedTask;
+		}
+
+		#endregion
+	}
+}
