Horde Storage - Writing to replication log when replicating if we actually ended up writing new blobs. This was disabled over a concern that we could end up with a recursion in the replication log, but if we only write when content was missing then it was valid. This resolves a issue were a 3rd region wasnt getting some of the writes it needed.

diff --git a/Engine/Source/Programs/Horde/HordeStorage/Horde.Storage/Implementation/Replication/RefsReplicator.cs b/Engine/Source/Programs/Horde/HordeStorage/Horde.Storage/Implementation/Replication/RefsReplicator.cs
--- a/Engine/Source/Programs/Horde/HordeStorage/Horde.Storage/Implementation/Replication/RefsReplicator.cs
+++ b/Engine/Source/Programs/Horde/HordeStorage/Horde.Storage/Implementation/Replication/RefsReplicator.cs
@@ -317,12 +317,11 @@ namespace Horde.Storage.Implementation
                     Info.CountOfRunningReplications = countOfObjectsCurrentlyReplicating;
                     Info.LastRun = DateTime.Now;
 
-                    await ReplicateOp(ns, snapshotLiveObject.Blob, cancellationToken);
-                    // TODO: Avoid adding to the replication log as we will get into infinite recursion if we do with the remote site
-                    // we could add to the replication log only if the op was missing, but that could cause issue for a 3rd site replicating from us
-                    // and the 3rd site is the whole reason we even add this to replication log in the first place
-                    //await AddToReplicationLog(ns, snapshotLiveObject.Bucket, snapshotLiveObject.Key, snapshotLiveObject.Blob);
-
+                    bool blobWasReplicated = await ReplicateOp(ns, snapshotLiveObject.Blob, cancellationToken);
+                    if (blobWasReplicated)
+                    {
+                        await AddToReplicationLog(ns, snapshotLiveObject.Bucket, snapshotLiveObject.Key, snapshotLiveObject.Blob);
+                    }
                 }
                 finally
                 {
@@ -391,6 +390,7 @@ namespace Horde.Storage.Implementation
                         replicationTasks.Add(currentOffset);
                     }
 
+                    bool blobWasReplicated = false;
                     // we do not need to replicate delete events
                     if (@event.Op != ReplicationLogEvent.OpType.Deleted)
                     {
@@ -399,13 +399,17 @@ namespace Horde.Storage.Implementation
                             throw new Exception($"Event: [REDACTED_EMAIL]} [REDACTED_EMAIL]} in namespace [REDACTED_EMAIL]} was missing a blob, unable to replicate it");
                         }
 
-                        await ReplicateOp(@event.Namespace, @event.Blob, replicationToken);
+                        blobWasReplicated = await ReplicateOp(@event.Namespace, @event.Blob, replicationToken);
                     }
 
-                    // TODO: Avoid adding to the replication log as we will get into infinite recursion if we do with the remote site
-                    // we could add to the replication log only if the op was missing, but that could cause issue for a 3rd site replicating from us
-                    // and the 3rd site is the whole reason we even add this to replication log in the first place
-                    // await AddToReplicationLog(@event.Namespace, @event.Bucket, @event.Key, @event.Blob);
+                    if (blobWasReplicated)
+                    {
+                        // add events should all have blobs, and blobs are only replicated for adds events so this should always be true
+                        if (@event.Blob != null)
+                        {
+                            await AddToReplicationLog(@event.Namespace, @event.Bucket, @event.Key, @event.Blob);
+                        }
+                    }
 
                     bool wasOldest;
                     lock (replicationTasks)
@@ -443,7 +447,7 @@ namespace Horde.Storage.Implementation
             return countOfReplicationsDone;
         }
 
-        private async Task ReplicateOp(NamespaceId ns, BlobIdentifier objectToReplicate, CancellationToken cancellationToken)
+        private async Task<bool> ReplicateOp(NamespaceId ns, BlobIdentifier objectToReplicate, CancellationToken cancellationToken)
         {
             using IScope scope = Tracer.Instance.StartActive("replicator.replicate_op");
             scope.Span.ResourceName = $"{ns}.{objectToReplicate}";
@@ -462,13 +466,13 @@ namespace Horde.Storage.Implementation
             if (referencesResponse.StatusCode == HttpStatusCode.BadRequest)
             {
                 _logger.Warning("Failed to resolve references for object {Blob} in {Namespace}. Skipping replication", objectToReplicate, ns);
-                return;
+                return false;
             }
             if (referencesResponse.StatusCode == HttpStatusCode.NotFound)
             {
                 // objects that do not exist can not be replicated so we skip them
                 _logger.Warning("Failed to resolve references for object {Blob} in {Namespace}. Got not found with message \"{Message}\". Skipping replication.", objectToReplicate, ns, referencesResponse.ReasonPhrase);
-                return;
+                return false;
             }
             referencesResponse.EnsureSuccessStatusCode();
 
@@ -516,6 +520,7 @@ namespace Horde.Storage.Implementation
             }
 
             await Task.WhenAll(blobReplicationTasks);
+            return missingBlobs.Length != 0;
         }
 
         private async IAsyncEnumerable<ReplicationLogEvent> GetRefEvents(NamespaceId ns, string? lastBucket, Guid? lastEvent, [EnumeratorCancellation] CancellationToken cancellationToken)
@@ -599,6 +604,11 @@ namespace Horde.Storage.Implementation
             _logger.Information("{Name} starting replication. Last transaction was {TransactionId} {Generation}", _name, State.ReplicatorOffset.GetValueOrDefault(0L), State.ReplicatingGeneration.GetValueOrDefault(Guid.Empty) );
         }
 
+        private async Task AddToReplicationLog(NamespaceId ns, BucketId bucket, IoHashKey key, BlobIdentifier blob)
+        {
+            await _replicationLog.InsertAddEvent(ns, bucket, key, blob);
+        }
+
         public void SetReplicationOffset(long? state)
         {
             throw new NotImplementedException();
