Horde Storage - Fixed issue for tests were our namespaces lists would contain duplicate entries.

diff --git a/Engine/Source/Programs/Horde/HordeStorage/Horde.Storage/Implementation/Objects/MemoryReferencesStore.cs b/Engine/Source/Programs/Horde/HordeStorage/Horde.Storage/Implementation/Objects/MemoryReferencesStore.cs
--- a/Engine/Source/Programs/Horde/HordeStorage/Horde.Storage/Implementation/Objects/MemoryReferencesStore.cs
+++ b/Engine/Source/Programs/Horde/HordeStorage/Horde.Storage/Implementation/Objects/MemoryReferencesStore.cs
@@ -14,7 +14,7 @@ namespace Horde.Storage.Implementation
     public class MemoryReferencesStore : IReferencesStore
     {
         private readonly ConcurrentDictionary<string, MemoryStoreObject> _objects = new ConcurrentDictionary<string, MemoryStoreObject>();
-        private readonly List<NamespaceId> _namespaces = new List<NamespaceId>();
+        private readonly HashSet<NamespaceId> _namespaces = new HashSet<NamespaceId>();
 
         public MemoryReferencesStore()
         {
