Horde Storage - Fixed up displayed storage pool name in DD when cleaning the default (empty string) pool

diff --git a/Engine/Source/Programs/Horde/HordeStorage/Horde.Storage/Implementation/GC/OrphanBlobCleanupRefs.cs b/Engine/Source/Programs/Horde/HordeStorage/Horde.Storage/Implementation/GC/OrphanBlobCleanupRefs.cs
--- a/Engine/Source/Programs/Horde/HordeStorage/Horde.Storage/Implementation/GC/OrphanBlobCleanupRefs.cs
+++ b/Engine/Source/Programs/Horde/HordeStorage/Horde.Storage/Implementation/GC/OrphanBlobCleanupRefs.cs
@@ -115,7 +115,8 @@ namespace Horde.Storage.Implementation
         private async Task<bool> GCBlob(string storagePool, List<NamespaceId> namespacesThatSharePool, BlobIdentifier blob, DateTime lastModifiedTime, CancellationToken cancellationToken)
         {
             using IScope removeBlobScope = Tracer.Instance.StartActive("gc.blob");
-            removeBlobScope.Span.ResourceName = $"{storagePool}.{blob}";
+            string storagePoolName = string.IsNullOrEmpty(storagePool) ? "default" : storagePool; 
+            removeBlobScope.Span.ResourceName = $"{storagePoolName}.{blob}";
 
             bool found = false;
 
