Horde Storage - Cache result of pod enumeration in kubernetes to avoid to many requests responses
Also added helm option to disable using other pods for blobs (fetchBlobsFromPeerPods)

diff --git a/Engine/Source/Programs/Horde/HordeStorage/Horde.Storage/Implementation/Blob/HordeStorageBlobStore.cs b/Engine/Source/Programs/Horde/HordeStorage/Horde.Storage/Implementation/Blob/HordeStorageBlobStore.cs
--- a/Engine/Source/Programs/Horde/HordeStorage/Horde.Storage/Implementation/Blob/HordeStorageBlobStore.cs
+++ b/Engine/Source/Programs/Horde/HordeStorage/Horde.Storage/Implementation/Blob/HordeStorageBlobStore.cs
@@ -3,7 +3,6 @@
 using System;
 using System.Collections.Concurrent;
 using System.Collections.Generic;
-using System.ComponentModel.DataAnnotations;
 using System.IO;
 using System.Net;
 using System.Net.Http;
@@ -177,6 +176,8 @@ namespace Horde.Storage.Implementation
     {
         private readonly IOptionsMonitor<KubernetesLeaderElectionSettings> _leaderSettings;
         private readonly Kubernetes _client;
+        private DateTime _podEnumerationValidUntil = DateTime.Now;
+        private List<string>? _lastPodEnumeration;
 
         public KubernetesHordeStorageServiceDiscovery(IOptionsMonitor<KubernetesLeaderElectionSettings> leaderSettings)
         {
@@ -186,6 +187,26 @@ namespace Horde.Storage.Implementation
         }
 
         public async IAsyncEnumerable<string> FindOtherHordeStorageInstances()
+        {
+            if (_lastPodEnumeration != null && _podEnumerationValidUntil > DateTime.Now)
+            {
+                foreach (string pod in _lastPodEnumeration)
+                {
+                    yield return pod;
+                }
+                yield break;
+            }
+
+            _lastPodEnumeration = await EnumeratePods().ToListAsync();
+            _podEnumerationValidUntil = DateTime.Now.AddMinutes(5);
+
+            foreach (string pod in _lastPodEnumeration)
+            {
+                yield return pod;
+            }
+        }
+
+        private async IAsyncEnumerable<string> EnumeratePods()
         {
             V1PodList podList = await _client.ListNamespacedPodAsync(_leaderSettings.CurrentValue.Namespace, labelSelector: _leaderSettings.CurrentValue.HordeStoragePodLabelSelector);
 
