Horde Storage - Ignore exceptions when attempting to connect to other horde storage pods - these are just meant as a cache layer and we can always fallback to S3 if there is a issue with them.

diff --git a/Engine/Source/Programs/Horde/HordeStorage/Horde.Storage/Implementation/Blob/HordeStorageBlobStore.cs b/Engine/Source/Programs/Horde/HordeStorage/Horde.Storage/Implementation/Blob/HordeStorageBlobStore.cs
--- a/Engine/Source/Programs/Horde/HordeStorage/Horde.Storage/Implementation/Blob/HordeStorageBlobStore.cs
+++ b/Engine/Source/Programs/Horde/HordeStorage/Horde.Storage/Implementation/Blob/HordeStorageBlobStore.cs
@@ -51,24 +51,31 @@ namespace Horde.Storage.Implementation
             // TODO: This could be done in parallel
             await foreach (string instance in _serviceDiscovery.FindOtherHordeStorageInstances())
             {
-                string filesystemLayerName = nameof(FileSystemStore);
-                using HttpRequestMessage getObjectRequest = await BuildHttpRequest(HttpMethod.Get, new Uri($"api/v1/blobs/{ns}/{blob}?storageLayers={filesystemLayerName}", UriKind.Relative));
-                getObjectRequest.Headers.Add("Accept", MediaTypeNames.Application.Octet);
-                HttpResponseMessage response = await GetHttpClient(instance).SendAsync(getObjectRequest);
-                if (response.StatusCode == HttpStatusCode.NotFound)
+                try
                 {
-                    continue;
+                    string filesystemLayerName = nameof(FileSystemStore);
+                    using HttpRequestMessage getObjectRequest = await BuildHttpRequest(HttpMethod.Get, new Uri($"api/v1/blobs/{ns}/{blob}?storageLayers={filesystemLayerName}", UriKind.Relative));
+                    getObjectRequest.Headers.Add("Accept", MediaTypeNames.Application.Octet);
+                    HttpResponseMessage response = await GetHttpClient(instance).SendAsync(getObjectRequest);
+                    if (response.StatusCode == HttpStatusCode.NotFound)
+                    {
+                        continue;
+                    }
+
+                    response.EnsureSuccessStatusCode();
+
+                    long? contentLength = response.Content.Headers.ContentLength;
+                    if (contentLength == null)
+                    {
+                        _logger.Warning("Content length missing in response from horde storage blob store. This is not supported, ignoring response");
+                        continue;
+                    }
+                    return new BlobContents(await response.Content.ReadAsStreamAsync(), contentLength.Value);
                 }
-
-                response.EnsureSuccessStatusCode();
-
-                long? contentLength = response.Content.Headers.ContentLength;
-                if (contentLength == null)
+                catch (Exception e)
                 {
-                    _logger.Warning("Content length missing in response from horde storage blob store. This is not supported, ignoring response");
-                    continue;
+                    _logger.Warning(e, "Exception when attempting to fetch blob {Blob} in namespace {Namespace} from instance {Instance}", blob, ns, instance);
                 }
-                return new BlobContents(await response.Content.ReadAsStreamAsync(), contentLength.Value);
             }
             
             throw new BlobNotFoundException(ns, blob);
@@ -79,17 +86,24 @@ namespace Horde.Storage.Implementation
             // TODO: This could be done in parallel
             await foreach (string instance in _serviceDiscovery.FindOtherHordeStorageInstances())
             {
-                string filesystemLayerName = nameof(FileSystemStore);
-                using HttpRequestMessage headObjectRequest = await BuildHttpRequest(HttpMethod.Head, new Uri($"api/v1/blobs/{ns}/{blob}?storageLayers={filesystemLayerName}", UriKind.Relative));
-                HttpResponseMessage response = await GetHttpClient(instance).SendAsync(headObjectRequest);
-                if (response.StatusCode == HttpStatusCode.NotFound)
+                try
                 {
-                    continue;
-                }
+                    string filesystemLayerName = nameof(FileSystemStore);
+                    using HttpRequestMessage headObjectRequest = await BuildHttpRequest(HttpMethod.Head, new Uri($"api/v1/blobs/{ns}/{blob}?storageLayers={filesystemLayerName}", UriKind.Relative));
+                    HttpResponseMessage response = await GetHttpClient(instance).SendAsync(headObjectRequest);
+                    if (response.StatusCode == HttpStatusCode.NotFound)
+                    {
+                        continue;
+                    }
 
-                response.EnsureSuccessStatusCode();
+                    response.EnsureSuccessStatusCode();
 
-                return true;
+                    return true;
+                }
+                catch (Exception e)
+                {
+                    _logger.Warning(e, "Exception when attempting to fetch blob {Blob} in namespace {Namespace} from instance {Instance}", blob, ns, instance);
+                }
             }
 
             return false;
