Horde Storage - Skip connecting to pods that are not running (e.g. that maybe starting up, shutting down or anything else)

diff --git a/Engine/Source/Programs/Horde/HordeStorage/Horde.Storage/Implementation/Blob/HordeStorageBlobStore.cs b/Engine/Source/Programs/Horde/HordeStorage/Horde.Storage/Implementation/Blob/HordeStorageBlobStore.cs
--- a/Engine/Source/Programs/Horde/HordeStorage/Horde.Storage/Implementation/Blob/HordeStorageBlobStore.cs
+++ b/Engine/Source/Programs/Horde/HordeStorage/Horde.Storage/Implementation/Blob/HordeStorageBlobStore.cs
@@ -226,6 +226,10 @@ namespace Horde.Storage.Implementation
 
             foreach (V1Pod pod in podList.Items)
             {
+                if (pod.Status.Phase != "Running")
+                {
+                    continue;
+                }
                 string ip = pod.Status.PodIP;
                 if (string.IsNullOrEmpty(ip))
                 {
