Horde: Add test for creating re-usable AWS fleet manager with settings

diff --git a/Engine/Source/Programs/Horde/Horde.Build.Tests/Fleet/AutoscaleServiceV2Test.cs b/Engine/Source/Programs/Horde/Horde.Build.Tests/Fleet/AutoscaleServiceV2Test.cs
--- a/Engine/Source/Programs/Horde/Horde.Build.Tests/Fleet/AutoscaleServiceV2Test.cs
+++ b/Engine/Source/Programs/Horde/Horde.Build.Tests/Fleet/AutoscaleServiceV2Test.cs
@@ -264,8 +264,10 @@ namespace Horde.Build.Tests.Fleet
 		[TestMethod]
 		public async Task CreateAwsReuse()
 		{
-			IFleetManager fm = await CreateFleetManager(new FleetManagerInfo(FleetManagerType.AwsReuse, null, "{}"));
+			IFleetManager fm = await CreateFleetManager(new FleetManagerInfo(FleetManagerType.AwsReuse, null, "{\"InstanceTypes\": [\"foo\", \"bar\"]}"));
 			Assert.AreEqual(typeof(AwsReuseFleetManager), fm.GetType());
+			Assert.AreEqual("foo", ((AwsReuseFleetManager)fm).Settings.InstanceTypes![0]);
+			Assert.AreEqual("bar", ((AwsReuseFleetManager)fm).Settings.InstanceTypes![1]);
 		}
 		
 		[TestMethod]
diff --git a/Engine/Source/Programs/Horde/Horde.Build/Agents/Fleet/Providers/AwsReuseFleetManager.cs b/Engine/Source/Programs/Horde/Horde.Build/Agents/Fleet/Providers/AwsReuseFleetManager.cs
--- a/Engine/Source/Programs/Horde/Horde.Build/Agents/Fleet/Providers/AwsReuseFleetManager.cs
+++ b/Engine/Source/Programs/Horde/Horde.Build/Agents/Fleet/Providers/AwsReuseFleetManager.cs
@@ -42,9 +42,13 @@ namespace Horde.Build.Agents.Fleet.Providers
 	/// </summary>
 	public sealed class AwsReuseFleetManager : IFleetManager
 	{
+		/// <summary>
+		/// Settings for fleet manager
+		/// </summary>
+		public AwsReuseFleetManagerSettings Settings { get; }
+		
 		private readonly IAmazonEC2 _ec2;
 		private readonly IAgentCollection _agentCollection;
-		private readonly AwsReuseFleetManagerSettings _settings;
 		private readonly ILogger _logger;
 
 		/// <summary>
@@ -54,8 +58,8 @@ namespace Horde.Build.Agents.Fleet.Providers
 		{
 			_ec2 = ec2;
 			_agentCollection = agentCollection;
-			_settings = settings;
 			_logger = logger;
+			Settings = settings;
 		}
 
 		/// <inheritdoc/>
@@ -81,9 +85,9 @@ namespace Horde.Build.Agents.Fleet.Providers
 
 			IEnumerable<Instance> stoppedInstances = describeResponse.Reservations.SelectMany(x => x.Instances);
 
-			if (_settings.InstanceTypes is { Count: >= 1 })
+			if (Settings.InstanceTypes is { Count: >= 1 })
 			{
-				InstanceType newInstanceType = InstanceType.FindValue(_settings.InstanceTypes[0]);
+				InstanceType newInstanceType = InstanceType.FindValue(Settings.InstanceTypes[0]);
 				foreach (Instance stoppedInstance in stoppedInstances)
 				{
 					if (stoppedInstance.InstanceType == newInstanceType) { continue; }
