Horde: Add missing tests to Dockerfile.

diff --git a/Engine/Source/Programs/Horde/Horde.Build/Dockerfile b/Engine/Source/Programs/Horde/Horde.Build/Dockerfile
--- a/Engine/Source/Programs/Horde/Horde.Build/Dockerfile
+++ b/Engine/Source/Programs/Horde/Horde.Build/Dockerfile
@@ -55,8 +55,11 @@ COPY Source/Programs/Horde/Horde.Build.Tests/ Source/Programs/Horde/Horde.Build.
 COPY Source/Programs/Shared/EpicGames.Serialization.Tests/ Source/Programs/Shared/EpicGames.Serialization.Tests/
 RUN redis-server --save "" --appendonly no --daemonize yes && \
 	mongod --noauth --quiet --fork --dbpath /tmp/mongodb --logpath /tmp/mongod.log && \
-	dotnet test Source/Programs/Horde/Horde.Build.Tests/Horde.Build.Tests.csproj --blame-hang-timeout 2m -p:WithHordeStorage=false && \
-	dotnet test Source/Programs/Shared/EpicGames.Serialization.Tests/EpicGames.Serialization.Tests.csproj --blame-hang-timeout 2m -p:WithHordeStorage=false
+	dotnet test Source/Programs/Shared/EpicGames.BuildGraph.Tests/EpicGames.BuildGraph.Tests.csproj --blame-hang-timeout && \
+	dotnet test Source/Programs/Shared/EpicGames.Core.Tests/EpicGames.Core.Tests.csproj --blame-hang-timeout && \
+	dotnet test Source/Programs/Shared/EpicGames.Horde.Tests/EpicGames.Horde.Tests.csproj --blame-hang-timeout && \
+	dotnet test Source/Programs/Shared/EpicGames.Serialization.Tests/EpicGames.Serialization.Tests.csproj --blame-hang-timeout 2m && \
+	dotnet test Source/Programs/Horde/Horde.Build.Tests/Horde.Build.Tests.csproj --blame-hang-timeout 2m
 
 # Create the runtime image
 FROM mcr.microsoft.com/dotnet/aspnet:6.0 AS runtime
