Horde Storage - Run blob fetching from other Horde Storage instances in parallel and be much more aggresive on timeouts for that connection in case there is issues connecting to the other pods (they may have shutdown for instance)

diff --git a/Engine/Source/Programs/Horde/HordeStorage/Horde.Storage/Implementation/Blob/HordeStorageBlobStore.cs b/Engine/Source/Programs/Horde/HordeStorage/Horde.Storage/Implementation/Blob/HordeStorageBlobStore.cs
--- a/Engine/Source/Programs/Horde/HordeStorage/Horde.Storage/Implementation/Blob/HordeStorageBlobStore.cs
+++ b/Engine/Source/Programs/Horde/HordeStorage/Horde.Storage/Implementation/Blob/HordeStorageBlobStore.cs
@@ -7,6 +7,7 @@ using System.IO;
 using System.Net;
 using System.Net.Http;
 using System.Net.Mime;
+using System.Threading;
 using System.Threading.Tasks;
 using Dasync.Collections;
 using EpicGames.Horde.Storage;
@@ -48,8 +49,8 @@ namespace Horde.Storage.Implementation
 
         public async Task<BlobContents> GetObject(NamespaceId ns, BlobIdentifier blob, LastAccessTrackingFlags flags = LastAccessTrackingFlags.DoTracking)
         {
-            // TODO: This could be done in parallel
-            await foreach (string instance in _serviceDiscovery.FindOtherHordeStorageInstances())
+            BlobContents? contents = null;
+            Parallel.ForEach(await _serviceDiscovery.FindOtherHordeStorageInstances().ToListAsync(), async (string instance, ParallelLoopState state) =>
             {
                 try
                 {
@@ -59,7 +60,7 @@ namespace Horde.Storage.Implementation
                     HttpResponseMessage response = await GetHttpClient(instance).SendAsync(getObjectRequest);
                     if (response.StatusCode == HttpStatusCode.NotFound)
                     {
-                        continue;
+                        return;
                     }
 
                     response.EnsureSuccessStatusCode();
@@ -68,45 +69,59 @@ namespace Horde.Storage.Implementation
                     if (contentLength == null)
                     {
                         _logger.Warning("Content length missing in response from horde storage blob store. This is not supported, ignoring response");
-                        continue;
+                        return;
                     }
-                    return new BlobContents(await response.Content.ReadAsStreamAsync(), contentLength.Value);
+                    
+                    contents = new BlobContents(await response.Content.ReadAsStreamAsync(), contentLength.Value);
+                    state.Break();
                 }
                 catch (Exception e)
                 {
-                    _logger.Warning(e, "Exception when attempting to fetch blob {Blob} in namespace {Namespace} from instance {Instance}", blob, ns, instance);
+                    _logger.Warning(e,
+                        "Exception when attempting to fetch blob {Blob} in namespace {Namespace} from instance {Instance}",
+                        blob, ns, instance);
                 }
+            });
+
+            // contents is incorrectly detected as never modified - its set in the foreach lambda
+#pragma warning disable CA1508 // Avoid dead conditional code
+            if (contents != null)
+            {
+                return contents;
             }
-            
+#pragma warning restore CA1508 // Avoid dead conditional code
             throw new BlobNotFoundException(ns, blob);
         }
 
         public async Task<bool> Exists(NamespaceId ns, BlobIdentifier blob, bool forceCheck = false)
         {
-            // TODO: This could be done in parallel
-            await foreach (string instance in _serviceDiscovery.FindOtherHordeStorageInstances())
+            bool found = false;
+            
+            Parallel.ForEach(await _serviceDiscovery.FindOtherHordeStorageInstances().ToListAsync(), async (string instance, ParallelLoopState state) =>
             {
                 try
                 {
                     string filesystemLayerName = nameof(FileSystemStore);
                     using HttpRequestMessage headObjectRequest = await BuildHttpRequest(HttpMethod.Head, new Uri($"api/v1/blobs/{ns}/{blob}?storageLayers={filesystemLayerName}", UriKind.Relative));
-                    HttpResponseMessage response = await GetHttpClient(instance).SendAsync(headObjectRequest);
+                    HttpResponseMessage response = await GetHttpClient(instance).SendAsync(headObjectRequest, CancellationToken.None);
                     if (response.StatusCode == HttpStatusCode.NotFound)
                     {
-                        continue;
+                        return;
                     }
 
                     response.EnsureSuccessStatusCode();
 
-                    return true;
+                    found = true;
+                    state.Break();
                 }
                 catch (Exception e)
                 {
-                    _logger.Warning(e, "Exception when attempting to fetch blob {Blob} in namespace {Namespace} from instance {Instance}", blob, ns, instance);
+                    _logger.Warning(e,
+                        "Exception when attempting to fetch blob {Blob} in namespace {Namespace} from instance {Instance}",
+                        blob, ns, instance);
                 }
-            }
-
-            return false;
+            });
+            return found;
         }
 
         private HttpClient GetHttpClient(string instance)
@@ -115,6 +130,9 @@ namespace Horde.Storage.Implementation
             {
                 HttpClient httpClient = _httpClientFactory.CreateClient(instance);
                 httpClient.BaseAddress = new Uri($"http://{instance}");
+
+                // for these connections to be useful they need to be fast - so timeout quickly if we can not establish the connection
+                httpClient.Timeout = TimeSpan.FromSeconds(1.0);
                 return httpClient;
             });
         }
@@ -235,6 +253,7 @@ namespace Horde.Storage.Implementation
                 {
                     continue;
                 }
+
                 // we typically expose the container on port 80
                 yield return $"{ip}:80";
             }
diff --git a/Engine/Source/Programs/Horde/HordeStorage/Horde.Storage/Tests/Functional/Blob/HordeStorageBlobTest.cs b/Engine/Source/Programs/Horde/HordeStorage/Horde.Storage/Tests/Functional/Blob/HordeStorageBlobTest.cs
--- a/Engine/Source/Programs/Horde/HordeStorage/Horde.Storage/Tests/Functional/Blob/HordeStorageBlobTest.cs
+++ b/Engine/Source/Programs/Horde/HordeStorage/Horde.Storage/Tests/Functional/Blob/HordeStorageBlobTest.cs
@@ -11,7 +11,6 @@ using System.Threading.Tasks;
 using Amazon.S3;
 using Amazon.S3.Model;
 using EpicGames.Horde.Storage;
-using Horde.Storage.FunctionalTests.Storage;
 using Horde.Storage.Implementation;
 using Jupiter.Implementation;
 using Microsoft.AspNetCore.Hosting;
@@ -122,6 +121,9 @@ namespace Horde.Storage.FunctionalTests.Storage
 
         protected async Task Teardown(IServiceProvider provider)
         {
+            _s3 = provider.GetService<IAmazonS3>();
+            Assert.IsNotNull(_s3);
+
             string S3BucketName = $"tests-{TestNamespaceName}";
             ListObjectsResponse response = await _s3!.ListObjectsAsync(S3BucketName);
             List<KeyVersion> objectKeys = response.S3Objects.Select(o => new KeyVersion { Key = o.Key }).ToList();
