Horde Storage - Another attempt at making the querying of peer horde storage instances parallel

diff --git a/Engine/Source/Programs/Horde/HordeStorage/Horde.Storage/Implementation/Blob/HordeStorageBlobStore.cs b/Engine/Source/Programs/Horde/HordeStorage/Horde.Storage/Implementation/Blob/HordeStorageBlobStore.cs
--- a/Engine/Source/Programs/Horde/HordeStorage/Horde.Storage/Implementation/Blob/HordeStorageBlobStore.cs
+++ b/Engine/Source/Programs/Horde/HordeStorage/Horde.Storage/Implementation/Blob/HordeStorageBlobStore.cs
@@ -47,81 +47,114 @@ namespace Horde.Storage.Implementation
             return request;
         }
 
-        public async Task<BlobContents> GetObject(NamespaceId ns, BlobIdentifier blob, LastAccessTrackingFlags flags = LastAccessTrackingFlags.DoTracking)
+        private async Task<BlobContents?> DoGetObject(string instance, NamespaceId ns, BlobIdentifier blob)
         {
-            BlobContents? contents = null;
-            Parallel.ForEach(await _serviceDiscovery.FindOtherHordeStorageInstances().ToListAsync(), async (string instance, ParallelLoopState state) =>
+            try
             {
-                try
+                string filesystemLayerName = nameof(FileSystemStore);
+                using HttpRequestMessage getObjectRequest = await BuildHttpRequest(HttpMethod.Get, new Uri($"api/v1/blobs/{ns}/{blob}?storageLayers={filesystemLayerName}", UriKind.Relative));
+                getObjectRequest.Headers.Add("Accept", MediaTypeNames.Application.Octet);
+                HttpResponseMessage response = await GetHttpClient(instance).SendAsync(getObjectRequest);
+                if (response.StatusCode == HttpStatusCode.NotFound)
                 {
-                    string filesystemLayerName = nameof(FileSystemStore);
-                    using HttpRequestMessage getObjectRequest = await BuildHttpRequest(HttpMethod.Get, new Uri($"api/v1/blobs/{ns}/{blob}?storageLayers={filesystemLayerName}", UriKind.Relative));
-                    getObjectRequest.Headers.Add("Accept", MediaTypeNames.Application.Octet);
-                    HttpResponseMessage response = await GetHttpClient(instance).SendAsync(getObjectRequest);
-                    if (response.StatusCode == HttpStatusCode.NotFound)
-                    {
-                        return;
-                    }
-
-                    response.EnsureSuccessStatusCode();
-
-                    long? contentLength = response.Content.Headers.ContentLength;
-                    if (contentLength == null)
-                    {
-                        _logger.Warning("Content length missing in response from horde storage blob store. This is not supported, ignoring response");
-                        return;
-                    }
-                    
-                    contents = new BlobContents(await response.Content.ReadAsStreamAsync(), contentLength.Value);
-                    state.Break();
+                    return null;
                 }
-                catch (Exception e)
+
+                response.EnsureSuccessStatusCode();
+
+                long? contentLength = response.Content.Headers.ContentLength;
+                if (contentLength == null)
                 {
-                    _logger.Warning(e,
-                        "Exception when attempting to fetch blob {Blob} in namespace {Namespace} from instance {Instance}",
-                        blob, ns, instance);
+                    _logger.Warning("Content length missing in response from horde storage blob store. This is not supported, ignoring response");
+                    return null;
                 }
-            });
+                    
+                return new BlobContents(await response.Content.ReadAsStreamAsync(), contentLength.Value);
+            }
+            catch (Exception e)
+            {
+                _logger.Warning(e,
+                    "Exception when attempting to fetch blob {Blob} in namespace {Namespace} from instance {Instance}",
+                    blob, ns, instance);
+            }
+
+            return null;
+        }
 
-            // contents is incorrectly detected as never modified - its set in the foreach lambda
-#pragma warning disable CA1508 // Avoid dead conditional code
-            if (contents != null)
+        public async Task<BlobContents> GetObject(NamespaceId ns, BlobIdentifier blob, LastAccessTrackingFlags flags = LastAccessTrackingFlags.DoTracking)
+        {
+            List<Task<BlobContents?>> tasks = new();
+
+            await foreach (string instance in _serviceDiscovery.FindOtherHordeStorageInstances())
+            {
+                Task<BlobContents?> task = DoGetObject(instance, ns, blob);
+                tasks.Add(task);
+            }
+
+            while (tasks.Count != 0)
             {
-                return contents;
+                Task<BlobContents?> finishedTask = await Task.WhenAny(tasks);
+                BlobContents? result = await finishedTask;
+                if (result != null)
+                {
+                    return result;
+                }
+
+                tasks.Remove(finishedTask);
             }
-#pragma warning restore CA1508 // Avoid dead conditional code
+
             throw new BlobNotFoundException(ns, blob);
         }
 
-        public async Task<bool> Exists(NamespaceId ns, BlobIdentifier blob, bool forceCheck = false)
+        private async Task<bool?> DoExists(string instance, NamespaceId ns, BlobIdentifier blob)
         {
-            bool found = false;
-            
-            Parallel.ForEach(await _serviceDiscovery.FindOtherHordeStorageInstances().ToListAsync(), async (string instance, ParallelLoopState state) =>
+            try
             {
-                try
+                string filesystemLayerName = nameof(FileSystemStore);
+                using HttpRequestMessage headObjectRequest = await BuildHttpRequest(HttpMethod.Head, new Uri($"api/v1/blobs/{ns}/{blob}?storageLayers={filesystemLayerName}", UriKind.Relative));
+                HttpResponseMessage response = await GetHttpClient(instance).SendAsync(headObjectRequest, CancellationToken.None);
+                if (response.StatusCode == HttpStatusCode.NotFound)
                 {
-                    string filesystemLayerName = nameof(FileSystemStore);
-                    using HttpRequestMessage headObjectRequest = await BuildHttpRequest(HttpMethod.Head, new Uri($"api/v1/blobs/{ns}/{blob}?storageLayers={filesystemLayerName}", UriKind.Relative));
-                    HttpResponseMessage response = await GetHttpClient(instance).SendAsync(headObjectRequest, CancellationToken.None);
-                    if (response.StatusCode == HttpStatusCode.NotFound)
-                    {
-                        return;
-                    }
-
-                    response.EnsureSuccessStatusCode();
-
-                    found = true;
-                    state.Break();
+                    return false;
                 }
-                catch (Exception e)
+
+                response.EnsureSuccessStatusCode();
+
+                return true;
+            }
+            catch (Exception e)
+            {
+                _logger.Warning(e,
+                    "Exception when attempting to fetch blob {Blob} in namespace {Namespace} from instance {Instance}",
+                    blob, ns, instance);
+            }
+
+            return null;
+        }
+
+        public async Task<bool> Exists(NamespaceId ns, BlobIdentifier blob, bool forceCheck = false)
+        {
+            List<Task<bool?>> tasks = new();
+
+            await foreach (string instance in _serviceDiscovery.FindOtherHordeStorageInstances())
+            {
+                Task<bool?> task = DoExists(instance, ns, blob);
+                tasks.Add(task);
+            }
+
+            while (tasks.Count != 0)
+            {
+                Task<bool?> finishedTask = await Task.WhenAny(tasks);
+                bool? result = await finishedTask;
+                if (result != null)
                 {
-                    _logger.Warning(e,
-                        "Exception when attempting to fetch blob {Blob} in namespace {Namespace} from instance {Instance}",
-                        blob, ns, instance);
+                    return result.Value;
                 }
-            });
-            return found;
+
+                tasks.Remove(finishedTask);
+            }
+
+            return false;
         }
 
         private HttpClient GetHttpClient(string instance)
