Horde: Fix another test build error.

diff --git a/Engine/Source/Programs/Horde/Horde.Build.Tests/BlobStoreTests.cs b/Engine/Source/Programs/Horde/Horde.Build.Tests/BlobStoreTests.cs
--- a/Engine/Source/Programs/Horde/Horde.Build.Tests/BlobStoreTests.cs
+++ b/Engine/Source/Programs/Horde/Horde.Build.Tests/BlobStoreTests.cs
@@ -19,12 +19,12 @@ namespace Horde.Build.Tests
 	public class BlobStoreTests : TestSetup
 	{
 		IMemoryCache? _cache;
-		MemoryBlobStore? _store;
+		MemoryStorageClient? _store;
 
 		IStorageClient CreateBlobStore()
 		{
 			_cache ??= new MemoryCache(new MemoryCacheOptions());
-			_store ??= new MemoryBlobStore(_cache, NullLogger.Instance);
+			_store ??= new MemoryStorageClient(_cache, NullLogger.Instance);
 			return _store;
 		}
 
