Horde: Move EpicGames.Horde tests into NET 6 image. Add placeholder commands to run EpicGames.Core tests also, but disabled due to warning about mismatched assemblies.

diff --git a/Engine/Source/Programs/Horde/Horde.Build/Dockerfile b/Engine/Source/Programs/Horde/Horde.Build/Dockerfile
--- a/Engine/Source/Programs/Horde/Horde.Build/Dockerfile
+++ b/Engine/Source/Programs/Horde/Horde.Build/Dockerfile
@@ -55,12 +55,19 @@ COPY Source/Programs/Horde/Horde.Build.Tests/ Source/Programs/Horde/Horde.Build.
 COPY Source/Programs/Shared/EpicGames.Serialization.Tests/ Source/Programs/Shared/EpicGames.Serialization.Tests/
 RUN redis-server --save "" --appendonly no --daemonize yes && \
 	mongod --noauth --quiet --fork --dbpath /tmp/mongodb --logpath /tmp/mongod.log && \
-	dotnet test Source/Programs/Shared/EpicGames.BuildGraph.Tests/EpicGames.BuildGraph.Tests.csproj --blame-hang-timeout 2m && \
-	dotnet test Source/Programs/Shared/EpicGames.Core.Tests/EpicGames.Core.Tests.csproj --blame-hang-timeout 2m && \
 	dotnet test Source/Programs/Shared/EpicGames.Horde.Tests/EpicGames.Horde.Tests.csproj --blame-hang-timeout 2m && \
 	dotnet test Source/Programs/Shared/EpicGames.Serialization.Tests/EpicGames.Serialization.Tests.csproj --blame-hang-timeout 2m && \
 	dotnet test Source/Programs/Horde/Horde.Build.Tests/Horde.Build.Tests.csproj --blame-hang-timeout 2m
 
+# Run tests on NET Core 3.1 - disabled due to warning: CSC : warning CS8032: An instance of analyzer Microsoft.Extensions.Logging.Generators.LoggerMessageGenerator cannot be created..
+# FROM mcr.microsoft.com/dotnet/sdk:3.1 AS test-env
+# WORKDIR /app
+# COPY Source/Programs/Shared/ Source/Programs/Shared/
+# RUN dotnet restore Source/Programs/Shared/EpicGames.BuildGraph.Tests/EpicGames.BuildGraph.Tests.csproj && \
+#	dotnet restore Source/Programs/Shared/EpicGames.Core.Tests/EpicGames.Core.Tests.csproj
+# RUN dotnet test Source/Programs/Shared/EpicGames.BuildGraph.Tests/EpicGames.BuildGraph.Tests.csproj && \
+#	dotnet test Source/Programs/Shared/EpicGames.Core.Tests/EpicGames.Core.Tests.csproj
+
 # Create the runtime image
 FROM mcr.microsoft.com/dotnet/aspnet:6.0 AS runtime
 WORKDIR /app
