Horde: Add alternative log chunk writer with a concurrency limit

* Limits number of concurrent writes to avoid high memory and I/O usage.
* Can be toggled by a feature flag.

diff --git a/Engine/Source/Programs/Horde/Horde.Build.Tests/LogFileServiceTests.cs b/Engine/Source/Programs/Horde/Horde.Build.Tests/LogFileServiceTests.cs
--- a/Engine/Source/Programs/Horde/Horde.Build.Tests/LogFileServiceTests.cs
+++ b/Engine/Source/Programs/Horde/Horde.Build.Tests/LogFileServiceTests.cs
@@ -45,7 +45,8 @@ namespace Horde.Build.Tests
 			ILogBuilder logBuilder = new RedisLogBuilder(GetRedisServiceSingleton().ConnectionPool, NullLogger.Instance);
 			_logStorage = new PersistentLogStorage(new MemoryStorageBackend().ForType<PersistentLogStorage>(), NullLogger<PersistentLogStorage>.Instance);
 			_clock = new FakeClock();
-			_logFileService = new LogFileService(logFileCollection, null!, logBuilder, _logStorage, _clock, logger);
+			TestOptionsMonitor<ServerSettings> settingsOpts = new (new ServerSettings());
+			_logFileService = new LogFileService(logFileCollection, null!, logBuilder, _logStorage, _clock, settingsOpts, logger);
         }
 
 		protected override void Dispose(bool disposing)
diff --git a/Engine/Source/Programs/Horde/Horde.Build.Tests/LogIndexingTests.cs b/Engine/Source/Programs/Horde/Horde.Build.Tests/LogIndexingTests.cs
--- a/Engine/Source/Programs/Horde/Horde.Build.Tests/LogIndexingTests.cs
+++ b/Engine/Source/Programs/Horde/Horde.Build.Tests/LogIndexingTests.cs
@@ -48,7 +48,8 @@ namespace Horde.Build.Tests
 			ILogBuilder logBuilder = new LocalLogBuilder();
 			_nullLogStorage = new NullLogStorage();
 			_logStorage = new LocalLogStorage(20, _nullLogStorage);
-			_logFileService = new LogFileService(logFileCollection, null!, logBuilder, _logStorage, new FakeClock(), logger);
+			TestOptionsMonitor<ServerSettings> settingsOpts = new (new ServerSettings());
+			_logFileService = new LogFileService(logFileCollection, null!, logBuilder, _logStorage, new FakeClock(), settingsOpts, logger);
 		}
 
 		protected override void Dispose(bool disposing)
diff --git a/Engine/Source/Programs/Horde/Horde.Build/Logs/LogFileService.cs b/Engine/Source/Programs/Horde/Horde.Build/Logs/LogFileService.cs
--- a/Engine/Source/Programs/Horde/Horde.Build/Logs/LogFileService.cs
+++ b/Engine/Source/Programs/Horde/Horde.Build/Logs/LogFileService.cs
@@ -21,8 +21,8 @@ using HordeCommon;
 using Microsoft.Extensions.Caching.Memory;
 using Microsoft.Extensions.Hosting;
 using Microsoft.Extensions.Logging;
+using Microsoft.Extensions.Options;
 using MongoDB.Bson;
-using MongoDB.Driver;
 using OpenTracing;
 using OpenTracing.Util;
 using Stream = System.IO.Stream;
@@ -217,11 +217,14 @@ namespace Horde.Build.Logs
 	/// </summary>
 	public sealed class LogFileService : IHostedService, ILogFileService, IDisposable
 	{
+		private const int MaxConcurrentChunkWrites = 10;
+		
 		private readonly ILogger<LogFileService> _logger;
 		private readonly ILogFileCollection _logFiles;
 		private readonly ILogEventCollection _logEvents;
 		private readonly ILogStorage _storage;
 		private readonly ILogBuilder _builder;
+		private readonly IOptionsMonitor<ServerSettings> _settings;
 
 		// Lock object for the <see cref="_writeTasks"/> and <see cref="_writeChunks"/> members
 		private readonly object _writeLock = new object();
@@ -397,7 +400,7 @@ namespace Horde.Build.Logs
 		/// <summary>
 		/// Constructor
 		/// </summary>
-		public LogFileService(ILogFileCollection logFiles, ILogEventCollection logEvents, ILogBuilder builder, ILogStorage storage, IClock clock, ILogger<LogFileService> logger)
+		public LogFileService(ILogFileCollection logFiles, ILogEventCollection logEvents, ILogBuilder builder, ILogStorage storage, IClock clock, IOptionsMonitor<ServerSettings> settings, ILogger<LogFileService> logger)
 		{
 			_logFiles = logFiles;
 			_logEvents = logEvents;
@@ -405,6 +408,7 @@ namespace Horde.Build.Logs
 			_builder = builder;
 			_storage = storage;
 			_ticker = clock.AddTicker<LogFileService>(TimeSpan.FromSeconds(30.0), TickAsync, logger);
+			_settings = settings;
 			_logger = logger;
 		}
 
@@ -1001,14 +1005,14 @@ namespace Horde.Build.Logs
 					_logger.LogError(ex, "Exception while waiting for write tasks to complete");
 				}
 			}
-			await IncrementalFlush();
+			await IncrementalFlushAsync(stoppingToken);
 		}
 				
 		/// <summary>
 		/// Flushes complete chunks to the storage provider
 		/// </summary>
 		/// <returns>Async task</returns>
-		private async Task IncrementalFlush()
+		private async Task IncrementalFlushAsync(CancellationToken cancellationToken)
 		{
 			using IScope scope = GlobalTracer.Instance.BuildSpan("LogFileService.IncrementalFlush").StartActive();
 			
@@ -1024,8 +1028,16 @@ namespace Horde.Build.Logs
 				await _builder.CompleteChunkAsync(logId, offset);
 			}
 
-			// Add tasks for flushing all the chunks
-			WriteCompleteChunks(flushChunks, true);
+			if (_settings.CurrentValue.FeatureFlags.LimitConcurrentLogChunkWriting)
+			{
+				// Flush all the chunks and await completion instead of running them async
+				await WriteCompleteChunksV2Async(flushChunks, true, cancellationToken);
+			}
+			else
+			{
+				// Add tasks for flushing all the chunks
+				WriteCompleteChunks(flushChunks, true);
+			}
 		}
 
 		/// <summary>
@@ -1113,6 +1125,50 @@ namespace Horde.Build.Logs
 			scope.Span.SetTag("numWriteTasksCreated", numTasksCreated);
 			_logger.LogInformation("{NumWriteTasksCreated} write tasks created", numTasksCreated);
 		}
+		
+		/// <summary>
+		/// Writes list of complete chunks
+		/// </summary>
+		/// <param name="chunksToWrite">List of chunks to write</param>
+		/// <param name="bCreateIndex">Create an index for the log</param>
+		/// <param name="cancellationToken">Cancellation token</param>
+		private async Task WriteCompleteChunksV2Async(List<(LogId, long)> chunksToWrite, bool bCreateIndex, CancellationToken cancellationToken)
+		{
+			using IScope scope = GlobalTracer.Instance.BuildSpan("LogFileService.WriteCompleteChunksV2Async").StartActive();
+			
+			HashSet<(LogId, long)> writeChunks = new ();
+			List<(LogId, List<long>)> offsetsToWrite = new();
+			
+			foreach (IGrouping<LogId, long> group in chunksToWrite.GroupBy(x => x.Item1, x => x.Item2))
+			{
+				LogId logId = group.Key;
+
+				// Find offsets of new chunks to write
+				List<long> offsets = new ();
+				foreach (long offset in group.OrderBy(x => x))
+				{
+					if (writeChunks.Add((logId, offset)))
+					{
+						offsets.Add(offset);
+					}
+				}
+				
+				// Create the write task
+				if (offsets.Count > 0)
+				{
+					offsetsToWrite.Add((logId, offsets));
+				}
+			}
+
+			_logger.LogInformation("Running {NumWrites} writes with a concurrency of {Concurrency}", offsetsToWrite.Count, MaxConcurrentChunkWrites);
+			scope.Span.SetTag("numOffsetsToWrite", offsetsToWrite.Count);
+			ParallelOptions opts = new() { MaxDegreeOfParallelism = MaxConcurrentChunkWrites, CancellationToken = cancellationToken };
+			await Parallel.ForEachAsync(offsetsToWrite, opts, async (x, innerCt) =>
+			{
+				(LogId logId, List<long> offsets) = x;
+				await WriteCompleteChunksForLogAsync(logId, offsets, bCreateIndex, innerCt);
+			});
+		}
 
 		/// <summary>
 		/// Writes a set of chunks to the database
@@ -1120,13 +1176,14 @@ namespace Horde.Build.Logs
 		/// <param name="logId">Log file to update</param>
 		/// <param name="offsets">Chunks to write</param>
 		/// <param name="bCreateIndex">Whether to create the index for this log</param>
+		/// <param name="cancellationToken">Cancellation token for the call</param>
 		/// <returns>Async task</returns>
-		private async Task<ILogFile?> WriteCompleteChunksForLogAsync(LogId logId, List<long> offsets, bool bCreateIndex)
+		private async Task<ILogFile?> WriteCompleteChunksForLogAsync(LogId logId, List<long> offsets, bool bCreateIndex, CancellationToken cancellationToken = default)
 		{
 			ILogFile? logFile = await _logFiles.GetLogFileAsync(logId);
 			if(logFile != null)
 			{
-				logFile = await WriteCompleteChunksForLogAsync(logFile, offsets, bCreateIndex);
+				logFile = await WriteCompleteChunksForLogAsync(logFile, offsets, bCreateIndex, cancellationToken);
 			}
 			return logFile;
 		}
@@ -1137,8 +1194,9 @@ namespace Horde.Build.Logs
 		/// <param name="logFileInterface">Log file to update</param>
 		/// <param name="offsets">Chunks to write</param>
 		/// <param name="bCreateIndex">Whether to create the index for this log</param>
+		/// <param name="cancellationToken">Cancellation token for the call</param>
 		/// <returns>Async task</returns>
-		private async Task<ILogFile?> WriteCompleteChunksForLogAsync(ILogFile logFileInterface, List<long> offsets, bool bCreateIndex)
+		private async Task<ILogFile?> WriteCompleteChunksForLogAsync(ILogFile logFileInterface, List<long> offsets, bool bCreateIndex, CancellationToken cancellationToken = default)
 		{
 			using IScope scope = GlobalTracer.Instance.BuildSpan("WriteCompleteChunksForLogAsync").StartActive();
 			scope.Span.SetTag("LogId", logFileInterface.Id.ToString());
@@ -1157,6 +1215,9 @@ namespace Horde.Build.Logs
 					chunkWriteTasks.Add(Task.Run(() => WriteChunkAsync(logFileInterface.Id, offset, lineIndex)));
 				}
 			}
+			
+			scope.Span.SetTag("NumWriteTasks", chunkWriteTasks.Count);
+			_logger.LogInformation("Writing {NumLogChunks} in parallel", chunkWriteTasks.Count);
 
 			// Wait for the tasks to complete, periodically updating the log file object
 			ILogFile? logFile = logFileInterface;
diff --git a/Engine/Source/Programs/Horde/Horde.Build/ServerSettings.cs b/Engine/Source/Programs/Horde/Horde.Build/ServerSettings.cs
--- a/Engine/Source/Programs/Horde/Horde.Build/ServerSettings.cs
+++ b/Engine/Source/Programs/Horde/Horde.Build/ServerSettings.cs
@@ -171,6 +171,11 @@ namespace Horde.Build
 		/// Use new auth config for custom auth servers in Okta
 		/// </summary>
 		public bool AuthSettingsV2 { get; set; } = false;
+		
+		/// <summary>
+		/// Limit concurrent log chunk writes and await them to reduce mem and I/O usage
+		/// </summary>
+		public bool LimitConcurrentLogChunkWriting { get; set; } = false;
 	}
 
 	/// <summary>
