Horde: Fix NET 7 warnings.

diff --git a/Engine/Source/Programs/Horde/Horde.Build.Tests/Compute/RedisTaskSchedulerTest.cs b/Engine/Source/Programs/Horde/Horde.Build.Tests/Compute/RedisTaskSchedulerTest.cs
--- a/Engine/Source/Programs/Horde/Horde.Build.Tests/Compute/RedisTaskSchedulerTest.cs
+++ b/Engine/Source/Programs/Horde/Horde.Build.Tests/Compute/RedisTaskSchedulerTest.cs
@@ -22,6 +22,16 @@ namespace Horde.Build.Tests.Compute
 			_scheduler = new RedisTaskScheduler<string, string>(redisService.ConnectionPool, new RedisKey("myBaseKey"), loggerFactory.CreateLogger<RedisTaskScheduler<string, string>>());
 		}
 
+		protected override void Dispose(bool disposing)
+		{
+			base.Dispose(disposing);
+
+			if (disposing)
+			{
+				_scheduler.Dispose();
+			}
+		}
+
 		[TestMethod]
 		public async Task EnqueueAndDequeueTask()
 		{
diff --git a/Engine/Source/Programs/Horde/Horde.Build.Tests/Fleet/AutoscaleServiceV2Test.cs b/Engine/Source/Programs/Horde/Horde.Build.Tests/Fleet/AutoscaleServiceV2Test.cs
--- a/Engine/Source/Programs/Horde/Horde.Build.Tests/Fleet/AutoscaleServiceV2Test.cs
+++ b/Engine/Source/Programs/Horde/Horde.Build.Tests/Fleet/AutoscaleServiceV2Test.cs
@@ -72,6 +72,16 @@ namespace Horde.Build.Tests.Fleet
 		readonly FleetManagerSpy _fleetManagerSpy = new();
 		readonly IDogStatsd _dogStatsD = new NoOpDogStatsd();
 
+		protected override void Dispose(bool disposing)
+		{
+			base.Dispose(disposing);
+
+			if (disposing)
+			{
+				_dogStatsD.Dispose();
+			}
+		}
+
 		[TestMethod]
 		public async Task OnlyEnabledAgentsAreAutoScaled()
 		{
diff --git a/Engine/Source/Programs/Horde/Horde.Build.Tests/VersionedCollectionTests.cs b/Engine/Source/Programs/Horde/Horde.Build.Tests/VersionedCollectionTests.cs
--- a/Engine/Source/Programs/Horde/Horde.Build.Tests/VersionedCollectionTests.cs
+++ b/Engine/Source/Programs/Horde/Horde.Build.Tests/VersionedCollectionTests.cs
@@ -5,6 +5,7 @@ using System.Collections.Generic;
 using System.Threading.Tasks;
 using Horde.Build.Server;
 using Horde.Build.Utilities;
+using Microsoft.Extensions.DependencyInjection;
 using Microsoft.VisualStudio.TestTools.UnitTesting;
 using MongoDB.Driver;
 using StackExchange.Redis;
@@ -46,16 +47,14 @@ namespace Horde.Build.Tests
 			[2] = typeof(DocumentV2),
 		};
 
-		readonly MongoService _mongoService;
-		readonly RedisService _redisService;
+		MongoService _mongoService => GetMongoServiceSingleton();
+		RedisService _redisService => GetRedisServiceSingleton();
 		readonly RedisKey _baseKey = new RedisKey("versioned/");
 		readonly IMongoCollection<VersionedDocument<string, DocumentV2>> _baseCollection;
 		readonly VersionedCollection<string, DocumentV2> _collection;
 
 		public VersionedCollectionTests()
 		{
-			_mongoService = GetMongoServiceSingleton();
-			_redisService = GetRedisServiceSingleton();
 			_collection = new VersionedCollection<string, DocumentV2>(_mongoService, "versioned", _redisService, _baseKey, s_documentTypes);
 			_baseCollection = _collection.BaseCollection;
 		}
