Horde: Preflight change build dialog fixes
#skipci
#rnx

diff --git a/Engine/Source/Programs/Horde/HordeDashboard/src/backend/Api.ts b/Engine/Source/Programs/Horde/HordeDashboard/src/backend/Api.ts
--- a/Engine/Source/Programs/Horde/HordeDashboard/src/backend/Api.ts
+++ b/Engine/Source/Programs/Horde/HordeDashboard/src/backend/Api.ts
@@ -1969,6 +1969,8 @@ export type DefaultPreflightConfig = {
 
 	/** The template id to query */
 	templateId?: string;
+
+	change?: ChangeQueryConfig;
 }
 
 
diff --git a/Engine/Source/Programs/Horde/HordeDashboard/src/components/NewBuild.tsx b/Engine/Source/Programs/Horde/HordeDashboard/src/components/NewBuild.tsx
--- a/Engine/Source/Programs/Horde/HordeDashboard/src/components/NewBuild.tsx
+++ b/Engine/Source/Programs/Horde/HordeDashboard/src/components/NewBuild.tsx
@@ -735,8 +735,23 @@ export const NewBuild: React.FC<{ streamId: string; show: boolean; onClose: (new
 
    if (template?.defaultChange) {
       defaultPreflightQuery = template.defaultChange;
+   } else {
+      const defaultChange = stream?.defaultPreflight?.change;
+
+      if (defaultChange) {
+         if (!defaultChange.name && defaultChange.templateId) {
+            const defaultTemplate = templateData.templates.find(t => t.id === defaultChange.templateId);
+            if (defaultTemplate) {
+               defaultChange.name = "Latest Success - " + defaultTemplate.name;
+            }          
+         }
+         
+         if (defaultChange.name) {
+            defaultPreflightQuery = [defaultChange];
+         }         
+      }      
    }
-
+   
    const defaultStreamPreflightTemplate = stream.templates.find(t => t.id === stream.defaultPreflight?.templateId);
 
    if (!defaultPreflightQuery && stream.defaultPreflight?.templateId) {
@@ -748,8 +763,7 @@ export const NewBuild: React.FC<{ streamId: string; show: boolean; onClose: (new
       if (!defaultPreflightQuery) {
          console.error(`Unable to find default stream preflight template ${stream.defaultPreflight?.templateId} in stream templates`);
       }
-   }
-
+   }   
 
    if (buildParams?.preflight) {
       templates = templates.filter(t => t.allowPreflights);
@@ -1253,8 +1267,8 @@ export const NewBuild: React.FC<{ streamId: string; show: boolean; onClose: (new
       let changeQueries: ChangeQueryConfig[] | undefined;
 
       if (typeof change !== 'number') {
-         const changeOption = buildParams.changeOption ?? "Default Change";
-         if (changeOption === "Default Change" && defaultPreflightQuery) {
+         const changeOption = buildParams.changeOption ?? "Latest Change";
+         if (changeOption !== "Latest Change" && defaultPreflightQuery) {
             changeQueries = defaultPreflightQuery;
          }
          else {
@@ -1457,7 +1471,7 @@ export const NewBuild: React.FC<{ streamId: string; show: boolean; onClose: (new
 
    // change setup
 
-   const changeOptions = ["Default Change"];
+   const changeOptions = ["Latest Change"];
 
    if (defaultPreflightQuery) {
       changeOptions.push(...defaultPreflightQuery.filter(n => !!n.name).map(n => n.name!));
@@ -1561,7 +1575,7 @@ export const NewBuild: React.FC<{ streamId: string; show: boolean; onClose: (new
                      </Stack>
                      <Stack grow />
                      <Stack>
-                        <ComboBox style={{ width: 240 }} label="Change" text={changeText} options={changeItems} disabled={readOnly} allowFreeform autoComplete="off" placeholder="Default Change" defaultValue={jobDetails?.jobData?.change?.toString()} onChange={(ev, option, index, value) => {
+                        <ComboBox style={{ width: 240 }} label="Change" text={changeText} options={changeItems} disabled={readOnly} allowFreeform autoComplete="off" placeholder="Latest Change" defaultValue={jobDetails?.jobData?.change?.toString()} onChange={(ev, option, index, value) => {
                            ev.preventDefault();
 
                            if (option) {
