Horde: Retire the device notification Slack channel now that better captured in telemetry
#rnx

diff --git a/Engine/Source/Programs/Horde/Horde.Build/Devices/DevicesController.cs b/Engine/Source/Programs/Horde/Horde.Build/Devices/DevicesController.cs
--- a/Engine/Source/Programs/Horde/Horde.Build/Devices/DevicesController.cs
+++ b/Engine/Source/Programs/Horde/Horde.Build/Devices/DevicesController.cs
@@ -833,8 +833,7 @@ namespace Horde.Build.Devices
 
 			if (String.IsNullOrEmpty(poolId))
 			{
-				string message = $"No pool specified, defaulting to UE4" + details;
-				await _deviceService.NotifyDeviceServiceAsync(message, null, request.JobId, request.StepId);
+				// @todo: We default to ue4, though this should be an error, or a configuration setting
 				poolId = "ue4";
 				//return BadRequest(Message);
 			}
@@ -845,7 +844,6 @@ namespace Horde.Build.Devices
 			{
 				_logger.LogError("Unknown pool {PoolId} {Details}", poolId, details);
 				string message = $"Unknown pool {poolId} " + details;
-				await _deviceService.NotifyDeviceServiceAsync(message, null, request.JobId, request.StepId);
 				return BadRequest(message);
 			}
 
@@ -884,9 +882,7 @@ namespace Horde.Build.Devices
 
 				if (platform == null)
 				{
-					string message = $"Unknown platform {platformId}" + details;					
-					await _deviceService.NotifyDeviceServiceAsync(message, null, request.JobId, request.StepId);
-					return BadRequest(message);
+					return BadRequest($"Unknown platform {platformId}" + details);
 				}
 
 				List<string> includeModels = new List<string>();
@@ -1138,8 +1134,6 @@ namespace Horde.Build.Devices
 				}
 			}			
 
-			await _deviceService.NotifyDeviceServiceAsync(message, device.Id, jobId, stepId);
-
 			return Ok();
 
 		}
diff --git a/Engine/Source/Programs/Horde/Horde.Build/Notifications/Sinks/SlackNotificationSink.cs b/Engine/Source/Programs/Horde/Horde.Build/Notifications/Sinks/SlackNotificationSink.cs
--- a/Engine/Source/Programs/Horde/Horde.Build/Notifications/Sinks/SlackNotificationSink.cs
+++ b/Engine/Source/Programs/Horde/Horde.Build/Notifications/Sinks/SlackNotificationSink.cs
@@ -1814,7 +1814,7 @@ namespace Horde.Build.Notifications.Sinks
 		/// <inheritdoc/>
 		public async Task NotifyDeviceServiceAsync(string message, IDevice? device = null, IDevicePool? pool = null, IStream? stream = null, IJob? job = null, IJobStep? step = null, INode? node = null, IUser? user = null)
 		{
-			string recipient = $"#{_settings.DeviceServiceNotificationChannel}";
+			string? recipient = null;
 
 			if (user != null)
 			{
@@ -1828,11 +1828,10 @@ namespace Horde.Build.Notifications.Sinks
 
 				recipient = slackRecipient;
 			}
-
-			_logger.LogDebug("Sending device service notification to {Recipient}", recipient);
-
-			if (_settings.DeviceServiceNotificationChannel != null)
+			
+			if (recipient != null)
 			{
+				_logger.LogDebug("Sending device service notification to {Recipient}", recipient);
 				await SendDeviceServiceMessage(recipient, message, device, pool, stream, job, step, node, user);
 			}
 		}
diff --git a/Engine/Source/Programs/Horde/Horde.Build/ServerSettings.cs b/Engine/Source/Programs/Horde/Horde.Build/ServerSettings.cs
--- a/Engine/Source/Programs/Horde/Horde.Build/ServerSettings.cs
+++ b/Engine/Source/Programs/Horde/Horde.Build/ServerSettings.cs
@@ -434,11 +434,6 @@ namespace Horde.Build
 		/// Channel to send stream notification update failures to
 		/// </summary>
 		public string? UpdateStreamsNotificationChannel { get; set; }
-
-		/// <summary>
-		/// Channel to send device notifications to
-		/// </summary>
-		public string? DeviceServiceNotificationChannel { get; set; }
 		
 		/// <summary>
 		/// Slack channel to send job related notifications to
