Horde Storage - Increased page size for the replication log retrival - should result in less scanning of db as once we know were we are we just read incrementally.

diff --git a/Engine/Source/Programs/Horde/HordeStorage/Horde.Storage/Controllers/ReplicationLogController.cs b/Engine/Source/Programs/Horde/HordeStorage/Horde.Storage/Controllers/ReplicationLogController.cs
--- a/Engine/Source/Programs/Horde/HordeStorage/Horde.Storage/Controllers/ReplicationLogController.cs
+++ b/Engine/Source/Programs/Horde/HordeStorage/Horde.Storage/Controllers/ReplicationLogController.cs
@@ -81,7 +81,7 @@ namespace Horde.Storage.Controllers
             [Required] NamespaceId ns,
             [FromQuery] string? lastBucket,
             [FromQuery] Guid? lastEvent,
-            [FromQuery] int count = 100
+            [FromQuery] int count = 1000
         )
         {
             ActionResult? result = await _requestHelper.HasAccessToNamespace(User, Request, ns, new [] { AclAction.ReadTransactionLog });
