Horde Storage - Refactored blob cleanup to avoid calculating disk size twice before starting the cleanup as that can take a very long time.

diff --git a/Engine/Source/Programs/Horde/HordeStorage/Horde.Storage/Implementation/Blob/FileSystemStore.cs b/Engine/Source/Programs/Horde/HordeStorage/Horde.Storage/Implementation/Blob/FileSystemStore.cs
--- a/Engine/Source/Programs/Horde/HordeStorage/Horde.Storage/Implementation/Blob/FileSystemStore.cs
+++ b/Engine/Source/Programs/Horde/HordeStorage/Horde.Storage/Implementation/Blob/FileSystemStore.cs
@@ -218,21 +218,23 @@ namespace Horde.Storage.Implementation
         {
             using IScope scope = Tracer.Instance.StartActive("gc.filesystem");
 
-            long size = await CalculateDiskSpaceUsed();
             ulong maxSizeBytes = _settings.CurrentValue.MaxSizeBytes;
             long triggerSize = (long) (maxSizeBytes * _settings.CurrentValue.TriggerThresholdPercentage);
             long targetSize = (long) (maxSizeBytes * _settings.CurrentValue.TargetThresholdPercentage); // Target to shrink to if triggered
             ulong countOfBlobsRemoved = 0;
 
-            if (size < triggerSize)
-            {
-                return 0;
-            }
-
             // Perform a maximum of 5 clean up runs
             for (int i = 0; i < 5; i++)
             {
-                size = await CalculateDiskSpaceUsed();
+                long size = await CalculateDiskSpaceUsed();
+
+                // first check to see if we should trigger at all, this happens for each run but only really matters for the first attempt
+                if (size < triggerSize)
+                {
+                    return countOfBlobsRemoved;
+                }
+
+                // then check if we have reached the target size, if not we should continue running
                 if (size <= targetSize)
                 {
                     return countOfBlobsRemoved;
