Horde: Fix exceptions on shutdown causing test failures.

diff --git a/Engine/Source/Programs/Horde/Horde.Build/Compute/RedisTaskScheduler.cs b/Engine/Source/Programs/Horde/Horde.Build/Compute/RedisTaskScheduler.cs
--- a/Engine/Source/Programs/Horde/Horde.Build/Compute/RedisTaskScheduler.cs
+++ b/Engine/Source/Programs/Horde/Horde.Build/Compute/RedisTaskScheduler.cs
@@ -121,7 +121,13 @@ namespace Horde.Build.Compute
 			if (!_queueUpdateTask.IsCompleted)
 			{
 				_cancellationSource.Cancel();
-				await _queueUpdateTask;
+				try
+				{
+					await _queueUpdateTask;
+				}
+				catch (OperationCanceledException)
+				{
+				}
 			}
 			_cancellationSource.Dispose();
 		}
