Horde: Include inputs & outputs in graph templates.

diff --git a/Engine/Source/Programs/Horde/Horde.Agent/Execution/BuildGraphExecutor.cs b/Engine/Source/Programs/Horde/Horde.Agent/Execution/BuildGraphExecutor.cs
--- a/Engine/Source/Programs/Horde/Horde.Agent/Execution/BuildGraphExecutor.cs
+++ b/Engine/Source/Programs/Horde/Horde.Agent/Execution/BuildGraphExecutor.cs
@@ -33,6 +33,8 @@ namespace Horde.Agent.Execution
 			public string Name { get; set; } = String.Empty;
 			public bool RunEarly { get; set; }
 			public bool? Warnings { get; set; }
+			public List<string> Inputs { get; set; } = new List<string>();
+			public List<string> Outputs { get; set; } = new List<string>();
 			public List<string> InputDependencies { get; set; } = new List<string>();
 			public List<string> OrderDependencies { get; set; } = new List<string>();
 			public Dictionary<string, string> Annotations { get; set; } = new Dictionary<string, string>(StringComparer.OrdinalIgnoreCase);
@@ -445,6 +447,14 @@ namespace Horde.Agent.Execution
 				{
 					CreateNodeRequest createNode = new CreateNodeRequest();
 					createNode.Name = exportedNode.Name;
+					if (exportedNode.Inputs != null)
+					{
+						createNode.Inputs.Add(exportedNode.Inputs);
+					}
+					if (exportedNode.Outputs != null)
+					{
+						createNode.Outputs.Add(exportedNode.Outputs);
+					}
 					if (exportedNode.InputDependencies != null)
 					{
 						createNode.InputDependencies.Add(exportedNode.InputDependencies);
diff --git a/Engine/Source/Programs/Horde/Horde.Build/Jobs/Graphs/GraphCollection.cs b/Engine/Source/Programs/Horde/Horde.Build/Jobs/Graphs/GraphCollection.cs
--- a/Engine/Source/Programs/Horde/Horde.Build/Jobs/Graphs/GraphCollection.cs
+++ b/Engine/Source/Programs/Horde/Horde.Build/Jobs/Graphs/GraphCollection.cs
@@ -33,6 +33,12 @@ namespace Horde.Build.Jobs.Graphs
 			[BsonRequired]
 			public string Name { get; set; }
 
+			[BsonIgnoreIfNull]
+			public NodeOutputRef[]? Inputs { get; set; }
+
+			[BsonIgnoreIfNull]
+			public string[]? OutputNames { get; set; }
+
 			[BsonIgnoreIfNull]
 			public NodeRef[] InputDependencies { get; set; }
 
@@ -54,6 +60,8 @@ namespace Horde.Build.Jobs.Graphs
 			[BsonDictionaryOptions(DictionaryRepresentation.ArrayOfDocuments)]
 			public NodeAnnotations? Annotations { get; set; }
 
+			IReadOnlyList<NodeOutputRef> INode.Inputs => Inputs ?? Array.Empty<NodeOutputRef>();
+			IReadOnlyList<string> INode.OutputNames => OutputNames ?? Array.Empty<string>();
 			IReadOnlyDictionary<string, string>? INode.Credentials => Credentials;
 			IReadOnlyDictionary<string, string>? INode.Properties => Properties;
 			IReadOnlyNodeAnnotations INode.Annotations => Annotations ?? NodeAnnotations.Empty;
@@ -66,9 +74,11 @@ namespace Horde.Build.Jobs.Graphs
 				OrderDependencies = null!;
 			}
 
-			public Node(string name, NodeRef[] inputDependencies, NodeRef[] orderDependencies, Priority priority, bool allowRetry, bool runEarly, bool warnings, Dictionary<string, string>? credentials, Dictionary<string, string>? properties, IReadOnlyNodeAnnotations? annotations)
+			public Node(string name, NodeOutputRef[]? inputs, string[]? outputNames, NodeRef[] inputDependencies, NodeRef[] orderDependencies, Priority priority, bool allowRetry, bool runEarly, bool warnings, Dictionary<string, string>? credentials, Dictionary<string, string>? properties, IReadOnlyNodeAnnotations? annotations)
 			{
 				Name = name;
+				Inputs = (inputs != null && inputs.Length > 0) ? inputs: null;
+				OutputNames = (outputNames != null && outputNames.Length > 0) ? outputNames : null;
 				InputDependencies = inputDependencies;
 				OrderDependencies = orderDependencies;
 				Priority = priority;
@@ -195,6 +205,9 @@ namespace Horde.Build.Jobs.Graphs
 			[BsonIgnore]
 			IReadOnlyDictionary<string, NodeRef>? _cachedNodeNameToRef;
 
+			[BsonIgnore]
+			IReadOnlyDictionary<string, NodeOutputRef>? _cachedNodeOutputNameToRef;
+
 			IReadOnlyList<INodeGroup> IGraph.Groups => Groups;
 			IReadOnlyList<IAggregate> IGraph.Aggregates => Aggregates;
 			IReadOnlyList<ILabel> IGraph.Labels => Labels;
@@ -215,6 +228,7 @@ namespace Horde.Build.Jobs.Graphs
 			public GraphDocument(GraphDocument baseGraph, List<NewGroup>? newGroupRequests, List<NewAggregate>? newAggregateRequests, List<NewLabel>? newLabelRequests)
 			{
 				Dictionary<string, NodeRef> nodeNameToRef = new Dictionary<string, NodeRef>(baseGraph.GetNodeNameToRef(), StringComparer.OrdinalIgnoreCase);
+				Dictionary<string, NodeOutputRef> nodeOutputNameToRef = new Dictionary<string, NodeOutputRef>(baseGraph.GetNodeOutputNameToRef(), StringComparer.OrdinalIgnoreCase);
 
 				// Update the new list of groups
 				List<NodeGroup> newGroups = new List<NodeGroup>(baseGraph.Groups);
@@ -232,12 +246,28 @@ namespace Horde.Build.Jobs.Graphs
 							bool bRunEarly = newNodeRequest.RunEarly ?? false;
 							bool bWarnings = newNodeRequest.Warnings ?? true;
 
+							NodeOutputRef[]? inputs = null;
+							if (newNodeRequest.Inputs != null && newNodeRequest.Inputs.Count > 0)
+							{
+								inputs = newNodeRequest.Inputs.Select(x => nodeOutputNameToRef[x]).ToArray();
+							}
+
 							NodeRef[] inputDependencies = (newNodeRequest.InputDependencies == null) ? Array.Empty<NodeRef>() : newNodeRequest.InputDependencies.Select(x => nodeNameToRef[x]).ToArray();
 							NodeRef[] orderDependencies = (newNodeRequest.OrderDependencies == null) ? Array.Empty<NodeRef>() : newNodeRequest.OrderDependencies.Select(x => nodeNameToRef[x]).ToArray();
 							orderDependencies = orderDependencies.Union(inputDependencies).ToArray();
-							nodes.Add(new Node(newNodeRequest.Name, inputDependencies, orderDependencies, priority, bAllowRetry, bRunEarly, bWarnings, newNodeRequest.Credentials, newNodeRequest.Properties, newNodeRequest.Annotations));
-
-							nodeNameToRef.Add(newNodeRequest.Name, new NodeRef(newGroups.Count, nodeIdx));
+							nodes.Add(new Node(newNodeRequest.Name, inputs, newNodeRequest.OutputNames?.ToArray(), inputDependencies, orderDependencies, priority, bAllowRetry, bRunEarly, bWarnings, newNodeRequest.Credentials, newNodeRequest.Properties, newNodeRequest.Annotations));
+
+							NodeRef nodeRef = new NodeRef(newGroups.Count, nodeIdx);
+							nodeNameToRef.Add(newNodeRequest.Name, nodeRef);
+
+							if (newNodeRequest.OutputNames != null)
+							{
+								for (int outputIdx = 0; outputIdx < newNodeRequest.OutputNames.Count; outputIdx++)
+								{
+									string outputName = newNodeRequest.OutputNames[outputIdx];
+									nodeOutputNameToRef[$"{newNodeRequest.Name}/{outputName}"] = new NodeOutputRef(nodeRef, outputIdx);
+								}
+							}
 						}
 						newGroups.Add(new NodeGroup(newGroupRequest.AgentType, nodes));
 					}
@@ -296,6 +326,32 @@ namespace Horde.Build.Jobs.Graphs
 				}
 				return _cachedNodeNameToRef;
 			}
+
+			public IReadOnlyDictionary<string, NodeOutputRef> GetNodeOutputNameToRef()
+			{
+				if (_cachedNodeOutputNameToRef == null)
+				{
+					Dictionary<string, NodeOutputRef> nodeOutputNameToRef = new Dictionary<string, NodeOutputRef>(StringComparer.OrdinalIgnoreCase);
+					for (int groupIdx = 0; groupIdx < Groups.Count; groupIdx++)
+					{
+						List<Node> nodes = Groups[groupIdx].Nodes;
+						for (int nodeIdx = 0; nodeIdx < nodes.Count; nodeIdx++)
+						{
+							Node node = nodes[nodeIdx];
+							if (node.OutputNames != null)
+							{
+								for (int outputIdx = 0; outputIdx < node.OutputNames.Length; outputIdx++)
+								{
+									string outputName = $"{node.Name}/{node.OutputNames[outputIdx]}";
+									nodeOutputNameToRef[outputName] = new NodeOutputRef(new NodeRef(groupIdx, nodeIdx), outputIdx);
+								}
+							}
+						}
+					}
+					_cachedNodeOutputNameToRef = nodeOutputNameToRef;
+				}
+				return _cachedNodeOutputNameToRef;
+			}
 		}
 
 		/// <summary>
@@ -395,7 +451,7 @@ namespace Horde.Build.Jobs.Graphs
 		/// <inheritdoc/>
 		public async Task<IGraph> AddAsync(ITemplate template)
 		{
-			Node node = new Node(IJob.SetupNodeName, Array.Empty<NodeRef>(), Array.Empty<NodeRef>(), Priority.High, true, false, true, null, null, null);
+			Node node = new Node(IJob.SetupNodeName, null, null, Array.Empty<NodeRef>(), Array.Empty<NodeRef>(), Priority.High, true, false, true, null, null, null);
 			NodeGroup group = new NodeGroup(template.InitialAgentType ?? "Win64", new List<Node> { node });
 
 			GraphDocument graph = new GraphDocument(new List<NodeGroup> { group }, new List<Aggregate>(), new List<Label>());
diff --git a/Engine/Source/Programs/Horde/Horde.Build/Jobs/Graphs/IGraph.cs b/Engine/Source/Programs/Horde/Horde.Build/Jobs/Graphs/IGraph.cs
--- a/Engine/Source/Programs/Horde/Horde.Build/Jobs/Graphs/IGraph.cs
+++ b/Engine/Source/Programs/Horde/Horde.Build/Jobs/Graphs/IGraph.cs
@@ -21,6 +21,16 @@ namespace Horde.Build.Jobs.Graphs
 		/// </summary>
 		public string Name { get; }
 
+		/// <summary>
+		/// References to inputs for this node
+		/// </summary>
+		public IReadOnlyList<NodeOutputRef> Inputs { get; }
+
+		/// <summary>
+		/// List of output names
+		/// </summary>
+		public IReadOnlyList<string> OutputNames { get; }
+
 		/// <summary>
 		/// Indices of nodes which must have succeeded for this node to run
 		/// </summary>
@@ -144,6 +154,38 @@ namespace Horde.Build.Jobs.Graphs
 		}
 	}
 
+	/// <summary>
+	/// Output from a node
+	/// </summary>
+	[DebuggerDisplay("{NodeRef}, Output: {OutputIdx}")]
+	public class NodeOutputRef
+	{
+		/// <summary>
+		/// Node producing the output
+		/// </summary>
+		public NodeRef NodeRef { get; set; }
+
+		/// <summary>
+		/// Index of the output
+		/// </summary>
+		public int OutputIdx { get; set; }
+
+		/// <summary>
+		/// Constructor
+		/// </summary>
+		public NodeOutputRef(NodeRef nodeRef, int outputIdx)
+		{
+			NodeRef = nodeRef;
+			OutputIdx = outputIdx;
+		}
+
+		/// <inheritdoc/>
+		public override bool Equals(object? other) => other is NodeOutputRef otherRef && otherRef.NodeRef == NodeRef && otherRef.OutputIdx == OutputIdx;
+
+		/// <inheritdoc/>
+		public override int GetHashCode() => HashCode.Combine(NodeRef.GetHashCode(), OutputIdx);
+	}
+
 	/// <summary>
 	/// An collection of node references
 	/// </summary>
@@ -378,6 +420,16 @@ namespace Horde.Build.Jobs.Graphs
 		/// </summary>
 		public string Name { get; set; } = null!;
 
+		/// <summary>
+		/// Input names
+		/// </summary>
+		public List<string>? Inputs { get; set; }
+
+		/// <summary>
+		/// Output names
+		/// </summary>
+		public List<string>? OutputNames { get; set; }
+
 		/// <summary>
 		/// List of nodes which must succeed for this node to run
 		/// </summary>
diff --git a/Engine/Source/Programs/Horde/Horde.Build/Server/RpcService.cs b/Engine/Source/Programs/Horde/Horde.Build/Server/RpcService.cs
--- a/Engine/Source/Programs/Horde/Horde.Build/Server/RpcService.cs
+++ b/Engine/Source/Programs/Horde/Horde.Build/Server/RpcService.cs
@@ -624,6 +624,14 @@ namespace Horde.Build.Server
 				response.Name = node.Name;
 				response.Credentials.Add(credentials);
 
+				foreach(NodeOutputRef input in node.Inputs)
+				{
+					INode inputNode = graph.GetNode(input.NodeRef);
+					response.Inputs.Add($"{inputNode.Name}/{inputNode.OutputNames[input.OutputIdx]}");
+				}
+
+				response.OutputNames.Add(node.OutputNames);
+
 				string templateName = "<unknown>";
 				if (job.TemplateHash != null)
 				{
