Horde: Fix stream summary making a bad request with cleared stream id
#skipci
#rnx

diff --git a/Engine/Source/Programs/Horde/HordeDashboard/src/components/StreamSummary.tsx b/Engine/Source/Programs/Horde/HordeDashboard/src/components/StreamSummary.tsx
--- a/Engine/Source/Programs/Horde/HordeDashboard/src/components/StreamSummary.tsx
+++ b/Engine/Source/Programs/Horde/HordeDashboard/src/components/StreamSummary.tsx
@@ -54,7 +54,9 @@ class SummaryHandler {
 
       try {
 
-         if (!this.streamId || this.updating) {
+         const streamId = this.streamId;
+
+         if (!streamId || this.updating) {
             return;
          }
 
@@ -66,7 +68,12 @@ class SummaryHandler {
          this.updating = true;
          const cancelID = this.cancelID++;
 
-         const values = await Promise.all([backend.getIssuesV2({ streamId: this.streamId, count: 512, resolved: false })]);
+         const values = await Promise.all([backend.getIssuesV2({ streamId: streamId, count: 512, resolved: false })]);
+
+         // early out if has been canceled 
+         if (this.canceled.has(cancelID)) {
+            return;
+         }
 
          let jiraKeys: string[] = [];
          
@@ -91,20 +98,20 @@ class SummaryHandler {
          });
 
          if (jiraKeys.length) {
-            const jiras = await backend.getExternalIssues(this.streamId, jiraKeys);
+            const jiras = await backend.getExternalIssues(streamId, jiraKeys);
             jiras.forEach(j => {
                this.jiraIssues.set(j.key, { issue: j, cacheTime: moment(Date.now()) });
             })
          }
 
-
-         this.issues = values[0].filter(p => p.promoted || jiraIssues.has(p.id));
-         this.unpromoted = values[0].filter(p => !p.promoted && !jiraIssues.has(p.id));
-
+         // early out if has been canceled 
          if (this.canceled.has(cancelID)) {
             return;
          }
 
+         this.issues = values[0].filter(p => p.promoted || jiraIssues.has(p.id));
+         this.unpromoted = values[0].filter(p => !p.promoted && !jiraIssues.has(p.id));
+
          this.intialLoad = false;
          this.updated();
 
