Horde: Add PF change description to change buttons on user pinned jobs and in stream jobs views
#rnx
#skipci

diff --git a/Engine/Source/Programs/Horde/HordeDashboard/src/backend/JobHandler.ts b/Engine/Source/Programs/Horde/HordeDashboard/src/backend/JobHandler.ts
--- a/Engine/Source/Programs/Horde/HordeDashboard/src/backend/JobHandler.ts
+++ b/Engine/Source/Programs/Horde/HordeDashboard/src/backend/JobHandler.ts
@@ -78,7 +78,7 @@ export class JobHandler {
         try {
 
             // discover new/updated jobs
-            let filter = "id,streamId,name,change,preflightChange,templateId,templateHash,graphHash,startedByUserInfo,abortedByUserInfo,createTime,state,arguments,updateTime,labels,defaultLabel";
+            let filter = "id,streamId,name,change,preflightChange,templateId,templateHash,graphHash,startedByUserInfo,abortedByUserInfo,createTime,state,arguments,updateTime,labels,defaultLabel,preflightDescription";
 
             if (this.includeBatches) {
                 filter += ",batches";
diff --git a/Engine/Source/Programs/Horde/HordeDashboard/src/components/ChangeButton.tsx b/Engine/Source/Programs/Horde/HordeDashboard/src/components/ChangeButton.tsx
--- a/Engine/Source/Programs/Horde/HordeDashboard/src/components/ChangeButton.tsx
+++ b/Engine/Source/Programs/Horde/HordeDashboard/src/components/ChangeButton.tsx
@@ -25,7 +25,7 @@ export const ChangeContextMenu: React.FC<{ target: ChangeContextMenuTarget, onDi
    if (jobChange === undefined) {
       jobChange = stepRef?.change ?? job.change!;
    }
-      
+
    let change = jobChange.toString();
    if (!stepRef && job?.preflightChange) {
       change = `PF ${job.preflightChange}`;
@@ -70,6 +70,23 @@ export const ChangeContextMenu: React.FC<{ target: ChangeContextMenuTarget, onDi
       });
 
 
+   } else if (job.preflightDescription) {
+      menuItems.push({
+         key: "commit", onRender: () => {
+            return <Stack style={{ padding: 18, paddingTop: 18, paddingBottom: 18, maxWidth: 500 }} tokens={{ childrenGap: 12 }}>
+               <Stack horizontal tokens={{ childrenGap: 12 }}>
+                  <Text variant="small" style={{ fontFamily: "Horde Open Sans Bold" }}>Description:</Text>
+                  <Text variant="small" >{job.preflightDescription}</Text>
+               </Stack>
+            </Stack>
+         }
+      });
+
+      menuItems.push({
+         key: "commit_divider",
+         itemType: ContextualMenuItemType.Divider
+      });
+
    }
 
    menuItems.push({
@@ -152,7 +169,7 @@ export const ChangeButton: React.FC<{ job?: JobData, stepRef?: GetJobStepRefResp
          <span ref={spanRef} style={{ padding: "2px 6px 2px 6px", height: "15px", cursor: "pointer" }} className={job.startedByUserInfo ? "cl-callout-button-user" : "cl-callout-button"} onClick={(ev) => { ev.stopPropagation(); ev.preventDefault(); setMenuShown(!menuShown) }} >{change}</span>
          {!hideAborted && <span ref={spanRef} style={{ padding: "2px 6px 2px 6px", height: "16px", cursor: "pointer", userSelect: "none", fontFamily: "Horde Open Sans SemiBold", fontSize: "10px", backgroundColor: colors.get(StatusColor.Failure), color: "rgb(255, 255, 255)" }} onClick={(ev) => { ev.preventDefault(); setMenuShown(!menuShown) }}>Aborted</span>}
       </Stack>
-      {menuShown && <ChangeContextMenu target={{ ref: spanRef }} job={job} commit={commit} stepRef={stepRef} rangeCL={rangeCL}  onDismiss={() => setMenuShown(false)} />}
+      {menuShown && <ChangeContextMenu target={{ ref: spanRef }} job={job} commit={commit} stepRef={stepRef} rangeCL={rangeCL} onDismiss={() => setMenuShown(false)} />}
    </div></Stack>);
 };
 
diff --git a/Engine/Source/Programs/Horde/HordeDashboard/src/components/UserHome.tsx b/Engine/Source/Programs/Horde/HordeDashboard/src/components/UserHome.tsx
--- a/Engine/Source/Programs/Horde/HordeDashboard/src/components/UserHome.tsx
+++ b/Engine/Source/Programs/Horde/HordeDashboard/src/components/UserHome.tsx
@@ -405,7 +405,7 @@ class UserJobsHandler {
       const maxJobs = 100;
 
       try {
-         let filter = "id,streamId,name,change,preflightChange,templateId,templateHash,graphHash,startedByUserInfo,createTime,state,arguments,updateTime,labels,defaultLabel,batches,autoSubmit,autoSubmitChange";
+         let filter = "id,streamId,name,change,preflightChange,templateId,templateHash,graphHash,startedByUserInfo,createTime,state,arguments,updateTime,labels,defaultLabel,batches,autoSubmit,autoSubmitChange,preflightDescription";
 
          const query: JobQuery = {
             filter: filter,
