Horde: Render plain text for broken structured logging tags
#skipci
#rnx

diff --git a/Engine/Source/Programs/Horde/HordeDashboard/src/components/LogRender.tsx b/Engine/Source/Programs/Horde/HordeDashboard/src/components/LogRender.tsx
--- a/Engine/Source/Programs/Horde/HordeDashboard/src/components/LogRender.tsx
+++ b/Engine/Source/Programs/Horde/HordeDashboard/src/components/LogRender.tsx
@@ -74,7 +74,7 @@ const renderTags = (line: LogLine, lineNumber: number | undefined, logStyle: any
       let formatter = "";
       tag = tag.replace("{", "").replace("}", "");
       if (tag.indexOf(":") !== -1) {
-         [tag, formatter] = tag.split(":");                  
+         [tag, formatter] = tag.split(":");
       }
 
       property = properties[tag];
@@ -104,7 +104,7 @@ const renderTags = (line: LogLine, lineNumber: number | undefined, logStyle: any
                property = (property as number).toFixed(precision);
             }
 
-         }         
+         }
       }
 
       let type = "";
@@ -113,7 +113,10 @@ const renderTags = (line: LogLine, lineNumber: number | undefined, logStyle: any
       if (property && typeof (property) !== "string" && typeof (property) !== "number" && typeof (property) !== "boolean" && property !== null) {
          record = property as Record<string, string>;
       } else {
-         text = (property === null || property === undefined) ? "null" : property.toString();
+         if (property === null || property === undefined) {
+            return null;
+         }
+         text = property.toString();
       }
 
       if (record) {
@@ -163,7 +166,7 @@ const renderTags = (line: LogLine, lineNumber: number | undefined, logStyle: any
       }
 
       return <span key={key} />;
-   })
+   }).filter(t => !!t);
 
    let remaining = line.format;
 
