Horde: Include input and outputs in job responses.

diff --git a/Engine/Source/Programs/Horde/Horde.Build/Jobs/Graphs/GraphMessages.cs b/Engine/Source/Programs/Horde/Horde.Build/Jobs/Graphs/GraphMessages.cs
--- a/Engine/Source/Programs/Horde/Horde.Build/Jobs/Graphs/GraphMessages.cs
+++ b/Engine/Source/Programs/Horde/Horde.Build/Jobs/Graphs/GraphMessages.cs
@@ -2,6 +2,7 @@
 
 using System.Collections.Generic;
 using System.Linq;
+using Amazon.S3.Model;
 using EpicGames.Core;
 using HordeCommon;
 
@@ -17,6 +18,16 @@ namespace Horde.Build.Jobs.Graphs
 		/// </summary>
 		public string Name { get; set; }
 
+		/// <summary>
+		/// Inputs for this node
+		/// </summary>
+		public List<string> Inputs { get; set; } = new List<string>();
+
+		/// <summary>
+		/// Output names from this node
+		/// </summary>
+		public List<string> OutputNames { get; set; }
+
 		/// <summary>
 		/// Indices of nodes which must have succeeded for this node to run
 		/// </summary>
@@ -75,6 +86,14 @@ namespace Horde.Build.Jobs.Graphs
 		public GetNodeResponse(INode node, IReadOnlyList<INodeGroup> groups)
 		{
 			Name = node.Name;
+
+			foreach (NodeOutputRef input in node.Inputs)
+			{
+				INode inputNode = groups[input.NodeRef.GroupIdx].Nodes[input.NodeRef.NodeIdx];
+				Inputs.Add($"{inputNode.Name}/{inputNode.OutputNames[input.OutputIdx]}");
+			}
+
+			OutputNames = new List<string>(node.OutputNames);
 			InputDependencies = new List<string>(node.InputDependencies.Select(x => groups[x.GroupIdx].Nodes[x.NodeIdx].Name));
 			OrderDependencies = new List<string>(node.OrderDependencies.Select(x => groups[x.GroupIdx].Nodes[x.NodeIdx].Name));
 			Priority = node.Priority;
