Horde: Catch and log exceptions from issue report sending

diff --git a/Engine/Source/Programs/Horde/Horde.Build/Notifications/NotificationService.cs b/Engine/Source/Programs/Horde/Horde.Build/Notifications/NotificationService.cs
--- a/Engine/Source/Programs/Horde/Horde.Build/Notifications/NotificationService.cs
+++ b/Engine/Source/Programs/Horde/Horde.Build/Notifications/NotificationService.cs
@@ -742,7 +742,15 @@ namespace Horde.Build.Notifications
 		{
 			foreach (INotificationSink sink in _sinks)
 			{
-				await sink.SendIssueReportAsync(report);
+				try
+				{
+					await sink.SendIssueReportAsync(report);
+				}
+				catch (Exception e)
+				{
+					string streamsWithChannel = String.Join(", ", report.Reports.Select(x => x.Stream.Name + " " + x.TriageChannel));
+					_logger.LogError(e, "Failed sending issue report to {Channel} for streams/channels {StreamsWithChannels})", report.Channel, streamsWithChannel);
+				}
 			}
 		}
 	}
