Horde Storage - Cap the number of objects considered before sorting on write time, workaround for out of memory exceptions when sorting the enumerable for very large (50 million+) number of files.

diff --git a/Engine/Source/Programs/Horde/HordeStorage/Horde.Storage/Implementation/Blob/FileSystemStore.cs b/Engine/Source/Programs/Horde/HordeStorage/Horde.Storage/Implementation/Blob/FileSystemStore.cs
--- a/Engine/Source/Programs/Horde/HordeStorage/Horde.Storage/Implementation/Blob/FileSystemStore.cs
+++ b/Engine/Source/Programs/Horde/HordeStorage/Horde.Storage/Implementation/Blob/FileSystemStore.cs
@@ -280,9 +280,11 @@ namespace Horde.Storage.Implementation
         /// </summary>
         /// <param name="ns">Namespace, if set to null all namespaces will be scanned</param>
         /// <param name="maxResults">Max results to return. Note that the entire namespace will be scanned no matter what.</param>
+        /// <param name="maxCountOfObjectsScanned">Max count of objects scanned before we stop.</param>
         /// <returns>Enumerable of least recently accessed objects as FileInfos</returns>
-        public IEnumerable<FileInfo> GetLeastRecentlyAccessedObjects(NamespaceId? ns = null, int maxResults = 10000)
+        public IEnumerable<FileInfo> GetLeastRecentlyAccessedObjects(NamespaceId? ns = null, int maxResults = 10_000, int maxCountOfObjectsScanned = 40_000_000)
         {
+            // TODO: The maxCountOfObjectsScanned is not a ideal solution, we should likely find a solution were we do not have to read all objects into memory like this but rather can scan over them to determine a reasonable last write time cutoff
             string path = ns != null ? Path.Combine(GetRootDir(), ns.ToString()!) : GetRootDir();
             DirectoryInfo di = new DirectoryInfo(path);
             if (!di.Exists)
@@ -290,7 +292,7 @@ namespace Horde.Storage.Implementation
                 return Array.Empty<FileInfo>();
             }
 
-            return di.EnumerateFiles("*", SearchOption.AllDirectories).OrderBy(x => x.LastWriteTime).Take(maxResults);
+            return di.EnumerateFiles("*", SearchOption.AllDirectories).Take(maxCountOfObjectsScanned).OrderBy(x => x.LastWriteTime).Take(maxResults);
         }
         
         /// <summary>
