Horde Storage - Added missing using when building cb package during ref get

diff --git a/Engine/Source/Programs/Horde/HordeStorage/Horde.Storage/Controllers/ReferencesController.cs b/Engine/Source/Programs/Horde/HordeStorage/Horde.Storage/Controllers/ReferencesController.cs
--- a/Engine/Source/Programs/Horde/HordeStorage/Horde.Storage/Controllers/ReferencesController.cs
+++ b/Engine/Source/Programs/Horde/HordeStorage/Horde.Storage/Controllers/ReferencesController.cs
@@ -235,7 +235,7 @@ namespace Horde.Storage.Controllers
 
                         IAsyncEnumerable<Attachment> attachments = _referenceResolver.GetAttachments(ns, cb);
 
-                        CbPackageBuilder writer = new CbPackageBuilder();
+                        using CbPackageBuilder writer = new CbPackageBuilder();
                         await writer.AddAttachment(objectRecord.BlobIdentifier.AsIoHash(), CbPackageAttachmentFlags.IsObject, blobMemory);
 
                         await foreach (Attachment attachment in attachments)
