Horde Storage - Documented how to run a local cook against the local running instance

diff --git a/Engine/Source/Programs/Horde/HordeStorage/README.md b/Engine/Source/Programs/Horde/HordeStorage/README.md
--- a/Engine/Source/Programs/Horde/HordeStorage/README.md
+++ b/Engine/Source/Programs/Horde/HordeStorage/README.md
@@ -232,6 +232,14 @@ Note that this port is primarily used for the speculative blob replication (see
 
 # Common operations
 
+## Running a local cook against a local instance
+If running a local instance (see `Running locally`) you can run your local cook against it by passing the option `-UE-CloudDataCacheHost=http://localhost` This assumes that your project has been setup to use Cloud DDC already and that it uses `UE-CloudDataCacheHost=None` as its host override (this can vary a bit between projects).
+
+If its working as intended you should see output like this in your cooker:
+```
+DerivedDataCache http://localhost: HTTP DDC: Healthy
+```
+
 ## Add new region
 
 The new region will need to contain:
