Horde Storage - Added build to publish releases on github , also refactored our dockerfile to make the datadog parts optional (publishing a image without that to github).

diff --git a/Engine/Source/Programs/Horde/BuildHorde.xml b/Engine/Source/Programs/Horde/BuildHorde.xml
--- a/Engine/Source/Programs/Horde/BuildHorde.xml
+++ b/Engine/Source/Programs/Horde/BuildHorde.xml
@@ -55,7 +55,7 @@
 			<Tag Files="$(StagePaths)" Except=".../.vs/...;.../.git/...;.../bin/...;.../obj/..." With="#InputFiles"/>
 
 			<!-- Build the Docker image -->
-			<Docker-Build BaseDir="Engine/Source" Files="#InputFiles" Tag="horde_storage:$(InformationalVersion)" DockerFile="Engine/Source/Programs/Horde/HordeStorage/dockerfile" Environment="UE_P4_CHANGELIST=(Change)"/>
+			<Docker-Build BaseDir="Engine/Source" Files="#InputFiles" Tag="horde_storage:$(InformationalVersion)" DockerFile="Engine/Source/Programs/Horde/HordeStorage/dockerfile" Environment="UE_P4_CHANGELIST=(Change)" UseBuildKit="true" Arguments="--target horde-storage-datadog"/>
 		</Node>
 	</Agent>
 </BuildGraph>
diff --git a/Engine/Source/Programs/Horde/HordeStorage/dockerfile b/Engine/Source/Programs/Horde/HordeStorage/dockerfile
--- a/Engine/Source/Programs/Horde/HordeStorage/dockerfile
+++ b/Engine/Source/Programs/Horde/HordeStorage/dockerfile
@@ -1,3 +1,5 @@
+# syntax=docker/dockerfile:1.2
+
 FROM mcr.microsoft.com/dotnet/sdk:6.0 AS build
 
 WORKDIR /app
@@ -6,6 +8,7 @@ ARG UE_P4_CHANGELIST=0
 
 # Copy csprojs for nuget restore
 COPY Programs/Shared/EpicGames.Core/*.csproj ./Programs/Shared/EpicGames.Core/
+COPY Programs/Shared/EpicGames.AspNet/*.csproj ./Programs/Shared/EpicGames.AspNet/
 COPY Programs/Shared/EpicGames.Horde/*.csproj ./Programs/Shared/EpicGames.Horde/
 COPY Programs/Shared/EpicGames.IoHash/*.csproj ./Programs/Shared/EpicGames.IoHash/
 COPY Programs/Shared/EpicGames.Serialization/*.csproj ./Programs/Shared/EpicGames.Serialization/
@@ -21,25 +24,40 @@ COPY Programs/Horde/HordeStorage/ Programs/Horde/HordeStorage/
 # build and create the dll
 RUN dotnet publish -c Release -o out Programs/Horde/HordeStorage/Horde.Storage/Horde.Storage.csproj
 
-# create the actual deployable image
-FROM mcr.microsoft.com/dotnet/aspnet:6.0 AS runtime
-# add curl for the other commands below, plus it is useful to have
-RUN apt-get update && apt-get install curl -y
+# setup the base version of horde storage
+FROM mcr.microsoft.com/dotnet/aspnet:6.0 AS horde-storage-base
+WORKDIR /app
+COPY --from=build /app/out ./
+
+# create the actual deployable slim image
+FROM horde-storage-base AS horde-storage-slim
+RUN adduser hordestorage
+USER hordestorage
+ENTRYPOINT dotnet Horde.Storage.dll
 
+# add some useful debug tooling
+FROM horde-storage-base AS horde-storage
+RUN apt-get update &&\ 
+    apt-get install curl unzip procps -y &&\
+    rm -rf /var/lib/apt/lists/* &&\
+    curl -sSL https://aka.ms/getvsdbgsh | /bin/sh /dev/stdin -v latest -l /vsdbg
+RUN adduser hordestorage
+USER hordestorage
+ENTRYPOINT dotnet Horde.Storage.dll
+
+FROM horde-storage-base AS horde-storage-datadog
+# Add the debugger from the normal horde storage build
+RUN apt-get update &&\ 
+    apt-get install curl unzip procps -y &&\
+    rm -rf /var/lib/apt/lists/* &&\
+    curl -sSL https://aka.ms/getvsdbgsh | /bin/sh /dev/stdin -v latest -l /vsdbg
 # Add the datadog tracer
 RUN curl -LO [REDACTED_PROFILE_URL]/dd-trace-dotnet/releases/download/v2.11.0/datadog-dotnet-apm_2.11.0_amd64.deb
 RUN dpkg -i ./datadog-dotnet-apm_2.11.0_amd64.deb
 RUN /opt/datadog/createLogPath.sh
 
-# install remote debugger
-RUN apt-get install unzip && \
-    curl -sSL https://aka.ms/getvsdbgsh | /bin/sh /dev/stdin -v latest -l /vsdbg
-# install ps for remote debugging with vscode
-RUN apt-get -y install procps
-
-WORKDIR /app
-
-COPY --from=build /app/out ./
+RUN adduser hordestorage
+USER hordestorage
 
 # datadog APM requirments
 ENV CORECLR_ENABLE_PROFILING=1
@@ -54,6 +72,4 @@ ENV DD_LOGS_INJECTION=true
 ENV DD_RUNTIME_METRICS_ENABLED=true
 # do not consider http 400-499 statuses errors
 ENV DD_HTTP_CLIENT_ERROR_STATUSES=500-599
-
-
 ENTRYPOINT dotnet Horde.Storage.dll
\ No newline at end of file
diff --git a/Engine/Source/Programs/Horde/HordeStorage/verifyChangelog.sh b/Engine/Source/Programs/Horde/HordeStorage/verifyChangelog.sh
--- a/Engine/Source/Programs/Horde/HordeStorage/verifyChangelog.sh
+++ b/Engine/Source/Programs/Horde/HordeStorage/verifyChangelog.sh
@@ -1,4 +1,30 @@
+#!/bin/bash
 dir=$(dirname "${BASH_SOURCE[0]}")
-current_version=$(yq e '.version' $dir/version.yaml)
-cat $dir/changelog.md | grep -q $current_version && ( echo 'Version found' && exit 0 ) || (echo "No changelog section found for release $current_version" >&2 && exit 1)
-       
\ No newline at end of file
+echo "Looking for files in $dir"
+
+current_version=$(yq e '.version' "$dir/version.yaml")
+echo "Current Version on disk is: '$current_version'"
+
+if grep -qF "$current_version" "$dir/changelog.md";  then
+    echo "Version found '$current_version'"
+else
+    echo "No changelog section found for release $current_version" >&2
+    exit 1
+fi
+
+release_version=$1
+if [[ -n "$release_version" ]]; then    
+    
+    if [[ "$current_version" != "$release_version" ]]; then    
+        echo "Mismatching version numbers found. Version specified: $release_version while version on disk is $current_version" >&2
+        exit 2
+    fi
+
+    if grep -qF "$release_version" "$dir/changelog.md"; then
+        echo "Release Version found '$release_version'"
+    else
+        echo "No changelog section found for release $release_version" >&2
+        exit 1
+    fi
+fi
+
