Horde: Expose settings for lease utilization strategy

diff --git a/Engine/Source/Programs/Horde/Horde.Build.Tests/Fleet/AutoscaleServiceV2Test.cs b/Engine/Source/Programs/Horde/Horde.Build.Tests/Fleet/AutoscaleServiceV2Test.cs
--- a/Engine/Source/Programs/Horde/Horde.Build.Tests/Fleet/AutoscaleServiceV2Test.cs
+++ b/Engine/Source/Programs/Horde/Horde.Build.Tests/Fleet/AutoscaleServiceV2Test.cs
@@ -193,8 +193,12 @@ namespace Horde.Build.Tests.Fleet
 		[TestMethod]
 		public async Task CreateLeaseUtilizationStrategy()
 		{
-			IPoolSizeStrategy s = await CreateStrategy(new PoolSizeStrategyInfo(PoolSizeStrategy.LeaseUtilization, null, "{}"));
+			string config = "{\"SampleTimeSec\": 10, \"NumSamples\": 20, \"NumSamplesForResult\": 30}";
+			IPoolSizeStrategy s = await CreateStrategy(new PoolSizeStrategyInfo(PoolSizeStrategy.LeaseUtilization, null, config));
 			Assert.AreEqual(typeof(LeaseUtilizationStrategy), s.GetType());
+			Assert.AreEqual(10, ((LeaseUtilizationStrategy)s).Settings.SampleTimeSec);
+			Assert.AreEqual(20, ((LeaseUtilizationStrategy)s).Settings.NumSamples);
+			Assert.AreEqual(30, ((LeaseUtilizationStrategy)s).Settings.NumSamplesForResult);
 		}
 		
 		[TestMethod]
diff --git a/Engine/Source/Programs/Horde/Horde.Build.Tests/Fleet/LeaseUtilizationStrategyTest.cs b/Engine/Source/Programs/Horde/Horde.Build.Tests/Fleet/LeaseUtilizationStrategyTest.cs
--- a/Engine/Source/Programs/Horde/Horde.Build.Tests/Fleet/LeaseUtilizationStrategyTest.cs
+++ b/Engine/Source/Programs/Horde/Horde.Build.Tests/Fleet/LeaseUtilizationStrategyTest.cs
@@ -97,7 +97,7 @@ namespace Horde.Build.Tests.Fleet
 
 		private async Task AssertPoolSizeAsync(IPool pool, int expectedNumAgents)
 		{
-			LeaseUtilizationStrategy strategy = new (AgentCollection, PoolCollection, LeaseCollection, Clock, Cache);
+			LeaseUtilizationStrategy strategy = new (AgentCollection, PoolCollection, LeaseCollection, Clock, Cache, new());
 			List<PoolSizeData> output = await strategy.CalcDesiredPoolSizesAsync(new() { new(pool, _poolAgents, null) });
 			Assert.AreEqual(1, output.Count);
 			Assert.AreEqual(expectedNumAgents, output[0].DesiredAgentCount);
diff --git a/Engine/Source/Programs/Horde/Horde.Build/Agents/Fleet/AutoscaleServiceV2.cs b/Engine/Source/Programs/Horde/Horde.Build/Agents/Fleet/AutoscaleServiceV2.cs
--- a/Engine/Source/Programs/Horde/Horde.Build/Agents/Fleet/AutoscaleServiceV2.cs
+++ b/Engine/Source/Programs/Horde/Horde.Build/Agents/Fleet/AutoscaleServiceV2.cs
@@ -349,13 +349,14 @@ namespace Horde.Build.Agents.Fleet
 						switch (info.Type)
 						{
 							case PoolSizeStrategy.JobQueue:
-								JobQueueSettings? settings = JsonSerializer.Deserialize<JobQueueSettings>(info.Config);
-								if (settings == null) throw new ArgumentException("Unable to deserialize pool sizing config");
-								return new JobQueueStrategy(_jobCollection, _graphCollection, _streamService, _clock, _cache, settings);
+								JobQueueSettings? jqSettings = JsonSerializer.Deserialize<JobQueueSettings>(info.Config);
+								if (jqSettings == null) throw new ArgumentException("Unable to deserialize pool sizing config: " + info.Config);
+								return new JobQueueStrategy(_jobCollection, _graphCollection, _streamService, _clock, _cache, jqSettings);
 							
 							case PoolSizeStrategy.LeaseUtilization:
-								// No settings for lease utilization strategy
-								return new LeaseUtilizationStrategy(_agentCollection, _poolCollection, _leaseCollection, _clock, _cache);
+								LeaseUtilizationSettings? luSettings = JsonSerializer.Deserialize<LeaseUtilizationSettings>(info.Config);
+								if (luSettings == null) throw new ArgumentException("Unable to deserialize pool sizing config: " + info.Config);
+								return new LeaseUtilizationStrategy(_agentCollection, _poolCollection, _leaseCollection, _clock, _cache, luSettings);
 							
 							case PoolSizeStrategy.NoOp:
 								return new NoOpPoolSizeStrategy();
@@ -373,7 +374,7 @@ namespace Horde.Build.Agents.Fleet
 				case PoolSizeStrategy.JobQueue:
 					return new JobQueueStrategy(_jobCollection, _graphCollection, _streamService, _clock, _cache, pool.JobQueueSettings);
 				case PoolSizeStrategy.LeaseUtilization:
-					return new LeaseUtilizationStrategy(_agentCollection, _poolCollection, _leaseCollection, _clock, _cache);
+					return new LeaseUtilizationStrategy(_agentCollection, _poolCollection, _leaseCollection, _clock, _cache, new());
 				case PoolSizeStrategy.NoOp:
 					return new NoOpPoolSizeStrategy();
 				default:
diff --git a/Engine/Source/Programs/Horde/Horde.Build/Agents/Fleet/LeaseUtilizationStrategy.cs b/Engine/Source/Programs/Horde/Horde.Build/Agents/Fleet/LeaseUtilizationStrategy.cs
--- a/Engine/Source/Programs/Horde/Horde.Build/Agents/Fleet/LeaseUtilizationStrategy.cs
+++ b/Engine/Source/Programs/Horde/Horde.Build/Agents/Fleet/LeaseUtilizationStrategy.cs
@@ -24,6 +24,20 @@ namespace Horde.Build.Agents.Fleet
 	/// </summary>
 	public class LeaseUtilizationSettings
 	{
+		/// <summary>
+		/// Time period for each sample
+		/// </summary>
+		public int SampleTimeSec { get; set;  } = 6 * 60;
+		
+		/// <summary>
+		/// Number of samples to collect for calculating lease utilization
+		/// </summary>
+		public int NumSamples { get; set;  } = 10;
+		
+		/// <summary>
+		/// Min number of samples for a valid result
+		/// </summary>
+		public int NumSamplesForResult { get; set;  } = 9;
 	}
 	
 	/// <summary>
@@ -89,14 +103,13 @@ namespace Horde.Build.Agents.Fleet
 			}
 		}
 		
+		internal LeaseUtilizationSettings Settings { get; }
+		
 		private readonly IAgentCollection _agentCollection;
 		private readonly IPoolCollection _poolCollection;
 		private readonly ILeaseCollection _leaseCollection;
 		private readonly IClock _clock;
 		private readonly IMemoryCache _cache;
-		private readonly TimeSpan _sampleTime = TimeSpan.FromMinutes(6.0);
-		private readonly int _numSamples = 10;
-		private readonly int _numSamplesForResult = 9;
 
 		/// <summary>
 		/// Constructor
@@ -106,19 +119,15 @@ namespace Horde.Build.Agents.Fleet
 		/// <param name="leaseCollection"></param>
 		/// <param name="clock"></param>
 		/// <param name="cache"></param>
-		/// <param name="numSamples">Number of samples to collect for calculating lease utilization</param>
-		/// <param name="numSamplesForResult">Min number of samples for a valid result</param>
-		/// <param name="sampleTime">Time period for each sample</param>
-		public LeaseUtilizationStrategy(IAgentCollection agentCollection, IPoolCollection poolCollection, ILeaseCollection leaseCollection, IClock clock, IMemoryCache cache, int numSamples = 10, int numSamplesForResult = 9, TimeSpan? sampleTime = null)
+		/// <param name="settings"></param>
+		public LeaseUtilizationStrategy(IAgentCollection agentCollection, IPoolCollection poolCollection, ILeaseCollection leaseCollection, IClock clock, IMemoryCache cache, LeaseUtilizationSettings settings)
 		{
 			_agentCollection = agentCollection;
 			_poolCollection = poolCollection;
 			_leaseCollection = leaseCollection;
 			_clock = clock;
 			_cache = cache;
-			_numSamples = numSamples;
-			_numSamplesForResult = numSamplesForResult;
-			_sampleTime = sampleTime ?? _sampleTime;
+			Settings = settings;
 		}
 
 		private async Task<Dictionary<AgentId, AgentData>> GetAgentDataAsync()
@@ -130,11 +139,11 @@ namespace Horde.Build.Agents.Fleet
 
 			// Query leases in last interval
 			DateTime maxTime = _clock.UtcNow;
-			DateTime minTime = maxTime - (_sampleTime * _numSamples);
+			DateTime minTime = maxTime - TimeSpan.FromSeconds(Settings.SampleTimeSec) * Settings.NumSamples;
 			List<ILease> leases = await _leaseCollection.FindLeasesAsync(minTime, maxTime);
 
 			// Add all the leases to a data object for each agent
-			Dictionary<AgentId, AgentData> agentIdToData = agents.ToDictionary(x => x.Id, x => new AgentData(x, _numSamples));
+			Dictionary<AgentId, AgentData> agentIdToData = agents.ToDictionary(x => x.Id, x => new AgentData(x, Settings.NumSamples));
 			foreach (ILease lease in leases)
 			{
 				AgentData? agentData;
@@ -149,10 +158,10 @@ namespace Horde.Build.Agents.Fleet
 			{
 				foreach (ILease lease in agentData.Leases.OrderBy(x => x.StartTime))
 				{
-					double minT = (lease.StartTime - minTime).TotalSeconds / _sampleTime.TotalSeconds;
+					double minT = (lease.StartTime - minTime).TotalSeconds / Settings.SampleTimeSec;
 					double maxT = (lease.FinishTime == null)
-						? _numSamples
-						: ((lease.FinishTime.Value - minTime).TotalSeconds / _sampleTime.TotalSeconds);
+						? Settings.NumSamples
+						: ((lease.FinishTime.Value - minTime).TotalSeconds / Settings.SampleTimeSec);
 
 					Any payload = Any.Parser.ParseFrom(lease.Payload.ToArray());
 					if (payload.Is(ExecuteJobTask.Descriptor))
@@ -177,7 +186,7 @@ namespace Horde.Build.Agents.Fleet
 			
 			// Get all the pools
 			List<IPool> pools = await _poolCollection.GetAsync();
-			Dictionary<PoolId, PoolData> poolToData = pools.ToDictionary(x => x.Id, x => new PoolData(x, _numSamples));
+			Dictionary<PoolId, PoolData> poolToData = pools.ToDictionary(x => x.Id, x => new PoolData(x, Settings.NumSamples));
 
 			// Find pool utilization over the query period
 			foreach (AgentData agentData in agentIdToData.Values)
@@ -222,7 +231,7 @@ namespace Horde.Build.Agents.Fleet
 				{
 					int minAgents = pool.MinAgents ?? 1;
 					int numReserveAgents = pool.NumReserveAgents ?? 5;
-					double utilization = poolData.Samples.Select(x => x._jobWork).OrderByDescending(x => x).Skip(_numSamples - _numSamplesForResult).First();
+					double utilization = poolData.Samples.Select(x => x._jobWork).OrderByDescending(x => x).Skip(Settings.NumSamples - Settings.NumSamplesForResult).First();
 				
 					// Number of agents in use over the sampling period. Can never be greater than number of agents available in pool.
 					int numAgentsUtilized = (int)utilization;
@@ -236,7 +245,7 @@ namespace Horde.Build.Agents.Fleet
 					sb.AppendFormat("Min=[{0,5:0.0}] ", poolData.Samples.Min(x => x._jobWork));
 					sb.AppendFormat("Max=[{0,5:0.0}] ", poolData.Samples.Max(x => x._jobWork));
 					sb.AppendFormat("Avg=[{0,5:0.0}] ", utilization);
-					sb.AppendFormat("Pct=[{0,5:0.0}] ", poolData.Samples.Sum(x => x._jobWork) / _numSamples);
+					sb.AppendFormat("Pct=[{0,5:0.0}] ", poolData.Samples.Sum(x => x._jobWork) / Settings.NumSamples);
 
 					result.Add(new(pool, poolSize.Agents, desiredAgentCount, sb.ToString()));
 				}
diff --git a/Engine/Source/Programs/Horde/Horde.Build/Startup.cs b/Engine/Source/Programs/Horde/Horde.Build/Startup.cs
--- a/Engine/Source/Programs/Horde/Horde.Build/Startup.cs
+++ b/Engine/Source/Programs/Horde/Horde.Build/Startup.cs
@@ -374,10 +374,6 @@ namespace Horde.Build
 			services.AddSingleton(typeof(ISingletonDocument<>), typeof(SingletonDocument<>));
 
 			services.AddSingleton<AutoscaleServiceV2>();
-			services.AddSingleton<LeaseUtilizationStrategy>();
-			services.AddSingleton<JobQueueStrategy>();
-			services.AddSingleton<NoOpPoolSizeStrategy>();
-			services.AddSingleton<ComputeQueueAwsMetricStrategy>();
 			
 			// Associate IFleetManager interface with the default implementation from config for convenience
 			// Though most fleet managers are created on a per-pool basis
