Horde: Catch and log exceptions when calculating pool sizes (auto-scaling)

Prevents exceptions aborting the entire run of queued pools to be sized.

diff --git a/Engine/Source/Programs/Horde/Horde.Build.Tests/Fleet/AutoscaleServiceV2Test.cs b/Engine/Source/Programs/Horde/Horde.Build.Tests/Fleet/AutoscaleServiceV2Test.cs
--- a/Engine/Source/Programs/Horde/Horde.Build.Tests/Fleet/AutoscaleServiceV2Test.cs
+++ b/Engine/Source/Programs/Horde/Horde.Build.Tests/Fleet/AutoscaleServiceV2Test.cs
@@ -96,6 +96,21 @@ namespace Horde.Build.Tests.Fleet
 			Assert.AreEqual(1, poolSizeDatas.Count);
 			Assert.AreEqual(2, poolSizeDatas[0].Agents.Count);
 		}
+		
+		[TestMethod]
+		public async Task ExceptionsAreSwallowedAndLogged()
+		{
+			using AutoscaleServiceV2 service = GetAutoscaleService(_fleetManagerSpy);
+			IPool pool = await PoolService.CreatePoolAsync("testPool", null, true, 0, 0, sizeStrategy: PoolSizeStrategy.LeaseUtilization);
+			List<PoolSizeData> poolSizeDatas = new()
+			{
+				new PoolSizeData(null!, new(), null),
+				new PoolSizeData(pool, new(), null),
+			};
+
+			poolSizeDatas = await service.CalculatePoolSizesAsync(poolSizeDatas, CancellationToken.None);
+			Assert.AreEqual(1, poolSizeDatas.Count);
+		}
 
 		[TestMethod]
 		public async Task ScaleOutCooldown()
diff --git a/Engine/Source/Programs/Horde/Horde.Build/Agents/Fleet/AutoscaleServiceV2.cs b/Engine/Source/Programs/Horde/Horde.Build/Agents/Fleet/AutoscaleServiceV2.cs
--- a/Engine/Source/Programs/Horde/Horde.Build/Agents/Fleet/AutoscaleServiceV2.cs
+++ b/Engine/Source/Programs/Horde/Horde.Build/Agents/Fleet/AutoscaleServiceV2.cs
@@ -156,9 +156,16 @@ namespace Horde.Build.Agents.Fleet
 
 			await Parallel.ForEachAsync(poolSizeDatas, options, async (input, innerCt) =>
 			{
-				IPoolSizeStrategy sizeStrategy = CreatePoolSizeStrategy(input.Pool);
-				PoolSizeData output = (await sizeStrategy.CalcDesiredPoolSizesAsync(new List<PoolSizeData> { input }))[0];
-				results.Enqueue(output);
+				try
+				{
+					IPoolSizeStrategy sizeStrategy = CreatePoolSizeStrategy(input.Pool);
+					PoolSizeData output = (await sizeStrategy.CalcDesiredPoolSizesAsync(new List<PoolSizeData> { input }))[0];
+					results.Enqueue(output);
+				}
+				catch (Exception e)
+				{
+					_logger.LogError(e, "Failed calculating pool size for pool {PoolId}", input.Pool?.Id);
+				}
 			});
 
 			return results.ToList();
