Horde Storage - Reverting change to run as a user within the container as its causing issues with accessing the cache file system.

diff --git a/Engine/Source/Programs/Horde/HordeStorage/dockerfile b/Engine/Source/Programs/Horde/HordeStorage/dockerfile
--- a/Engine/Source/Programs/Horde/HordeStorage/dockerfile
+++ b/Engine/Source/Programs/Horde/HordeStorage/dockerfile
@@ -31,8 +31,6 @@ COPY --from=build /app/out ./
 
 # create the actual deployable slim image
 FROM horde-storage-base AS horde-storage-slim
-RUN adduser hordestorage
-USER hordestorage
 ENTRYPOINT dotnet Horde.Storage.dll
 
 # add some useful debug tooling
@@ -41,8 +39,6 @@ RUN apt-get update &&\
     apt-get install curl unzip procps -y &&\
     rm -rf /var/lib/apt/lists/* &&\
     curl -sSL https://aka.ms/getvsdbgsh | /bin/sh /dev/stdin -v latest -l /vsdbg
-RUN adduser hordestorage
-USER hordestorage
 ENTRYPOINT dotnet Horde.Storage.dll
 
 FROM horde-storage-base AS horde-storage-datadog
@@ -56,9 +52,6 @@ RUN curl -LO [REDACTED_PROFILE_URL]/dd-trace-dotnet/releases/download/v2.11.
 RUN dpkg -i ./datadog-dotnet-apm_2.11.0_amd64.deb
 RUN /opt/datadog/createLogPath.sh
 
-RUN adduser hordestorage
-USER hordestorage
-
 # datadog APM requirments
 ENV CORECLR_ENABLE_PROFILING=1
 ENV CORECLR_PROFILER={846F5F1C-F9AE-4B07-969E-05C26BC060D8}
