Horde: Add better handling of empty or unset config for pool sizing strategies

diff --git a/Engine/Source/Programs/Horde/Horde.Build.Tests/Fleet/AutoscaleServiceV2Test.cs b/Engine/Source/Programs/Horde/Horde.Build.Tests/Fleet/AutoscaleServiceV2Test.cs
--- a/Engine/Source/Programs/Horde/Horde.Build.Tests/Fleet/AutoscaleServiceV2Test.cs
+++ b/Engine/Source/Programs/Horde/Horde.Build.Tests/Fleet/AutoscaleServiceV2Test.cs
@@ -216,6 +216,16 @@ namespace Horde.Build.Tests.Fleet
 			Assert.AreEqual(30, ((LeaseUtilizationStrategy)s).Settings.NumSamplesForResult);
 		}
 		
+		[TestMethod]
+		public async Task EmptyOrInvalidJsonConfig()
+		{
+			await CreateStrategy(new PoolSizeStrategyInfo(PoolSizeStrategy.LeaseUtilization, null, ""));
+			await CreateStrategy(new PoolSizeStrategyInfo(PoolSizeStrategy.LeaseUtilization, null, "{}"));
+			await CreateStrategy(new PoolSizeStrategyInfo(PoolSizeStrategy.LeaseUtilization, null, "  {} "));
+			
+			await Assert.ThrowsExceptionAsync<JsonException>(() => CreateStrategy(new PoolSizeStrategyInfo(PoolSizeStrategy.LeaseUtilization, null, "BAD_JSON")));
+		}
+		
 		[TestMethod]
 		public async Task CreateNoOpStrategy()
 		{
diff --git a/Engine/Source/Programs/Horde/Horde.Build/Agents/Fleet/AutoscaleServiceV2.cs b/Engine/Source/Programs/Horde/Horde.Build/Agents/Fleet/AutoscaleServiceV2.cs
--- a/Engine/Source/Programs/Horde/Horde.Build/Agents/Fleet/AutoscaleServiceV2.cs
+++ b/Engine/Source/Programs/Horde/Horde.Build/Agents/Fleet/AutoscaleServiceV2.cs
@@ -356,13 +356,11 @@ namespace Horde.Build.Agents.Fleet
 						switch (info.Type)
 						{
 							case PoolSizeStrategy.JobQueue:
-								JobQueueSettings? jqSettings = JsonSerializer.Deserialize<JobQueueSettings>(info.Config);
-								if (jqSettings == null) throw new ArgumentException("Unable to deserialize pool sizing config: " + info.Config);
+								JobQueueSettings jqSettings = DeserializeConfig<JobQueueSettings>(info.Config);
 								return new JobQueueStrategy(_jobCollection, _graphCollection, _streamService, _clock, _cache, jqSettings);
 							
 							case PoolSizeStrategy.LeaseUtilization:
-								LeaseUtilizationSettings? luSettings = JsonSerializer.Deserialize<LeaseUtilizationSettings>(info.Config);
-								if (luSettings == null) throw new ArgumentException("Unable to deserialize pool sizing config: " + info.Config);
+								LeaseUtilizationSettings luSettings = DeserializeConfig<LeaseUtilizationSettings>(info.Config);
 								return new LeaseUtilizationStrategy(_agentCollection, _poolCollection, _leaseCollection, _clock, _cache, luSettings);
 							
 							case PoolSizeStrategy.NoOp:
@@ -388,5 +386,13 @@ namespace Horde.Build.Agents.Fleet
 					throw new ArgumentException("Unknown pool size strategy " + pool.SizeStrategy);
 			}
 		}
+
+		private static T DeserializeConfig<T>(string json)
+		{
+			json = String.IsNullOrEmpty(json) ? "{}" : json;
+			T? config = JsonSerializer.Deserialize<T>(json);
+			if (config == null) throw new ArgumentException("Unable to deserialize config: " + json);
+			return config;
+		}
 	}
 }
