Horde: Additional logging to catch step refs not updating
#rnx

diff --git a/Engine/Source/Programs/Horde/Horde.Build/Jobs/IJobStepRefCollection.cs b/Engine/Source/Programs/Horde/Horde.Build/Jobs/IJobStepRefCollection.cs
--- a/Engine/Source/Programs/Horde/Horde.Build/Jobs/IJobStepRefCollection.cs
+++ b/Engine/Source/Programs/Horde/Horde.Build/Jobs/IJobStepRefCollection.cs
@@ -10,6 +10,7 @@ using Horde.Build.Logs;
 using Horde.Build.Streams;
 using Horde.Build.Utilities;
 using HordeCommon;
+using Microsoft.Extensions.Logging;
 
 namespace Horde.Build.Jobs
 {
@@ -84,7 +85,7 @@ namespace Horde.Build.Jobs
 
 	static class JobStepRefCollectionExtensions
 	{
-		public static async Task UpdateAsync(this IJobStepRefCollection jobStepRefs, IJob job, IJobStepBatch batch, IJobStep step, IGraph graph)
+		public static async Task UpdateAsync(this IJobStepRefCollection jobStepRefs, IJob job, IJobStepBatch batch, IJobStep step, IGraph graph, ILogger? logger = null)
 		{
 			if (job.PreflightChange == 0)
 			{
@@ -118,6 +119,11 @@ namespace Horde.Build.Jobs
 				{
 					outcome = JobStepOutcome.Unspecified;
 				}
+				
+				if (logger != null)
+				{
+					logger.LogInformation("Updating step reference {StepId} for job {JobId}, batch {BatchId}, with outcome {JobStepOutcome}", step.Id, job.Id, batch.Id, outcome);
+				}				
 
 				await jobStepRefs.InsertOrReplaceAsync(new JobStepRefId(job.Id, batch.Id, step.Id), job.Name, nodeName, job.StreamId, job.TemplateId, job.Change, step.LogId, batch.PoolId, batch.AgentId, outcome, job.UpdateIssues, lastSuccess, lastWarning, waitTime, initTime, step.StartTimeUtc ?? DateTime.UtcNow, step.FinishTimeUtc);
 			}
diff --git a/Engine/Source/Programs/Horde/Horde.Build/Jobs/JobCollection.cs b/Engine/Source/Programs/Horde/Horde.Build/Jobs/JobCollection.cs
--- a/Engine/Source/Programs/Horde/Horde.Build/Jobs/JobCollection.cs
+++ b/Engine/Source/Programs/Horde/Horde.Build/Jobs/JobCollection.cs
@@ -1161,7 +1161,7 @@ namespace Horde.Build.Jobs
 		/// <inheritdoc/>
 		public async Task<IJob?> TryCancelLeaseAsync(IJob job, int batchIdx)
 		{
-			_logger.LogDebug("Cancelling lease {LeaseId} for agent {AgentId}", job.Batches[batchIdx].LeaseId, job.Batches[batchIdx].AgentId);
+			_logger.LogInformation("Cancelling lease {LeaseId} for agent {AgentId}", job.Batches[batchIdx].LeaseId, job.Batches[batchIdx].AgentId);
 
 			JobDocument jobDocument = Clone((JobDocument)job);
 
@@ -1184,7 +1184,7 @@ namespace Horde.Build.Jobs
 		{
 			JobDocument jobDocument = Clone((JobDocument)job);
 			JobStepBatchDocument batch = jobDocument.Batches[batchIdx];
-			_logger.LogDebug("Failing batch {JobId}:{BatchId} with error {Error}", job.Id, batch.Id, error);
+			_logger.LogInformation("Failing batch {JobId}:{BatchId} with error {Error}", job.Id, batch.Id, error);
 
 			UpdateDefinitionBuilder<JobDocument> updateBuilder = new UpdateDefinitionBuilder<JobDocument>();
 			List<UpdateDefinition<JobDocument>> updates = new List<UpdateDefinition<JobDocument>>();
diff --git a/Engine/Source/Programs/Horde/Horde.Build/Jobs/JobService.cs b/Engine/Source/Programs/Horde/Horde.Build/Jobs/JobService.cs
--- a/Engine/Source/Programs/Horde/Horde.Build/Jobs/JobService.cs
+++ b/Engine/Source/Programs/Horde/Horde.Build/Jobs/JobService.cs
@@ -1016,7 +1016,7 @@ namespace Horde.Build.Jobs
 				
 				if (oldState != newState || oldOutcome != newOutcome)
 				{
-					_logger.LogDebug("Transitioned job {JobId}, batch {BatchId}, step {StepId} from {OldState} to {NewState}", job.Id, batchId, stepId, oldState, newState);
+					_logger.LogInformation("Transitioned job {JobId}, batch {BatchId}, step {StepId} from {OldState} to {NewState}", job.Id, batchId, stepId, oldState, newState);
 
 					// Send any updates for modified badges
 					if (oldLabelStates != null)
diff --git a/Engine/Source/Programs/Horde/Horde.Build/Jobs/JobTaskSource.cs b/Engine/Source/Programs/Horde/Horde.Build/Jobs/JobTaskSource.cs
--- a/Engine/Source/Programs/Horde/Horde.Build/Jobs/JobTaskSource.cs
+++ b/Engine/Source/Programs/Horde/Horde.Build/Jobs/JobTaskSource.cs
@@ -1036,7 +1036,7 @@ namespace Horde.Build.Jobs
 					{
 						error = JobStepBatchError.ExecutionError;
 					}
-
+					
 					IGraph graph = await _graphs.GetAsync(job.GraphHash);
 					job = await _jobs.TryFailBatchAsync(job, batchIdx, graph, error);
 
@@ -1044,11 +1044,11 @@ namespace Horde.Build.Jobs
 					{
 						if (batch.Error != JobStepBatchError.None)
 						{
-							logger.LogInformation("Failed job {JobId}, batch {BatchId} with error {Error}", job.Id, batch.Id, batch.Error);
+							logger.LogInformation("Failed lease {LeaseId}, job {JobId}, batch {BatchId} with error {Error}", leaseId, job.Id, batch.Id, batch.Error);
 						}
 						if (runningStepIdx != -1)
 						{
-							await _jobStepRefs.UpdateAsync(job, batch, batch.Steps[runningStepIdx], graph);
+							await _jobStepRefs.UpdateAsync(job, batch, batch.Steps[runningStepIdx], graph, logger);
 						}
 						break;
 					}
diff --git a/Engine/Source/Programs/Horde/Horde.Build/Tasks/ITaskSource.cs b/Engine/Source/Programs/Horde/Horde.Build/Tasks/ITaskSource.cs
--- a/Engine/Source/Programs/Horde/Horde.Build/Tasks/ITaskSource.cs
+++ b/Engine/Source/Programs/Horde/Horde.Build/Tasks/ITaskSource.cs
@@ -207,7 +207,7 @@ namespace Horde.Build.Tasks
 		/// <inheritdoc cref="ITaskSource.OnLeaseFinishedAsync(IAgent, LeaseId, Any, LeaseOutcome, ReadOnlyMemory{Byte}, ILogger)"/>
 		public virtual Task OnLeaseFinishedAsync(IAgent agent, LeaseId leaseId, TMessage payload, LeaseOutcome outcome, ReadOnlyMemory<byte> output, ILogger logger)
 		{
-			logger.LogInformation("Lease {LeaseId} complete", leaseId);
+			logger.LogInformation("Lease {LeaseId} complete, outcome {LeaseOutcome}", leaseId, outcome);
 			return Task.CompletedTask;
 		}
 
