Horde: Fix inputs/outputs not being propagated to graph definition.

diff --git a/Engine/Source/Programs/Horde/Horde.Agent/Execution/BuildGraphExecutor.cs b/Engine/Source/Programs/Horde/Horde.Agent/Execution/BuildGraphExecutor.cs
--- a/Engine/Source/Programs/Horde/Horde.Agent/Execution/BuildGraphExecutor.cs
+++ b/Engine/Source/Programs/Horde/Horde.Agent/Execution/BuildGraphExecutor.cs
@@ -34,7 +34,7 @@ namespace Horde.Agent.Execution
 			public bool RunEarly { get; set; }
 			public bool? Warnings { get; set; }
 			public List<string> Inputs { get; set; } = new List<string>();
-			public List<string> Outputs { get; set; } = new List<string>();
+			public List<string> OutputNames { get; set; } = new List<string>();
 			public List<string> InputDependencies { get; set; } = new List<string>();
 			public List<string> OrderDependencies { get; set; } = new List<string>();
 			public Dictionary<string, string> Annotations { get; set; } = new Dictionary<string, string>(StringComparer.OrdinalIgnoreCase);
@@ -451,9 +451,9 @@ namespace Horde.Agent.Execution
 					{
 						createNode.Inputs.Add(exportedNode.Inputs);
 					}
-					if (exportedNode.Outputs != null)
+					if (exportedNode.OutputNames != null)
 					{
-						createNode.Outputs.Add(exportedNode.Outputs);
+						createNode.Outputs.Add(exportedNode.OutputNames);
 					}
 					if (exportedNode.InputDependencies != null)
 					{
diff --git a/Engine/Source/Programs/Horde/Horde.Build/Jobs/Graphs/IGraph.cs b/Engine/Source/Programs/Horde/Horde.Build/Jobs/Graphs/IGraph.cs
--- a/Engine/Source/Programs/Horde/Horde.Build/Jobs/Graphs/IGraph.cs
+++ b/Engine/Source/Programs/Horde/Horde.Build/Jobs/Graphs/IGraph.cs
@@ -479,6 +479,8 @@ namespace Horde.Build.Jobs.Graphs
 		/// Constructor
 		/// </summary>
 		/// <param name="name">Name of the node</param>
+		/// <param name="inputs">List of inputs for the node</param>
+		/// <param name="outputNames">List of output names for the node</param>
 		/// <param name="inputDependencies">List of nodes which must have completed succesfully for this node to run</param>
 		/// <param name="orderDependencies">List of nodes which must have completed for this node to run</param>
 		/// <param name="priority">Priority of this node</param>
@@ -488,9 +490,11 @@ namespace Horde.Build.Jobs.Graphs
 		/// <param name="credentials">Credentials required for this node to run</param>
 		/// <param name="properties">Properties for the node</param>
 		/// <param name="annotations">User annotations for this node</param>
-		public NewNode(string name, List<string>? inputDependencies = null, List<string>? orderDependencies = null, Priority? priority = null, bool? allowRetry = null, bool? runEarly = null, bool? warnings = null, Dictionary<string, string>? credentials = null, Dictionary<string, string>? properties = null, IReadOnlyNodeAnnotations? annotations = null)
+		public NewNode(string name, List<string>? inputs = null, List<string>? outputNames = null, List<string>? inputDependencies = null, List<string>? orderDependencies = null, Priority? priority = null, bool? allowRetry = null, bool? runEarly = null, bool? warnings = null, Dictionary<string, string>? credentials = null, Dictionary<string, string>? properties = null, IReadOnlyNodeAnnotations? annotations = null)
 		{
 			Name = name;
+			Inputs = inputs;
+			OutputNames = outputNames;
 			InputDependencies = inputDependencies;
 			OrderDependencies = orderDependencies;
 			Priority = priority;
diff --git a/Engine/Source/Programs/Horde/Horde.Build/Server/RpcService.cs b/Engine/Source/Programs/Horde/Horde.Build/Server/RpcService.cs
--- a/Engine/Source/Programs/Horde/Horde.Build/Server/RpcService.cs
+++ b/Engine/Source/Programs/Horde/Horde.Build/Server/RpcService.cs
@@ -826,7 +826,7 @@ namespace Horde.Build.Server
 				List<NewNode> newNodes = new List<NewNode>();
 				foreach (CreateNodeRequest node in group.Nodes)
 				{
-					NewNode newNode = new NewNode(node.Name, node.InputDependencies.ToList(), node.OrderDependencies.ToList(), node.Priority, node.AllowRetry, node.RunEarly, node.Warnings, new Dictionary<string, string>(node.Credentials), new Dictionary<string, string>(node.Properties), new NodeAnnotations(node.Annotations));
+					NewNode newNode = new NewNode(node.Name, node.Inputs.ToList(), node.OutputNames.ToList(), node.InputDependencies.ToList(), node.OrderDependencies.ToList(), node.Priority, node.AllowRetry, node.RunEarly, node.Warnings, new Dictionary<string, string>(node.Credentials), new Dictionary<string, string>(node.Properties), new NodeAnnotations(node.Annotations));
 					newNodes.Add(newNode);
 				}
 				newGroups.Add(new NewGroup(group.AgentType, newNodes));
diff --git a/Engine/Source/Programs/Shared/EpicGames.BuildGraph/BgGraphDef.cs b/Engine/Source/Programs/Shared/EpicGames.BuildGraph/BgGraphDef.cs
--- a/Engine/Source/Programs/Shared/EpicGames.BuildGraph/BgGraphDef.cs
+++ b/Engine/Source/Programs/Shared/EpicGames.BuildGraph/BgGraphDef.cs
@@ -447,10 +447,10 @@ namespace EpicGames.BuildGraph
 						}
 						jsonWriter.WriteArrayEnd();
 
-						jsonWriter.WriteArrayStart("Outputs");
+						jsonWriter.WriteArrayStart("OutputNames");
 						foreach (BgNodeOutput output in node.Outputs)
 						{
-							jsonWriter.WriteValue($"{output.ProducingNode.Name}/{output.TagName}");
+							jsonWriter.WriteValue($"{output.TagName}");
 						}
 						jsonWriter.WriteArrayEnd();
 
