Horde: Fix input/output names (again).

diff --git a/Engine/Source/Programs/Horde/Horde.Agent/Execution/BuildGraphExecutor.cs b/Engine/Source/Programs/Horde/Horde.Agent/Execution/BuildGraphExecutor.cs
--- a/Engine/Source/Programs/Horde/Horde.Agent/Execution/BuildGraphExecutor.cs
+++ b/Engine/Source/Programs/Horde/Horde.Agent/Execution/BuildGraphExecutor.cs
@@ -13,7 +13,6 @@ using System.Text.Json.Serialization;
 using System.Threading;
 using System.Threading.Tasks;
 using EpicGames.Core;
-using EpicGames.Perforce;
 using Grpc.Core;
 using Horde.Agent.Execution.Interfaces;
 using Horde.Agent.Parser;
@@ -34,7 +33,7 @@ namespace Horde.Agent.Execution
 			public bool RunEarly { get; set; }
 			public bool? Warnings { get; set; }
 			public List<string> Inputs { get; set; } = new List<string>();
-			public List<string> OutputNames { get; set; } = new List<string>();
+			public List<string> Outputs { get; set; } = new List<string>();
 			public List<string> InputDependencies { get; set; } = new List<string>();
 			public List<string> OrderDependencies { get; set; } = new List<string>();
 			public Dictionary<string, string> Annotations { get; set; } = new Dictionary<string, string>(StringComparer.OrdinalIgnoreCase);
@@ -451,9 +450,9 @@ namespace Horde.Agent.Execution
 					{
 						createNode.Inputs.Add(exportedNode.Inputs);
 					}
-					if (exportedNode.OutputNames != null)
+					if (exportedNode.Outputs != null)
 					{
-						createNode.Outputs.Add(exportedNode.OutputNames);
+						createNode.Outputs.Add(exportedNode.Outputs);
 					}
 					if (exportedNode.InputDependencies != null)
 					{
diff --git a/Engine/Source/Programs/Horde/Horde.Build/Jobs/Graphs/GraphCollection.cs b/Engine/Source/Programs/Horde/Horde.Build/Jobs/Graphs/GraphCollection.cs
--- a/Engine/Source/Programs/Horde/Horde.Build/Jobs/Graphs/GraphCollection.cs
+++ b/Engine/Source/Programs/Horde/Horde.Build/Jobs/Graphs/GraphCollection.cs
@@ -255,17 +255,17 @@ namespace Horde.Build.Jobs.Graphs
 							NodeRef[] inputDependencies = (newNodeRequest.InputDependencies == null) ? Array.Empty<NodeRef>() : newNodeRequest.InputDependencies.Select(x => nodeNameToRef[x]).ToArray();
 							NodeRef[] orderDependencies = (newNodeRequest.OrderDependencies == null) ? Array.Empty<NodeRef>() : newNodeRequest.OrderDependencies.Select(x => nodeNameToRef[x]).ToArray();
 							orderDependencies = orderDependencies.Union(inputDependencies).ToArray();
-							nodes.Add(new Node(newNodeRequest.Name, inputs, newNodeRequest.OutputNames?.ToArray(), inputDependencies, orderDependencies, priority, bAllowRetry, bRunEarly, bWarnings, newNodeRequest.Credentials, newNodeRequest.Properties, newNodeRequest.Annotations));
+							nodes.Add(new Node(newNodeRequest.Name, inputs, newNodeRequest.Outputs?.ToArray(), inputDependencies, orderDependencies, priority, bAllowRetry, bRunEarly, bWarnings, newNodeRequest.Credentials, newNodeRequest.Properties, newNodeRequest.Annotations));
 
 							NodeRef nodeRef = new NodeRef(newGroups.Count, nodeIdx);
 							nodeNameToRef.Add(newNodeRequest.Name, nodeRef);
 
-							if (newNodeRequest.OutputNames != null)
+							if (newNodeRequest.Outputs != null)
 							{
-								for (int outputIdx = 0; outputIdx < newNodeRequest.OutputNames.Count; outputIdx++)
+								for (int outputIdx = 0; outputIdx < newNodeRequest.Outputs.Count; outputIdx++)
 								{
-									string outputName = newNodeRequest.OutputNames[outputIdx];
-									nodeOutputNameToRef[$"{newNodeRequest.Name}/{outputName}"] = new NodeOutputRef(nodeRef, outputIdx);
+									string outputName = newNodeRequest.Outputs[outputIdx];
+									nodeOutputNameToRef[outputName] = new NodeOutputRef(nodeRef, outputIdx);
 								}
 							}
 						}
diff --git a/Engine/Source/Programs/Horde/Horde.Build/Jobs/Graphs/IGraph.cs b/Engine/Source/Programs/Horde/Horde.Build/Jobs/Graphs/IGraph.cs
--- a/Engine/Source/Programs/Horde/Horde.Build/Jobs/Graphs/IGraph.cs
+++ b/Engine/Source/Programs/Horde/Horde.Build/Jobs/Graphs/IGraph.cs
@@ -428,7 +428,7 @@ namespace Horde.Build.Jobs.Graphs
 		/// <summary>
 		/// Output names
 		/// </summary>
-		public List<string>? OutputNames { get; set; }
+		public List<string>? Outputs { get; set; }
 
 		/// <summary>
 		/// List of nodes which must succeed for this node to run
@@ -480,7 +480,7 @@ namespace Horde.Build.Jobs.Graphs
 		/// </summary>
 		/// <param name="name">Name of the node</param>
 		/// <param name="inputs">List of inputs for the node</param>
-		/// <param name="outputNames">List of output names for the node</param>
+		/// <param name="outputs">List of output names for the node</param>
 		/// <param name="inputDependencies">List of nodes which must have completed succesfully for this node to run</param>
 		/// <param name="orderDependencies">List of nodes which must have completed for this node to run</param>
 		/// <param name="priority">Priority of this node</param>
@@ -490,11 +490,11 @@ namespace Horde.Build.Jobs.Graphs
 		/// <param name="credentials">Credentials required for this node to run</param>
 		/// <param name="properties">Properties for the node</param>
 		/// <param name="annotations">User annotations for this node</param>
-		public NewNode(string name, List<string>? inputs = null, List<string>? outputNames = null, List<string>? inputDependencies = null, List<string>? orderDependencies = null, Priority? priority = null, bool? allowRetry = null, bool? runEarly = null, bool? warnings = null, Dictionary<string, string>? credentials = null, Dictionary<string, string>? properties = null, IReadOnlyNodeAnnotations? annotations = null)
+		public NewNode(string name, List<string>? inputs = null, List<string>? outputs = null, List<string>? inputDependencies = null, List<string>? orderDependencies = null, Priority? priority = null, bool? allowRetry = null, bool? runEarly = null, bool? warnings = null, Dictionary<string, string>? credentials = null, Dictionary<string, string>? properties = null, IReadOnlyNodeAnnotations? annotations = null)
 		{
 			Name = name;
 			Inputs = inputs;
-			OutputNames = outputNames;
+			Outputs = outputs;
 			InputDependencies = inputDependencies;
 			OrderDependencies = orderDependencies;
 			Priority = priority;
diff --git a/Engine/Source/Programs/Horde/Horde.Build/Server/RpcService.cs b/Engine/Source/Programs/Horde/Horde.Build/Server/RpcService.cs
--- a/Engine/Source/Programs/Horde/Horde.Build/Server/RpcService.cs
+++ b/Engine/Source/Programs/Horde/Horde.Build/Server/RpcService.cs
@@ -826,7 +826,7 @@ namespace Horde.Build.Server
 				List<NewNode> newNodes = new List<NewNode>();
 				foreach (CreateNodeRequest node in group.Nodes)
 				{
-					NewNode newNode = new NewNode(node.Name, node.Inputs.ToList(), node.OutputNames.ToList(), node.InputDependencies.ToList(), node.OrderDependencies.ToList(), node.Priority, node.AllowRetry, node.RunEarly, node.Warnings, new Dictionary<string, string>(node.Credentials), new Dictionary<string, string>(node.Properties), new NodeAnnotations(node.Annotations));
+					NewNode newNode = new NewNode(node.Name, node.Inputs.ToList(), node.Outputs.ToList(), node.InputDependencies.ToList(), node.OrderDependencies.ToList(), node.Priority, node.AllowRetry, node.RunEarly, node.Warnings, new Dictionary<string, string>(node.Credentials), new Dictionary<string, string>(node.Properties), new NodeAnnotations(node.Annotations));
 					newNodes.Add(newNode);
 				}
 				newGroups.Add(new NewGroup(group.AgentType, newNodes));
diff --git a/Engine/Source/Programs/Shared/EpicGames.BuildGraph/BgGraphDef.cs b/Engine/Source/Programs/Shared/EpicGames.BuildGraph/BgGraphDef.cs
--- a/Engine/Source/Programs/Shared/EpicGames.BuildGraph/BgGraphDef.cs
+++ b/Engine/Source/Programs/Shared/EpicGames.BuildGraph/BgGraphDef.cs
@@ -443,14 +443,14 @@ namespace EpicGames.BuildGraph
 						jsonWriter.WriteArrayStart("Inputs");
 						foreach (BgNodeOutput input in node.Inputs)
 						{
-							jsonWriter.WriteValue($"{input.ProducingNode.Name}/{input.TagName}");
+							jsonWriter.WriteValue(input.TagName);
 						}
 						jsonWriter.WriteArrayEnd();
 
-						jsonWriter.WriteArrayStart("OutputNames");
+						jsonWriter.WriteArrayStart("Outputs");
 						foreach (BgNodeOutput output in node.Outputs)
 						{
-							jsonWriter.WriteValue($"{output.TagName}");
+							jsonWriter.WriteValue(output.TagName);
 						}
 						jsonWriter.WriteArrayEnd();
 
