Horde: Fix tests.

diff --git a/Engine/Source/Programs/Horde/Horde.Build.Tests/JobCollectionTests.cs b/Engine/Source/Programs/Horde/Horde.Build.Tests/JobCollectionTests.cs
--- a/Engine/Source/Programs/Horde/Horde.Build.Tests/JobCollectionTests.cs
+++ b/Engine/Source/Programs/Horde/Horde.Build.Tests/JobCollectionTests.cs
@@ -41,7 +41,7 @@ namespace Horde.Build.Tests
 
 		static NewNode AddNode(NewGroup @group, string name, string[]? inputDependencies, Action<NewNode>? action = null)
 		{
-			NewNode node = new NewNode(name, inputDependencies?.ToList(), inputDependencies?.ToList(), null, null, null, null, null, null);
+			NewNode node = new NewNode(name, inputDependencies: inputDependencies?.ToList(), orderDependencies: inputDependencies?.ToList());
 			if (action != null)
 			{
 				action.Invoke(node);
