Horde: Remove node name prefix from graph responses.

diff --git a/Engine/Source/Programs/Horde/Horde.Build/Jobs/Graphs/GraphMessages.cs b/Engine/Source/Programs/Horde/Horde.Build/Jobs/Graphs/GraphMessages.cs
--- a/Engine/Source/Programs/Horde/Horde.Build/Jobs/Graphs/GraphMessages.cs
+++ b/Engine/Source/Programs/Horde/Horde.Build/Jobs/Graphs/GraphMessages.cs
@@ -24,9 +24,9 @@ namespace Horde.Build.Jobs.Graphs
 		public List<string> Inputs { get; set; } = new List<string>();
 
 		/// <summary>
-		/// Output names from this node
+		/// Output from this node
 		/// </summary>
-		public List<string> OutputNames { get; set; }
+		public List<string> Outputs { get; set; } = new List<string>();
 
 		/// <summary>
 		/// Indices of nodes which must have succeeded for this node to run
@@ -86,14 +86,8 @@ namespace Horde.Build.Jobs.Graphs
 		public GetNodeResponse(INode node, IReadOnlyList<INodeGroup> groups)
 		{
 			Name = node.Name;
-
-			foreach (NodeOutputRef input in node.Inputs)
-			{
-				INode inputNode = groups[input.NodeRef.GroupIdx].Nodes[input.NodeRef.NodeIdx];
-				Inputs.Add($"{inputNode.Name}/{inputNode.OutputNames[input.OutputIdx]}");
-			}
-
-			OutputNames = new List<string>(node.OutputNames);
+			Inputs.AddRange(node.Inputs.Select(x => groups[x.NodeRef.GroupIdx].Nodes[x.NodeRef.NodeIdx].OutputNames[x.OutputIdx]));
+			Outputs.AddRange(node.OutputNames);
 			InputDependencies = new List<string>(node.InputDependencies.Select(x => groups[x.GroupIdx].Nodes[x.NodeIdx].Name));
 			OrderDependencies = new List<string>(node.OrderDependencies.Select(x => groups[x.GroupIdx].Nodes[x.NodeIdx].Name));
 			Priority = node.Priority;
