Horde: Experimental mode to manage shared storage within the Horde agent itself.

diff --git a/Engine/Source/Programs/Horde/Horde.Agent/Execution/BuildGraphExecutor.cs b/Engine/Source/Programs/Horde/Horde.Agent/Execution/BuildGraphExecutor.cs
--- a/Engine/Source/Programs/Horde/Horde.Agent/Execution/BuildGraphExecutor.cs
+++ b/Engine/Source/Programs/Horde/Horde.Agent/Execution/BuildGraphExecutor.cs
@@ -12,11 +12,13 @@ using System.Text.Json;
 using System.Text.Json.Serialization;
 using System.Threading;
 using System.Threading.Tasks;
+using Amazon.EC2.Model;
 using EpicGames.Core;
 using Grpc.Core;
 using Horde.Agent.Execution.Interfaces;
 using Horde.Agent.Parser;
 using Horde.Agent.Utility;
+using Horde.Storage.Utility;
 using HordeCommon;
 using HordeCommon.Rpc;
 using Microsoft.Extensions.Logging;
@@ -603,25 +605,234 @@ namespace Horde.Agent.Execution
 				FetchPreprocessedFile(localPreprocessedSchema, sharedStorageDir, logger);
 				arguments.AppendArgument("-ImportSchema=", localPreprocessedSchema.FullName);
 			}
-			else if(_scriptFileName != null)
+			else if (_scriptFileName != null)
 			{
 				arguments.AppendArgument(ScriptArgumentPrefix, _scriptFileName);
 			}
 			arguments.AppendArgument("-SingleNode=", step.Name);
-			if (sharedStorageDir != null)
-			{
-				arguments.AppendArgument("-SharedStorageDir=", sharedStorageDir.FullName);
-			}
 //			Arguments.AppendArgument("-TokenSignature=", JobId.ToString());
+
+			bool manageSharedStorage = false;
 			foreach (string additionalArgument in _additionalArguments)
 			{
-				if (!_preprocessScript || !additionalArgument.StartsWith("-set:", StringComparison.OrdinalIgnoreCase))
+				if (additionalArgument.Equals("-ManageSharedStorage", StringComparison.OrdinalIgnoreCase))
+				{
+					manageSharedStorage = true;
+				}
+				else if (!_preprocessScript || !additionalArgument.StartsWith("-set:", StringComparison.OrdinalIgnoreCase))
 				{
 					arguments.AppendArgument(additionalArgument);
 				}
 			}
 
-			return await ExecuteAutomationToolAsync(step, workspaceDir, arguments.ToString(), logger, cancellationToken) == 0;
+			if (manageSharedStorage && sharedStorageDir != null)
+			{
+				return await ExecuteWithTempStorageAsync(step, workspaceDir, sharedStorageDir, arguments.ToString(), logger, cancellationToken);
+			}
+			else
+			{
+				if (sharedStorageDir != null)
+				{
+					arguments.AppendArgument("-SharedStorageDir=", sharedStorageDir.FullName);
+				}
+				return await ExecuteAutomationToolAsync(step, workspaceDir, arguments.ToString(), logger, cancellationToken) == 0;
+			}
+		}
+
+		private async Task<bool> ExecuteWithTempStorageAsync(BeginStepResponse step, DirectoryReference workspaceDir, DirectoryReference sharedStorageDir, string arguments, ILogger logger, CancellationToken cancellationToken)
+		{
+			TempStorage Storage = new TempStorage(workspaceDir, DirectoryReference.Combine(workspaceDir, "Engine", "Saved", "BuildGraph"), sharedStorageDir, true);
+			logger.LogInformation("Using Horde-managed shared storage via {SharedStorageDir}", sharedStorageDir);
+
+			// Create the mapping of tag names to file sets
+			Dictionary<string, HashSet<FileReference>> TagNameToFileSet = new Dictionary<string, HashSet<FileReference>>();
+
+			// Read all the input tags for this node, and build a list of referenced input storage blocks
+			HashSet<TempStorageBlock> InputStorageBlocks = new HashSet<TempStorageBlock>();
+			foreach (string Input in step.Inputs)
+			{
+				int slashIdx = Input.IndexOf('/', StringComparison.Ordinal);
+				if (slashIdx == -1)
+				{
+					logger.LogError("Missing slash from node input: {Input}", Input);
+					return false;
+				}
+
+				string nodeName = Input.Substring(0, slashIdx);
+				string tagName = Input.Substring(slashIdx + 1);
+
+				TempStorageFileList? FileList = Storage.ReadFileList(nodeName, tagName, logger);
+				if (FileList == null)
+				{
+					logger.LogError("Unable to read file list for {Node}/{Tag} from shared storage ({SharedStorageDir})", nodeName, tagName, sharedStorageDir);
+					return false;
+				}
+
+				TagNameToFileSet[tagName] = FileList.ToFileSet(workspaceDir);
+				InputStorageBlocks.UnionWith(FileList.Blocks);
+			}
+
+			// Read the manifests for all the input storage blocks
+			Dictionary<TempStorageBlock, TempStorageManifest> InputManifests = new Dictionary<TempStorageBlock, TempStorageManifest>();
+			using (IScope Scope = GlobalTracer.Instance.BuildSpan("TempStorage").WithTag("resource", "read").StartActive())
+			{
+				Scope.Span.SetTag("blocks", InputStorageBlocks.Count);
+				foreach (TempStorageBlock InputStorageBlock in InputStorageBlocks)
+				{
+					TempStorageManifest Manifest = Storage.Retrieve(InputStorageBlock.NodeName, InputStorageBlock.OutputName, logger);
+					InputManifests[InputStorageBlock] = Manifest;
+				}
+				Scope.Span.SetTag("size", InputManifests.Sum(x => x.Value.GetTotalSize()));
+			}
+
+			// Read all the input storage blocks, keeping track of which block each file came from
+			Dictionary<FileReference, TempStorageBlock> FileToStorageBlock = new Dictionary<FileReference, TempStorageBlock>();
+			foreach (KeyValuePair<TempStorageBlock, TempStorageManifest> Pair in InputManifests)
+			{
+				TempStorageBlock InputStorageBlock = Pair.Key;
+				foreach (FileReference File in Pair.Value.Files.Select(x => x.ToFileReference(workspaceDir)))
+				{
+					logger.LogInformation("Reading input block: {File}", File);
+
+					TempStorageBlock? CurrentStorageBlock;
+					if (FileToStorageBlock.TryGetValue(File, out CurrentStorageBlock) && !TempStorage.IsDuplicateBuildProduct(File))
+					{
+						logger.LogError("File '{File}' was produced by {InputBlock} and {CurrentBlock}", File, InputStorageBlock.ToString(), CurrentStorageBlock.ToString());
+					}
+					FileToStorageBlock[File] = InputStorageBlock;
+				}
+			}
+
+			if (await ExecuteAutomationToolAsync(step, workspaceDir, arguments, logger, cancellationToken) != 0)
+			{
+				return false;
+			}
+
+			// Check that none of the inputs have been clobbered
+			Dictionary<string, string> ModifiedFiles = new Dictionary<string, string>(StringComparer.OrdinalIgnoreCase);
+			foreach (TempStorageFile File in InputManifests.Values.SelectMany(x => x.Files))
+			{
+				string? Message;
+				if (!ModifiedFiles.ContainsKey(File.RelativePath) && !File.Compare(workspaceDir, out Message))
+				{
+					ModifiedFiles.Add(File.RelativePath, Message);
+				}
+			}
+			if (ModifiedFiles.Count > 0)
+			{
+				string modifiedFileList = "";
+				if (ModifiedFiles.Count < 100)
+				{
+					modifiedFileList = String.Join("\n", ModifiedFiles.Select(x => x.Value));
+				}
+				else
+				{
+					modifiedFileList = String.Join("\n", ModifiedFiles.Take(100).Select(x => x.Value));
+					modifiedFileList += $"{Environment.NewLine}And {ModifiedFiles.Count - 100} more.";
+				}
+
+				logger.LogError("Build product(s) from a previous step have been modified:\n{FileList}", modifiedFileList);
+				return false;
+			}
+			/*
+			// Determine all the output files which are required to be copied to temp storage (because they're referenced by nodes in another agent)
+			HashSet<FileReference> ReferencedOutputFiles = new HashSet<FileReference>();
+			foreach (BgAgentDef Agent in Graph.Agents)
+			{
+				bool bSameAgent = Agent.Nodes.Contains(Node);
+				foreach (BgNodeDef OtherNode in Agent.Nodes)
+				{
+					if (!bSameAgent)
+					{
+						foreach (BgNodeOutput Input in OtherNode.Inputs.Where(x => x.ProducingNode == Node))
+						{
+							ReferencedOutputFiles.UnionWith(TagNameToFileSet[Input.TagName]);
+						}
+					}
+				}
+			}
+			*/
+			// Find a block name for all new outputs
+			Dictionary<FileReference, string> FileToOutputName = new Dictionary<FileReference, string>();
+			foreach (string outputName in step.OutputNames)
+			{
+				string outputNameWithoutHash = outputName.TrimStart('#');
+				bool isDefaultOutput = outputNameWithoutHash.Equals(step.Name, StringComparison.OrdinalIgnoreCase);
+
+				HashSet<FileReference> Files = TagNameToFileSet[outputName];
+				foreach (FileReference File in Files)
+				{
+					if (!FileToStorageBlock.ContainsKey(File) && File.IsUnderDirectory(workspaceDir))
+					{
+						if (isDefaultOutput)
+						{
+							if (!FileToOutputName.ContainsKey(File))
+							{
+								FileToOutputName[File] = "";
+							}
+						}
+						else
+						{
+							string? OutputName;
+							if (FileToOutputName.TryGetValue(File, out OutputName) && OutputName.Length > 0)
+							{
+								FileToOutputName[File] = String.Format("{0}+{1}", OutputName, outputNameWithoutHash);
+							}
+							else
+							{
+								FileToOutputName[File] = outputNameWithoutHash;
+							}
+						}
+					}
+				}
+			}
+
+			// Invert the dictionary to make a mapping of storage block to the files each contains
+			Dictionary<string, HashSet<FileReference>> OutputStorageBlockToFiles = new Dictionary<string, HashSet<FileReference>>();
+			foreach (KeyValuePair<FileReference, string> Pair in FileToOutputName)
+			{
+				HashSet<FileReference>? Files;
+				if (!OutputStorageBlockToFiles.TryGetValue(Pair.Value, out Files))
+				{
+					Files = new HashSet<FileReference>();
+					OutputStorageBlockToFiles.Add(Pair.Value, Files);
+				}
+				Files.Add(Pair.Key);
+			}
+
+			// Write all the storage blocks, and update the mapping from file to storage block
+			using (GlobalTracer.Instance.BuildSpan("TempStorage").WithTag("resource", "Write").StartActive())
+			{
+				foreach (KeyValuePair<string, HashSet<FileReference>> Pair in OutputStorageBlockToFiles)
+				{
+					TempStorageBlock OutputBlock = new TempStorageBlock(step.Name, Pair.Key);
+					foreach (FileReference File in Pair.Value)
+					{
+						FileToStorageBlock.Add(File, OutputBlock);
+					}
+					Storage.Archive(step.Name, Pair.Key, Pair.Value.ToArray(), true/*Pair.Value.Any(x => ReferencedOutputFiles.Contains(x))*/, logger);
+				}
+
+				// Publish all the output tags
+				foreach (string outputName in step.OutputNames)
+				{
+					HashSet<FileReference> Files = TagNameToFileSet[outputName];
+
+					HashSet<TempStorageBlock> StorageBlocks = new HashSet<TempStorageBlock>();
+					foreach (FileReference File in Files)
+					{
+						TempStorageBlock? StorageBlock;
+						if (FileToStorageBlock.TryGetValue(File, out StorageBlock))
+						{
+							StorageBlocks.Add(StorageBlock);
+						}
+					}
+
+					Storage.WriteFileList(step.Name, outputName, Files, StorageBlocks.ToArray(), logger);
+				}
+			}
+
+			return true;
 		}
 
 		protected async Task<int> ExecuteAutomationToolAsync(BeginStepResponse step, DirectoryReference workspaceDir, string arguments, ILogger logger, CancellationToken cancellationToken)
diff --git a/Engine/Source/Programs/Horde/Horde.Agent/Utility/TempStorage.cs b/Engine/Source/Programs/Horde/Horde.Agent/Utility/TempStorage.cs
new file mode 100644
--- /dev/null
+++ b/Engine/Source/Programs/Horde/Horde.Agent/Utility/TempStorage.cs
@@ -0,0 +1,1142 @@
+// Copyright Epic Games, Inc. All Rights Reserved.
+
+using EpicGames.Core;
+using Microsoft.Extensions.Logging;
+using OpenTracing;
+using OpenTracing.Util;
+using System;
+using System.Collections.Concurrent;
+using System.Collections.Generic;
+using System.Diagnostics;
+using System.Diagnostics.CodeAnalysis;
+using System.IO;
+using System.IO.Compression;
+using System.Linq;
+using System.Threading;
+using System.Threading.Tasks;
+using System.Xml;
+using System.Xml.Serialization;
+
+#pragma warning disable CA1819 // Properties should not return arrays
+
+namespace Horde.Storage.Utility
+{
+	/// <summary>
+	/// Exception thrown by the temp storage system
+	/// </summary>
+	public sealed class TempStorageException : Exception
+	{
+		/// <summary>
+		/// Constructor
+		/// </summary>
+		public TempStorageException(string message, Exception? innerException = null)
+			: base(message, innerException)
+		{
+		}
+	}
+
+	/// <summary>
+	/// Stores the name of a temp storage block
+	/// </summary>
+	public class TempStorageBlock
+	{
+		/// <summary>
+		/// Name of the node
+		/// </summary>
+		[XmlAttribute]
+		public string NodeName { get; set; }
+
+		/// <summary>
+		/// Name of the output from this node
+		/// </summary>
+		[XmlAttribute]
+		public string OutputName { get; set; }
+
+		/// <summary>
+		/// Default constructor, for XML serialization.
+		/// </summary>
+		private TempStorageBlock()
+		{
+			NodeName = String.Empty;
+			OutputName = String.Empty;
+		}
+
+		/// <summary>
+		/// Construct a temp storage block
+		/// </summary>
+		/// <param name="nodeName">Name of the node</param>
+		/// <param name="outputName">Name of the node's output</param>
+		public TempStorageBlock(string nodeName, string outputName)
+		{
+			NodeName = nodeName;
+			OutputName = outputName;
+		}
+
+		/// <summary>
+		/// Tests whether two temp storage blocks are equal
+		/// </summary>
+		/// <param name="other">The object to compare against</param>
+		/// <returns>True if the blocks are equivalent</returns>
+		public override bool Equals(object? other) => other is TempStorageBlock otherBlock && NodeName.Equals(otherBlock.NodeName, StringComparison.OrdinalIgnoreCase) && OutputName.Equals(otherBlock.OutputName, StringComparison.OrdinalIgnoreCase);
+
+		/// <summary>
+		/// Returns a hash code for this block name
+		/// </summary>
+		/// <returns>Hash code for the block</returns>
+		public override int GetHashCode() => HashCode.Combine(NodeName, OutputName);
+
+		/// <summary>
+		/// Returns the name of this block for debugging purposes
+		/// </summary>
+		/// <returns>Name of this block as a string</returns>
+		public override string ToString() => $"{NodeName}/{OutputName}";
+	}
+
+	/// <summary>
+	/// Information about a single file in temp storage
+	/// </summary>
+	[DebuggerDisplay("{RelativePath}")]
+	public class TempStorageFile
+	{
+		/// <summary>
+		/// The path of the file, relative to the engine root. Stored using forward slashes.
+		/// </summary>
+		[XmlAttribute]
+		public string RelativePath { get; set; }
+
+		/// <summary>
+		/// The last modified time of the file, in UTC ticks since the Epoch.
+		/// </summary>
+		[XmlAttribute]
+		public long LastWriteTimeUtcTicks { get; set; }
+
+		/// <summary>
+		/// Length of the file
+		/// </summary>
+		[XmlAttribute]
+		public long Length { get; set; }
+
+		/// <summary>
+		/// Digest for the file. Not all files are hashed.
+		/// </summary>
+		[XmlAttribute]
+		public string? Digest { get; set; }
+
+		/// <summary>
+		/// Default constructor, for XML serialization.
+		/// </summary>
+		private TempStorageFile()
+		{
+			RelativePath = String.Empty;
+		}
+
+		/// <summary>
+		/// Constructor
+		/// </summary>
+		/// <param name="fileInfo">File to be added</param>
+		/// <param name="rootDir">Root directory to store paths relative to</param>
+		public TempStorageFile(FileInfo fileInfo, DirectoryReference rootDir)
+		{
+			// Check the file exists and is in the right location
+			FileReference file = new FileReference(fileInfo);
+			if(!file.IsUnderDirectory(rootDir))
+			{
+				throw new TempStorageException($"Attempt to add file to temp storage manifest that is outside the root directory ({file.FullName})");
+			}
+			if(!fileInfo.Exists)
+			{
+				throw new TempStorageException($"Attempt to add file to temp storage manifest that does not exist ({file.FullName})");
+			}
+
+			RelativePath = file.MakeRelativeTo(rootDir).Replace(Path.DirectorySeparatorChar, '/');
+			LastWriteTimeUtcTicks = fileInfo.LastWriteTimeUtc.Ticks;
+			Length = fileInfo.Length;
+
+			if (GenerateDigest())
+			{
+				Digest = ContentHash.MD5(file).ToString();
+			}
+		}
+
+		/// <summary>
+		/// Compare stored for this file with the one on disk, and output an error if they differ.
+		/// </summary>
+		/// <param name="rootDir">Root directory for this branch</param>
+		/// <param name="logger">Logger for output</param>
+		/// <returns>True if the files are identical, false otherwise</returns>
+		public bool Compare(DirectoryReference rootDir, ILogger logger)
+		{
+			string? message;
+			if(Compare(rootDir, out message))
+			{
+				if(message != null)
+				{
+					logger.LogInformation("{Message}", message);
+				}
+				return true;
+			}
+			else
+			{
+				if(message != null)
+				{
+					logger.LogError("{Message}", message);
+				}
+				return false;
+			}
+		}
+
+		/// <summary>
+		/// Compare stored for this file with the one on disk, and output an error if they differ.
+		/// </summary>
+		/// <param name="rootDir">Root directory for this branch</param>
+		/// <param name="message">Message describing the difference</param>
+		/// <returns>True if the files are identical, false otherwise</returns>
+		public bool Compare(DirectoryReference rootDir, [NotNullWhen(false)] out string? message)
+		{
+			FileReference localFile = ToFileReference(rootDir);
+
+			// Get the local file info, and check it exists
+			FileInfo info = new FileInfo(localFile.FullName);
+			if(!info.Exists)
+			{
+				message = String.Format("Missing file from manifest - {0}", RelativePath);
+				return false;
+			}
+
+			// Check the size matches
+			if(info.Length != Length)
+			{
+				if(TempStorage.IsDuplicateBuildProduct(localFile))
+				{
+					message = String.Format("Ignored file size mismatch for {0} - was {1} bytes, expected {2} bytes", RelativePath, info.Length, Length);
+					return true;
+				}
+				else
+				{
+					message = String.Format("File size differs from manifest - {0} is {1} bytes, expected {2} bytes", RelativePath, info.Length, Length);
+					return false;
+				}
+			}
+
+			// Check the timestamp of the file matches. On FAT filesystems writetime has a two seconds resolution (see http://msdn.microsoft.com/en-us/library/windows/desktop/ms724290%28v=vs.85%29.aspx)
+			TimeSpan timeDifference = new TimeSpan(info.LastWriteTimeUtc.Ticks - LastWriteTimeUtcTicks);
+			if (timeDifference.TotalSeconds >= -2 && timeDifference.TotalSeconds <= +2)
+			{
+				message = null;
+				return true;
+			}
+
+			// Check if the files have been modified
+			DateTime expectedLocal = new DateTime(LastWriteTimeUtcTicks, DateTimeKind.Utc).ToLocalTime();
+			if (Digest != null)
+			{
+				string localDigest = ContentHash.MD5(localFile).ToString();
+				if (Digest.Equals(localDigest, StringComparison.Ordinal))
+				{
+					message = null;
+					return true;
+				}
+				else
+				{
+					message = String.Format("Digest mismatch for {0} - was {1} ({2}), expected {3} ({4}), TimeDifference {5}", RelativePath, localDigest, info.LastWriteTime, Digest, expectedLocal, timeDifference);
+					return false;
+				}
+			}
+			else
+			{
+				if (RequireMatchingTimestamps() && !TempStorage.IsDuplicateBuildProduct(localFile))
+				{
+					message = String.Format("File date/time mismatch for {0} - was {1}, expected {2}, TimeDifference {3}", RelativePath, info.LastWriteTime, expectedLocal, timeDifference);
+					return false;
+				}
+				else
+				{
+					message = String.Format("Ignored file date/time mismatch for {0} - was {1}, expected {2}, TimeDifference {3}", RelativePath, info.LastWriteTime, expectedLocal, timeDifference);
+					return true;
+				}
+			}
+		}
+
+		/// <summary>
+		/// Whether we should compare timestamps for this file. Some build products are harmlessly overwritten as part of the build process, so we flag those here.
+		/// </summary>
+		/// <returns>True if we should compare the file's timestamp, false otherwise</returns>
+		bool RequireMatchingTimestamps()
+		{
+			return !RelativePath.Contains("/Binaries/DotNET/", StringComparison.OrdinalIgnoreCase) && !RelativePath.Contains("/Binaries/Mac/", StringComparison.OrdinalIgnoreCase);
+		}
+
+		/// <summary>
+		/// Determines whether to generate a digest for the current file
+		/// </summary>
+		/// <returns>True to generate a digest for this file, rather than relying on timestamps</returns>
+		bool GenerateDigest()
+		{
+			return RelativePath.EndsWith(".version", StringComparison.OrdinalIgnoreCase) || RelativePath.EndsWith(".modules", StringComparison.OrdinalIgnoreCase);
+		}
+
+		/// <summary>
+		/// Determine whether to serialize the digest property
+		/// </summary>
+		/// <returns></returns>
+		public bool ShouldSerializeDigest()
+		{
+			return Digest != null;
+		}
+
+		/// <summary>
+		/// Gets a local file reference for this file, given a root directory to base it from.
+		/// </summary>
+		/// <param name="rootDir">The local root directory</param>
+		/// <returns>Reference to the file</returns>
+		public FileReference ToFileReference(DirectoryReference rootDir)
+		{
+			return FileReference.Combine(rootDir, RelativePath.Replace('/', Path.DirectorySeparatorChar));
+		}
+	}
+
+	/// <summary>
+	/// Information about a single file in temp storage
+	/// </summary>
+	[DebuggerDisplay("{Name}")]
+	public class TempStorageZipFile
+	{
+		/// <summary>
+		/// Name of this file, including extension
+		/// </summary>
+		[XmlAttribute]
+		public string Name { get; set; }
+
+		/// <summary>
+		/// Length of the file in bytes
+		/// </summary>
+		[XmlAttribute]
+		public long Length { get; set; }
+
+		/// <summary>
+		/// Default constructor, for XML serialization
+		/// </summary>
+		private TempStorageZipFile()
+		{
+			Name = String.Empty;
+		}
+
+		/// <summary>
+		/// Constructor
+		/// </summary>
+		/// <param name="info">FileInfo for the zip file</param>
+		public TempStorageZipFile(FileInfo info)
+		{
+			Name = info.Name;
+			Length = info.Length;
+		}
+	}
+
+	/// <summary>
+	/// A manifest storing information about build products for a node's output
+	/// </summary>
+	public class TempStorageManifest
+	{
+		/// <summary>
+		/// List of output files
+		/// </summary>
+		[XmlArray]
+		[XmlArrayItem("File")]
+		public TempStorageFile[] Files { get; set; }
+
+		/// <summary>
+		/// List of compressed archives containing the given files
+		/// </summary>
+		[XmlArray]
+		[XmlArrayItem("ZipFile")]
+		public TempStorageZipFile[] ZipFiles { get; set; }
+
+		/// <summary>
+		/// Construct a static Xml serializer to avoid throwing an exception searching for the reflection info at runtime
+		/// </summary>
+		static readonly XmlSerializer s_serializer = XmlSerializer.FromTypes(new Type[]{ typeof(TempStorageManifest) })[0];
+
+		/// <summary>
+		/// Construct an empty temp storage manifest
+		/// </summary>
+		private TempStorageManifest()
+		{
+			Files = Array.Empty<TempStorageFile>();
+			ZipFiles = Array.Empty<TempStorageZipFile>();
+		}
+
+		/// <summary>
+		/// Creates a manifest from a flat list of files (in many folders) and a BaseFolder from which they are rooted.
+		/// </summary>
+		/// <param name="files">List of full file paths</param>
+		/// <param name="rootDir">Root folder for all the files. All files must be relative to this RootDir.</param>
+		public TempStorageManifest(FileInfo[] files, DirectoryReference rootDir)
+		{
+			Files = files.Select(x => new TempStorageFile(x, rootDir)).ToArray();
+			ZipFiles = Array.Empty<TempStorageZipFile>();
+		}
+
+		/// <summary>
+		/// Gets the total size of the files stored in this manifest
+		/// </summary>
+		/// <returns>The total size of all files</returns>
+		public long GetTotalSize()
+		{
+			long result = 0;
+			foreach(TempStorageFile file in Files)
+			{
+				result += file.Length;
+			}
+			return result;
+		}
+
+		/// <summary>
+		/// Load a manifest from disk
+		/// </summary>
+		/// <param name="file">File to load</param>
+		public static TempStorageManifest Load(FileReference file)
+		{
+			using (StreamReader reader = new(file.FullName))
+			{
+				XmlReaderSettings settings = new XmlReaderSettings();
+				using (XmlReader xmlReader = XmlReader.Create(reader, settings))
+				{
+					return (TempStorageManifest)s_serializer.Deserialize(xmlReader)!;
+				}
+			}
+		}
+
+		/// <summary>
+		/// Saves a manifest to disk
+		/// </summary>
+		/// <param name="file">File to save</param>
+		public void Save(FileReference file)
+		{
+			using(StreamWriter writer = new StreamWriter(file.FullName))
+			{
+				XmlWriterSettings writerSettings = new() { Indent = true };
+				using (XmlWriter xmlWriter = XmlWriter.Create(writer, writerSettings))
+				{
+					s_serializer.Serialize(xmlWriter, this);
+				}
+			}
+		}
+	}
+
+	/// <summary>
+	/// Stores the contents of a tagged file set
+	/// </summary>
+	public class TempStorageFileList
+	{
+		/// <summary>
+		/// List of files that are in this tag set, relative to the root directory
+		/// </summary>
+		[XmlArray]
+		[XmlArrayItem("LocalFile")]
+		public string[] LocalFiles { get; set; }
+
+		/// <summary>
+		/// List of files that are in this tag set, but not relative to the root directory
+		/// </summary>
+		[XmlArray]
+		[XmlArrayItem("LocalFile")]
+		public string[] ExternalFiles { get; set; }
+
+		/// <summary>
+		/// List of referenced storage blocks
+		/// </summary>
+		[XmlArray]
+		[XmlArrayItem("Block")]
+		public TempStorageBlock[] Blocks { get; set; }
+
+		/// <summary>
+		/// Construct a static Xml serializer to avoid throwing an exception searching for the reflection info at runtime
+		/// </summary>
+		static readonly XmlSerializer s_serializer = XmlSerializer.FromTypes(new Type[]{ typeof(TempStorageFileList) })[0];
+
+		/// <summary>
+		/// Construct an empty file list for deserialization
+		/// </summary>
+		private TempStorageFileList()
+		{
+			LocalFiles = Array.Empty<string>();
+			ExternalFiles = Array.Empty<string>();
+			Blocks = Array.Empty<TempStorageBlock>();
+		}
+
+		/// <summary>
+		/// Creates a manifest from a flat list of files (in many folders) and a BaseFolder from which they are rooted.
+		/// </summary>
+		/// <param name="files">List of full file paths</param>
+		/// <param name="rootDir">Root folder for all the files. All files must be relative to this RootDir.</param>
+		/// <param name="blocks">Referenced storage blocks required for these files</param>
+		public TempStorageFileList(IEnumerable<FileReference> files, DirectoryReference rootDir, IEnumerable<TempStorageBlock> blocks)
+		{
+			List<string> newLocalFiles = new List<string>();
+			List<string> newExternalFiles = new List<string>();
+			foreach(FileReference file in files)
+			{
+				if(file.IsUnderDirectory(rootDir))
+				{
+					newLocalFiles.Add(file.MakeRelativeTo(rootDir).Replace(Path.DirectorySeparatorChar, '/'));
+				}
+				else
+				{
+					newExternalFiles.Add(file.FullName.Replace(Path.DirectorySeparatorChar, '/'));
+				}
+			}
+			LocalFiles = newLocalFiles.ToArray();
+			ExternalFiles = newExternalFiles.ToArray();
+
+			Blocks = blocks.ToArray();
+		}
+
+		/// <summary>
+		/// Load this list of files from disk
+		/// </summary>
+		/// <param name="file">File to load</param>
+		public static TempStorageFileList Load(FileReference file)
+		{
+			using(StreamReader reader = new StreamReader(file.FullName))
+			{
+				XmlReaderSettings settings = new XmlReaderSettings();
+				using (XmlReader xmlReader = XmlReader.Create(reader, settings))
+				{
+					return (TempStorageFileList)s_serializer.Deserialize(xmlReader)!;
+				}
+			}
+		}
+
+		/// <summary>
+		/// Saves this list of files to disk
+		/// </summary>
+		/// <param name="file">File to save</param>
+		public void Save(FileReference file)
+		{
+			using (StreamWriter writer = new StreamWriter(file.FullName))
+			{
+				XmlWriterSettings writerSettings = new() { Indent = true };
+				using (XmlWriter xmlWriter = XmlWriter.Create(writer, writerSettings))
+				{
+					s_serializer.Serialize(xmlWriter, this);
+				}
+			}
+		}
+
+		/// <summary>
+		/// Converts this file list into a set of FileReference objects
+		/// </summary>
+		/// <param name="rootDir">The root directory to rebase local files</param>
+		/// <returns>Set of files</returns>
+		public HashSet<FileReference> ToFileSet(DirectoryReference rootDir)
+		{
+			HashSet<FileReference> Files = new HashSet<FileReference>();
+			Files.UnionWith(LocalFiles.Select(x => FileReference.Combine(rootDir, x)));
+			Files.UnionWith(ExternalFiles.Select(x => new FileReference(x)));
+			return Files;
+		}
+	}
+
+	/// <summary>
+	/// Tracks the state of the current build job using the filesystem, allowing jobs to be restarted after a failure or expanded to include larger targets, and 
+	/// providing a proxy for different machines executing parts of the build in parallel to transfer build products and share state as part of a build system.
+	/// 
+	/// If a shared temp storage directory is provided - typically a mounted path on a network share - all build products potentially needed as inputs by another node
+	/// are compressed and copied over, along with metadata for them (see TempStorageFile) and flags for build events that have occurred (see TempStorageEvent).
+	/// 
+	/// The local temp storage directory contains the same information, with the exception of the archived build products. Metadata is still kept to detect modified 
+	/// build products between runs. If data is not present in local temp storage, it's retrieved from shared temp storage and cached in local storage.
+	/// </summary>
+	class TempStorage
+	{
+		/// <summary>
+		/// Root directory for this branch.
+		/// </summary>
+		readonly DirectoryReference _rootDir;
+
+		/// <summary>
+		/// The local temp storage directory (typically somewhere under /Engine/Saved directory).
+		/// </summary>
+		readonly DirectoryReference _localDir;
+
+		/// <summary>
+		/// The shared temp storage directory; typically a network location. May be null.
+		/// </summary>
+		readonly DirectoryReference _sharedDir;
+
+		/// <summary>
+		/// Whether to allow writes to shared storage
+		/// </summary>
+		readonly bool _writeToSharedStorage;
+
+		/// <summary>
+		/// Constructor
+		/// </summary>
+		/// <param name="rootDir">Root directory for this branch</param>
+		/// <param name="localDir">The local temp storage directory.</param>
+		/// <param name="sharedDir">The shared temp storage directory. May be null.</param>
+		/// <param name="writeToSharedStorage">Whether to write to shared storage, or only permit reads from it</param>
+		public TempStorage(DirectoryReference rootDir, DirectoryReference localDir, DirectoryReference sharedDir, bool writeToSharedStorage)
+		{
+			_rootDir = rootDir;
+			_localDir = localDir;
+			_sharedDir = sharedDir;
+			_writeToSharedStorage = writeToSharedStorage;
+		}
+
+		/// <summary>
+		/// Mark the given node as complete
+		/// </summary>
+		/// <param name="nodeName">Name of the node</param>
+		public void MarkAsComplete(string nodeName)
+		{
+			// Create the marker locally
+			FileReference localFile = GetCompleteMarkerFile(_localDir, nodeName);
+			DirectoryReference.CreateDirectory(localFile.Directory);
+			File.OpenWrite(localFile.FullName).Close();
+
+			// Create the marker in the shared directory
+			if(_sharedDir != null && _writeToSharedStorage)
+			{
+				FileReference sharedFile = GetCompleteMarkerFile(_sharedDir, nodeName);
+				DirectoryReference.CreateDirectory(sharedFile.Directory);
+				File.OpenWrite(sharedFile.FullName).Close();
+			}
+		}
+
+		/// <summary>
+		/// Checks the integrity of the give node's local build products.
+		/// </summary>
+		/// <param name="nodeName">The node to retrieve build products for</param>
+		/// <param name="tagNames">List of tag names from this node.</param>
+		/// <param name="logger">Logger for output</param>
+		/// <returns>True if the node is complete and valid, false if not (and typically followed by a call to CleanNode()).</returns>
+		public bool CheckLocalIntegrity(string nodeName, IEnumerable<string> tagNames, ILogger logger)
+		{
+			// If the node is not locally complete, fail immediately.
+			FileReference completeMarkerFile = GetCompleteMarkerFile(_localDir, nodeName);
+			if(!FileReference.Exists(completeMarkerFile))
+			{
+				return false;
+			}
+
+			// Check that each of the tags exist
+			HashSet<TempStorageBlock> blocks = new HashSet<TempStorageBlock>();
+			foreach(string tagName in tagNames)
+			{
+				// Check the local manifest exists
+				FileReference LocalFileListLocation = GetTaggedFileListLocation(_localDir, nodeName, tagName);
+				if(!FileReference.Exists(LocalFileListLocation))
+				{
+					return false;
+				}
+
+				// Check the local manifest matches the shared manifest
+				if(_sharedDir != null)
+				{
+					// Check the shared manifest exists
+					FileReference sharedFileListLocation = GetManifestLocation(_sharedDir, nodeName, tagName);
+					if(!FileReference.Exists(sharedFileListLocation))
+					{
+						return false;
+					}
+
+					// Check the manifests are identical, byte by byte
+					byte[] localManifestBytes = FileReference.ReadAllBytes(LocalFileListLocation);
+					byte[] sharedManifestBytes = FileReference.ReadAllBytes(sharedFileListLocation);
+					if (!localManifestBytes.SequenceEqual(sharedManifestBytes))
+					{
+						return false;
+					}
+				}
+
+				// Read the manifest and add the referenced blocks to be checked
+				TempStorageFileList localFileList = TempStorageFileList.Load(LocalFileListLocation);
+				blocks.UnionWith(localFileList.Blocks);
+			}
+
+			// Check that each of the outputs match
+			foreach(TempStorageBlock block in blocks)
+			{
+				// Check the local manifest exists
+				FileReference localManifestFile = GetManifestLocation(_localDir, block.NodeName, block.OutputName);
+				if(!FileReference.Exists(localManifestFile))
+				{
+					return false;
+				}
+
+				// Check the local manifest matches the shared manifest
+				if(_sharedDir != null)
+				{
+					// Check the shared manifest exists
+					FileReference sharedManifestFile = GetManifestLocation(_sharedDir, block.NodeName, block.OutputName);
+					if(!FileReference.Exists(sharedManifestFile))
+					{
+						return false;
+					}
+
+					// Check the manifests are identical, byte by byte
+					byte[] localManifestBytes = FileReference.ReadAllBytes(localManifestFile);
+					byte[] sharedManifestBytes = FileReference.ReadAllBytes(sharedManifestFile);
+					if(!localManifestBytes.SequenceEqual(sharedManifestBytes))
+					{
+						return false;
+					}
+				}
+
+				// Read the manifest and check the files
+				TempStorageManifest localManifest = TempStorageManifest.Load(localManifestFile);
+				if(localManifest.Files.Any(x => !x.Compare(_rootDir, logger)))
+				{
+					return false;
+				}
+			}
+			return true;
+		}
+
+		/// <summary>
+		/// Reads a set of tagged files from disk
+		/// </summary>
+		/// <param name="nodeName">Name of the node which produced the tag set</param>
+		/// <param name="tagName">Name of the tag, with a '#' prefix</param>
+		/// <param name="logger">Logger for output</param>
+		/// <returns>The set of files</returns>
+		public TempStorageFileList? ReadFileList(string nodeName, string tagName, ILogger logger)
+		{
+			TempStorageFileList? fileList;
+
+			// Try to read the tag set from the local directory
+			FileReference localFileListLocation = GetTaggedFileListLocation(_localDir, nodeName, tagName);
+			if(FileReference.Exists(localFileListLocation))
+			{
+				logger.LogInformation("Reading local file list from {File}", localFileListLocation.FullName);
+				fileList = TempStorageFileList.Load(localFileListLocation);
+			}
+			else
+			{
+				// Check we have shared storage
+				if(_sharedDir == null)
+				{
+					throw new TempStorageException($"Missing local file list - {localFileListLocation.FullName}");
+				}
+
+				// Make sure the manifest exists.
+				FileReference sharedFileListLocation = GetTaggedFileListLocation(_sharedDir, nodeName, tagName);
+				if (!FileReference.Exists(sharedFileListLocation))
+				{
+					throw new TempStorageException($"Missing local or shared file list - {sharedFileListLocation.FullName}");
+				}
+
+				try
+				{
+					// Read the shared manifest
+					logger.LogInformation("Copying shared tag set from {Source} to {Target}", sharedFileListLocation.FullName, localFileListLocation.FullName);
+					fileList = TempStorageFileList.Load(sharedFileListLocation);
+				}
+				catch
+				{
+					throw new TempStorageException($"Local or shared file list {sharedFileListLocation.FullName} was found but failed to be read");
+				}
+
+				// Save the manifest locally
+				DirectoryReference.CreateDirectory(localFileListLocation.Directory);
+				fileList?.Save(localFileListLocation);
+			}
+			return fileList;
+		}
+
+		/// <summary>
+		/// Writes a list of tagged files to disk
+		/// </summary>
+		/// <param name="nodeName">Name of the node which produced the tag set</param>
+		/// <param name="tagName">Name of the tag, with a '#' prefix</param>
+		/// <param name="files">List of files in this set</param>
+		/// <param name="blocks">List of referenced storage blocks</param>
+		/// <param name="logger">Logger for output</param>
+		/// <returns>The set of files</returns>
+		public void WriteFileList(string nodeName, string tagName, IEnumerable<FileReference> files, IEnumerable<TempStorageBlock> blocks, ILogger logger)
+		{
+			// Create the file list
+			TempStorageFileList fileList = new TempStorageFileList(files, _rootDir, blocks);
+
+			// Save the set of files to the local and shared locations
+			FileReference localFileListLocation = GetTaggedFileListLocation(_localDir, nodeName, tagName);
+			if(_sharedDir != null && _writeToSharedStorage)
+			{
+				FileReference sharedFileListLocation = GetTaggedFileListLocation(_sharedDir, nodeName, tagName);
+
+				try
+				{
+					logger.LogInformation("Saving file list to {File} and {OtherFile}", localFileListLocation.FullName, sharedFileListLocation.FullName);
+					DirectoryReference.CreateDirectory(sharedFileListLocation.Directory);
+					fileList.Save(sharedFileListLocation);
+				}
+				catch (Exception ex)
+				{
+					throw new TempStorageException($"Failed to save file list {localFileListLocation} to {sharedFileListLocation}, exception: {ex}");
+				}
+			}
+			else
+			{
+				logger.LogInformation("Saving file list to {File}", localFileListLocation.FullName);
+			}
+
+			// Save the local file list
+			DirectoryReference.CreateDirectory(localFileListLocation.Directory);
+			fileList.Save(localFileListLocation);
+		}
+
+		/// <summary>
+		/// Saves the given files (that should be rooted at the branch root) to a shared temp storage manifest with the given temp storage node and game.
+		/// </summary>
+		/// <param name="NodeName">The node which created the storage block</param>
+		/// <param name="BlockName">Name of the block to retrieve. May be null or empty.</param>
+		/// <param name="BuildProducts">Array of build products to be archived</param>
+		/// <param name="bPushToRemote">Allow skipping the copying of this manifest to shared storage, because it's not required by any other agent</param>
+		/// <param name="logger">Logger for output</param>
+		/// <returns>The created manifest instance (which has already been saved to disk).</returns>
+		public TempStorageManifest Archive(string NodeName, string BlockName, FileReference[] BuildProducts, bool bPushToRemote, ILogger logger)
+		{
+			using (IScope scope = GlobalTracer.Instance.BuildSpan("StoreToTempStorage").StartActive())
+			{
+				// Create a manifest for the given build products
+				FileInfo[] files = BuildProducts.Select(x => new FileInfo(x.FullName)).ToArray();
+				TempStorageManifest manifest = new TempStorageManifest(files, _rootDir);
+
+				// Compress the files and copy to shared storage if necessary
+				bool remote = _sharedDir != null && bPushToRemote && _writeToSharedStorage;
+				if(remote)
+				{
+					// Create the shared directory for this node
+					FileReference sharedManifestFile = GetManifestLocation(_sharedDir!, NodeName, BlockName);
+					DirectoryReference.CreateDirectory(sharedManifestFile.Directory);
+
+					// Zip all the build products
+					FileInfo[] zipFiles = ParallelZipFiles(files, _rootDir, sharedManifestFile.Directory, sharedManifestFile.GetFileNameWithoutExtension(), logger);
+					manifest.ZipFiles = zipFiles.Select(x => new TempStorageZipFile(x)).ToArray();
+
+					// Save the shared manifest
+					logger.LogInformation("Saving shared manifest to {File}", sharedManifestFile.FullName);
+					manifest.Save(sharedManifestFile);
+				}
+
+				// Save the local manifest
+				FileReference localManifestFile = GetManifestLocation(_localDir, NodeName, BlockName);
+				logger.LogInformation("Saving local manifest to {File}", localManifestFile.FullName);
+				manifest.Save(localManifestFile);
+
+				// Update the stats
+				long zipFilesTotalSize = (manifest.ZipFiles == null)? 0 : manifest.ZipFiles.Sum(x => x.Length);
+				scope.Span.SetTag("numFiles", files.Length);
+				scope.Span.SetTag("manifestSize", manifest.GetTotalSize());
+				scope.Span.SetTag("manifestZipFilesSize", zipFilesTotalSize);
+				scope.Span.SetTag("isRemote", remote);
+				scope.Span.SetTag("blockName", BlockName);
+				return manifest;
+			}
+		}
+
+		/// <summary>
+		/// Retrieve an output of the given node. Fetches and decompresses the files from shared storage if necessary, or validates the local files.
+		/// </summary>
+		/// <param name="nodeName">The node which created the storage block</param>
+		/// <param name="outputName">Name of the block to retrieve. May be null or empty.</param>
+		/// <param name="logger">Logger for output</param>
+		/// <returns>Manifest of the files retrieved</returns>
+		public TempStorageManifest Retrieve(string nodeName, string outputName, ILogger logger)
+		{
+			using (IScope scope = GlobalTracer.Instance.BuildSpan("RetrieveFromTempStorage").StartActive())
+			{
+				// Get the path to the local manifest
+				FileReference localManifestFile = GetManifestLocation(_localDir, nodeName, outputName);
+				bool local = FileReference.Exists(localManifestFile);
+
+				// Read the manifest, either from local storage or shared storage
+				TempStorageManifest? manifest = null;
+				if(local)
+				{
+					logger.LogInformation("Reading shared manifest from {File}", localManifestFile.FullName);
+					manifest = TempStorageManifest.Load(localManifestFile);
+				}
+				else
+				{
+					// Check we have shared storage
+					if(_sharedDir == null)
+					{
+						throw new TempStorageException($"Missing local manifest for node - {localManifestFile.FullName}");
+					}
+
+					// Get the shared directory for this node
+					FileReference sharedManifestFile = GetManifestLocation(_sharedDir, nodeName, outputName);
+
+					// Make sure the manifest exists
+					if(!FileReference.Exists(sharedManifestFile))
+					{
+						throw new TempStorageException($"Missing local or shared manifest for node - {sharedManifestFile.FullName}");
+					}
+
+					// Read the shared manifest
+					logger.LogInformation("Copying shared manifest from {Source} to {Target}", sharedManifestFile.FullName, localManifestFile.FullName);
+					manifest = TempStorageManifest.Load(sharedManifestFile);
+
+					// Unzip all the build products
+					DirectoryReference sharedNodeDir = GetDirectoryForNode(_sharedDir, nodeName);
+					FileInfo[] zipFiles = manifest!.ZipFiles.Select(x => new FileInfo(FileReference.Combine(sharedNodeDir, x.Name).FullName)).ToArray();
+					ParallelUnzipFiles(zipFiles, _rootDir);
+
+					// Update the timestamps to match the manifest. Zip files only use local time, and there's no guarantee it matches the local clock.
+					foreach(TempStorageFile ManifestFile in manifest.Files)
+					{
+						FileReference File = ManifestFile.ToFileReference(_rootDir);
+						System.IO.File.SetLastWriteTimeUtc(File.FullName, new DateTime(ManifestFile.LastWriteTimeUtcTicks, DateTimeKind.Utc));
+					}
+
+					// Save the manifest locally
+					DirectoryReference.CreateDirectory(localManifestFile.Directory);
+					manifest.Save(localManifestFile);
+				}
+
+				// Check all the local files are as expected
+				bool bAllMatch = true;
+				foreach(TempStorageFile File in manifest.Files)
+				{
+					bAllMatch &= File.Compare(_rootDir, logger);
+				}
+				if(!bAllMatch)
+				{
+					throw new TempStorageException("Files have been modified");
+				}
+
+				// Update the stats and return
+				scope.Span.SetTag("numFiles", manifest.Files.Length);
+				scope.Span.SetTag("manifestSize", manifest.Files.Sum(x => x.Length));
+				scope.Span.SetTag("manifestZipFilesSize", local? 0 : manifest.ZipFiles.Sum(x => x.Length));
+				scope.Span.SetTag("isRemote", !local);
+				scope.Span.SetTag("outputName", outputName);
+				return manifest;
+			}
+		}
+
+		/// <summary>
+		/// Zips a set of files (that must be rooted at the given RootDir) to a set of zip files in the given OutputDir. The files will be prefixed with the given basename.
+		/// </summary>
+		/// <param name="InputFiles">Fully qualified list of files to zip (must be rooted at RootDir).</param>
+		/// <param name="RootDir">Root Directory where all files will be extracted.</param>
+		/// <param name="OutputDir">Location to place the set of zip files created.</param>
+		/// <param name="ZipBaseName">The basename of the set of zip files.</param>
+		/// <param name="logger">Logger for output</param>
+		/// <returns>Some metrics about the zip process.</returns>
+		/// <remarks>
+		/// This function tries to zip the files in parallel as fast as it can. It makes no guarantees about how many zip files will be created or which files will be in which zip,
+		/// but it does try to reasonably balance the file sizes.
+		/// </remarks>
+		static FileInfo[] ParallelZipFiles(FileInfo[] InputFiles, DirectoryReference RootDir, DirectoryReference OutputDir, string ZipBaseName, ILogger logger)
+		{
+			// First get the sizes of all the files. We won't parallelize if there isn't enough data to keep the number of zips down.
+			var FilesInfo = InputFiles
+				.Select(InputFile => new { File = new FileReference(InputFile), FileSize = InputFile.Length })
+				.ToList();
+
+			// Profiling results show that we can zip 100MB quite fast and it is not worth parallelizing that case and creating a bunch of zips that are relatively small.
+			const long MinFileSizeToZipInParallel = 1024 * 1024 * 100L;
+			bool bZipInParallel = FilesInfo.Sum(FileInfo => FileInfo.FileSize) >= MinFileSizeToZipInParallel;
+
+			// order the files in descending order so our threads pick up the biggest ones first.
+			// We want to end with the smaller files to more effectively fill in the gaps
+			ConcurrentQueue<FileReference> FilesToZip = new(FilesInfo.OrderByDescending(FileInfo => FileInfo.FileSize).Select(FileInfo => FileInfo.File));
+
+			ConcurrentBag<FileInfo> ZipFiles = new ConcurrentBag<FileInfo>();
+
+			DirectoryReference ZipDir = OutputDir;
+
+			// We deliberately avoid Parallel.ForEach here because profiles have shown that dynamic partitioning creates
+			// too many zip files, and they can be of too varying size, creating uneven work when unzipping later,
+			// as ZipFile cannot unzip files in parallel from a single archive.
+			// We can safely assume the build system will not be doing more important things at the same time, so we simply use all our logical cores,
+			// which has shown to be optimal via profiling, and limits the number of resulting zip files to the number of logical cores.
+			List<Thread> ZipThreads = (
+				from CoreNum in Enumerable.Range(0, bZipInParallel ? Environment.ProcessorCount : 1)
+				select new Thread((object? indexObject) =>
+				{
+					int index = (int)indexObject!;
+					FileReference ZipFileName = FileReference.Combine(ZipDir, String.Format("{0}{1}.zip", ZipBaseName, bZipInParallel ? "-" + index.ToString("00") : ""));
+					// don't create the zip unless we have at least one file to add
+					FileReference? File;
+					if (FilesToZip.TryDequeue(out File))
+					{
+						try
+						{
+							// Create one zip per thread using the given basename
+							using (ZipArchive ZipArchive = ZipFile.Open(ZipFileName.FullName, ZipArchiveMode.Create))
+							{
+								// pull from the queue until we are out of files.
+								do
+								{
+									// use fastest compression. In our best case we are CPU bound, so this is a good tradeoff,
+									// cutting overall time by 2/3 while only modestly increasing the compression ratio (22.7% -> 23.8% for RootEditor PDBs).
+									// This is in cases of a super hot cache, so the operation was largely CPU bound.
+									ZipArchiveExtensions.CreateEntryFromFile_CrossPlatform(ZipArchive, File.FullName, File.MakeRelativeTo(RootDir).Replace(Path.DirectorySeparatorChar, '/'), CompressionLevel.Fastest);
+								} while (FilesToZip.TryDequeue(out File));
+							}
+						}
+						catch (IOException ex)
+						{
+							logger.LogError(ex, "Unable to open file for TempStorage zip: \"{File}\"", ZipFileName);
+							throw new TempStorageException($"Unable to open file {ZipFileName.FullName}");
+						}
+
+						ZipFiles.Add(new FileInfo(ZipFileName.FullName));
+					}
+				})).ToList();
+
+			for (int index = 0; index < ZipThreads.Count; index++)
+			{
+				Thread thread = ZipThreads[index];
+				thread.Start(index);
+			}
+
+			ZipThreads.ForEach(thread => thread.Join());
+
+			return ZipFiles.OrderBy(x => x.Name).ToArray();
+		}
+
+		/// <summary>
+		/// Unzips a set of zip files with a given basename in a given folder to a given RootDir.
+		/// </summary>
+		/// <param name="zipFiles">Files to extract</param>
+		/// <param name="rootDir">Root Directory where all files will be extracted.</param>
+		/// <returns>Some metrics about the unzip process.</returns>
+		/// <remarks>
+		/// The code is expected to be the used as the symmetrical inverse of <see cref="ParallelZipFiles"/>, but could be used independently, as long as the files in the zip do not overlap.
+		/// </remarks>
+		private static void ParallelUnzipFiles(FileInfo[] zipFiles, DirectoryReference rootDir)
+		{
+			Parallel.ForEach(zipFiles,
+				(zipFile) =>
+				{
+					string localZipFile = zipFile.FullName;
+
+					// unzip the files manually instead of caling ZipFile.ExtractToDirectory() because we need to overwrite readonly files. Because of this, creating the directories is up to us as well.
+					using (ZipArchive zipArchive = ZipFile.OpenRead(localZipFile))
+					{
+						foreach (ZipArchiveEntry entry in zipArchive.Entries)
+						{
+							// We must delete any existing file, even if it's readonly. .Net does not do this by default.
+							FileReference extractedFileName = FileReference.Combine(rootDir, entry.FullName);
+							if (FileReference.Exists(extractedFileName))
+							{
+								FileUtils.ForceDeleteFile(extractedFileName);
+							}
+							else
+							{
+								DirectoryReference.CreateDirectory(extractedFileName.Directory!);
+							}
+
+							entry.ExtractToFile_CrossPlatform(extractedFileName.FullName, true);
+						}
+					}
+				});
+		}
+
+		/// <summary>
+		/// Gets the directory used to store data for the given node
+		/// </summary>
+		/// <param name="baseDir">A local or shared temp storage root directory.</param>
+		/// <param name="nodeName">Name of the node</param>
+		/// <returns>Directory to contain a node's data</returns>
+		static DirectoryReference GetDirectoryForNode(DirectoryReference baseDir, string nodeName)
+		{
+			return DirectoryReference.Combine(baseDir, nodeName);
+		}
+
+		/// <summary>
+		/// Gets the path to the manifest created for a node's output.
+		/// </summary>
+		/// <param name="baseDir">A local or shared temp storage root directory.</param>
+		/// <param name="nodeName">Name of the node to get the file for</param>
+		/// <param name="blockName">Name of the output block to get the manifest for</param>
+		static FileReference GetManifestLocation(DirectoryReference baseDir, string nodeName, string blockName)
+		{
+			return FileReference.Combine(baseDir, nodeName, String.IsNullOrEmpty(blockName)? "Manifest.xml" : String.Format("Manifest-{0}.xml", blockName));
+		}
+
+		/// <summary>
+		/// Gets the path to the file created to store a tag manifest for a node
+		/// </summary>
+		/// <param name="baseDir">A local or shared temp storage root directory.</param>
+		/// <param name="nodeName">Name of the node to get the file for</param>
+		/// <param name="tagName">Name of the tag to get the manifest for</param>
+		static FileReference GetTaggedFileListLocation(DirectoryReference baseDir, string nodeName, string tagName)
+		{
+			Debug.Assert(tagName.StartsWith("#", StringComparison.Ordinal));
+			return FileReference.Combine(baseDir, nodeName, String.Format("Tag-{0}.xml", tagName.Substring(1)));
+		}
+
+		/// <summary>
+		/// Gets the path to a file created to indicate that a node is complete, under the given base directory.
+		/// </summary>
+		/// <param name="baseDir">A local or shared temp storage root directory.</param>
+		/// <param name="nodeName">Name of the node to get the file for</param>
+		static FileReference GetCompleteMarkerFile(DirectoryReference baseDir, string nodeName)
+		{
+			return FileReference.Combine(GetDirectoryForNode(baseDir, nodeName), "Complete");
+		}
+
+		/// <summary>
+		/// Checks whether the given path is allowed as a build product that can be produced by more than one node (timestamps may be modified, etc..). Used to suppress
+		/// warnings about build products being overwritten.
+		/// </summary>
+		/// <param name="localFile">File name to check</param>
+		/// <returns>True if the given path may be output by multiple build products</returns>
+		public static bool IsDuplicateBuildProduct(FileReference localFile)
+		{
+			string fileName = localFile.GetFileName();
+			if (fileName.Equals("AgentInterface.dll", StringComparison.OrdinalIgnoreCase) || fileName.Equals("AgentInterface.pdb", StringComparison.OrdinalIgnoreCase))
+			{
+				return true;
+			}
+			if (fileName.Equals("dxil.dll", StringComparison.OrdinalIgnoreCase))
+			{
+				return true;
+			}
+			if (fileName.Equals("dxcompiler.dll", StringComparison.OrdinalIgnoreCase))
+			{
+				return true;
+			}
+			if (fileName.Equals("embree.2.14.0.dll", StringComparison.OrdinalIgnoreCase) || fileName.Equals("libembree.2.14.0.dylib", StringComparison.OrdinalIgnoreCase))
+			{
+				return true;
+			}
+			if (fileName.Equals("tbb.dll", StringComparison.OrdinalIgnoreCase) || fileName.Equals("tbb.pdb", StringComparison.OrdinalIgnoreCase) || fileName.Equals("libtbb.dylib", StringComparison.OrdinalIgnoreCase) || fileName.Equals("tbb.psym", StringComparison.OrdinalIgnoreCase))
+			{
+				return true;
+			}
+			if (fileName.Equals("tbbmalloc.dll", StringComparison.OrdinalIgnoreCase) || fileName.Equals("libtbbmalloc.dylib", StringComparison.OrdinalIgnoreCase))
+			{
+				return true;
+			}
+			if (fileName.EndsWith(".dll", StringComparison.OrdinalIgnoreCase))
+			{
+				return true;
+			}
+			if (fileName.StartsWith("lib", StringComparison.OrdinalIgnoreCase) && fileName.EndsWith(".dylib", StringComparison.OrdinalIgnoreCase))
+			{
+				return true;
+			}
+			if (fileName.StartsWith("lib", StringComparison.OrdinalIgnoreCase) && fileName.EndsWith(".so", StringComparison.OrdinalIgnoreCase))
+			{
+				return true;
+			}
+			if (fileName.StartsWith("lib", StringComparison.OrdinalIgnoreCase) && fileName.Contains(".so.", StringComparison.OrdinalIgnoreCase))
+			{
+				// e.g. a Unix shared library with a version number suffix.
+				return true;
+			}
+			if (fileName.Equals("plugInfo.json", StringComparison.OrdinalIgnoreCase))
+			{
+				return true;
+			}
+			return false;
+		}
+	}
+}
