Horde: Only copy output blocks referenced by a downstream node on another agent to temp storage.

diff --git a/Engine/Source/Programs/Horde/Horde.Agent/Execution/BuildGraphExecutor.cs b/Engine/Source/Programs/Horde/Horde.Agent/Execution/BuildGraphExecutor.cs
--- a/Engine/Source/Programs/Horde/Horde.Agent/Execution/BuildGraphExecutor.cs
+++ b/Engine/Source/Programs/Horde/Horde.Agent/Execution/BuildGraphExecutor.cs
@@ -641,146 +641,140 @@ namespace Horde.Agent.Execution
 
 		private async Task<bool> ExecuteWithTempStorageAsync(BeginStepResponse step, DirectoryReference workspaceDir, DirectoryReference sharedStorageDir, string arguments, ILogger logger, CancellationToken cancellationToken)
 		{
-			TempStorage Storage = new TempStorage(workspaceDir, DirectoryReference.Combine(workspaceDir, "Engine", "Saved", "BuildGraph"), sharedStorageDir, true);
+			TempStorage storage = new TempStorage(workspaceDir, DirectoryReference.Combine(workspaceDir, "Engine", "Saved", "BuildGraph"), sharedStorageDir, true);
 			logger.LogInformation("Using Horde-managed shared storage via {SharedStorageDir}", sharedStorageDir);
 
 			// Create the mapping of tag names to file sets
-			Dictionary<string, HashSet<FileReference>> TagNameToFileSet = new Dictionary<string, HashSet<FileReference>>();
+			Dictionary<string, HashSet<FileReference>> tagNameToFileSet = new Dictionary<string, HashSet<FileReference>>();
 
 			// Read all the input tags for this node, and build a list of referenced input storage blocks
-			HashSet<TempStorageBlock> InputStorageBlocks = new HashSet<TempStorageBlock>();
-			foreach (string Input in step.Inputs)
+			HashSet<TempStorageBlock> inputStorageBlocks = new HashSet<TempStorageBlock>();
+			foreach (string input in step.Inputs)
 			{
-				int slashIdx = Input.IndexOf('/', StringComparison.Ordinal);
+				int slashIdx = input.IndexOf('/', StringComparison.Ordinal);
 				if (slashIdx == -1)
 				{
-					logger.LogError("Missing slash from node input: {Input}", Input);
+					logger.LogError("Missing slash from node input: {Input}", input);
 					return false;
 				}
 
-				string nodeName = Input.Substring(0, slashIdx);
-				string tagName = Input.Substring(slashIdx + 1);
+				string nodeName = input.Substring(0, slashIdx);
+				string tagName = input.Substring(slashIdx + 1);
 
-				TempStorageFileList? FileList = Storage.ReadFileList(nodeName, tagName, logger);
-				if (FileList == null)
+				TempStorageFileList? fileList = storage.ReadFileList(nodeName, tagName, logger);
+				if (fileList == null)
 				{
 					logger.LogError("Unable to read file list for {Node}/{Tag} from shared storage ({SharedStorageDir})", nodeName, tagName, sharedStorageDir);
 					return false;
 				}
 
-				TagNameToFileSet[tagName] = FileList.ToFileSet(workspaceDir);
-				InputStorageBlocks.UnionWith(FileList.Blocks);
+				tagNameToFileSet[tagName] = fileList.ToFileSet(workspaceDir);
+				inputStorageBlocks.UnionWith(fileList.Blocks);
 			}
 
 			// Read the manifests for all the input storage blocks
-			Dictionary<TempStorageBlock, TempStorageManifest> InputManifests = new Dictionary<TempStorageBlock, TempStorageManifest>();
-			using (IScope Scope = GlobalTracer.Instance.BuildSpan("TempStorage").WithTag("resource", "read").StartActive())
+			Dictionary<TempStorageBlock, TempStorageManifest> inputManifests = new Dictionary<TempStorageBlock, TempStorageManifest>();
+			using (IScope scope = GlobalTracer.Instance.BuildSpan("TempStorage").WithTag("resource", "read").StartActive())
 			{
-				Scope.Span.SetTag("blocks", InputStorageBlocks.Count);
-				foreach (TempStorageBlock InputStorageBlock in InputStorageBlocks)
+				scope.Span.SetTag("blocks", inputStorageBlocks.Count);
+				foreach (TempStorageBlock inputStorageBlock in inputStorageBlocks)
 				{
-					TempStorageManifest Manifest = Storage.Retrieve(InputStorageBlock.NodeName, InputStorageBlock.OutputName, logger);
-					InputManifests[InputStorageBlock] = Manifest;
+					TempStorageManifest manifest = storage.Retrieve(inputStorageBlock.NodeName, inputStorageBlock.OutputName, logger);
+					inputManifests[inputStorageBlock] = manifest;
 				}
-				Scope.Span.SetTag("size", InputManifests.Sum(x => x.Value.GetTotalSize()));
+				scope.Span.SetTag("size", inputManifests.Sum(x => x.Value.GetTotalSize()));
 			}
 
 			// Read all the input storage blocks, keeping track of which block each file came from
-			Dictionary<FileReference, TempStorageBlock> FileToStorageBlock = new Dictionary<FileReference, TempStorageBlock>();
-			foreach (KeyValuePair<TempStorageBlock, TempStorageManifest> Pair in InputManifests)
+			Dictionary<FileReference, TempStorageBlock> fileToStorageBlock = new Dictionary<FileReference, TempStorageBlock>();
+			foreach (KeyValuePair<TempStorageBlock, TempStorageManifest> pair in inputManifests)
 			{
-				TempStorageBlock InputStorageBlock = Pair.Key;
-				foreach (FileReference File in Pair.Value.Files.Select(x => x.ToFileReference(workspaceDir)))
+				TempStorageBlock inputStorageBlock = pair.Key;
+				foreach (FileReference file in pair.Value.Files.Select(x => x.ToFileReference(workspaceDir)))
 				{
-					logger.LogInformation("Reading input block: {File}", File);
+					logger.LogInformation("Reading input block: {File}", file);
 
-					TempStorageBlock? CurrentStorageBlock;
-					if (FileToStorageBlock.TryGetValue(File, out CurrentStorageBlock) && !TempStorage.IsDuplicateBuildProduct(File))
+					TempStorageBlock? currentStorageBlock;
+					if (fileToStorageBlock.TryGetValue(file, out currentStorageBlock) && !TempStorage.IsDuplicateBuildProduct(file))
 					{
-						logger.LogError("File '{File}' was produced by {InputBlock} and {CurrentBlock}", File, InputStorageBlock.ToString(), CurrentStorageBlock.ToString());
+						logger.LogError("File '{File}' was produced by {InputBlock} and {CurrentBlock}", file, inputStorageBlock.ToString(), currentStorageBlock.ToString());
 					}
-					FileToStorageBlock[File] = InputStorageBlock;
+					fileToStorageBlock[file] = inputStorageBlock;
 				}
 			}
 
+			// Run UAT
 			if (await ExecuteAutomationToolAsync(step, workspaceDir, arguments, logger, cancellationToken) != 0)
 			{
 				return false;
 			}
 
 			// Check that none of the inputs have been clobbered
-			Dictionary<string, string> ModifiedFiles = new Dictionary<string, string>(StringComparer.OrdinalIgnoreCase);
-			foreach (TempStorageFile File in InputManifests.Values.SelectMany(x => x.Files))
+			Dictionary<string, string> modifiedFiles = new Dictionary<string, string>(StringComparer.OrdinalIgnoreCase);
+			foreach (TempStorageFile file in inputManifests.Values.SelectMany(x => x.Files))
 			{
-				string? Message;
-				if (!ModifiedFiles.ContainsKey(File.RelativePath) && !File.Compare(workspaceDir, out Message))
+				string? message;
+				if (!modifiedFiles.ContainsKey(file.RelativePath) && !file.Compare(workspaceDir, out message))
 				{
-					ModifiedFiles.Add(File.RelativePath, Message);
+					modifiedFiles.Add(file.RelativePath, message);
 				}
 			}
-			if (ModifiedFiles.Count > 0)
+			if (modifiedFiles.Count > 0)
 			{
 				string modifiedFileList = "";
-				if (ModifiedFiles.Count < 100)
+				if (modifiedFiles.Count < 100)
 				{
-					modifiedFileList = String.Join("\n", ModifiedFiles.Select(x => x.Value));
+					modifiedFileList = String.Join("\n", modifiedFiles.Select(x => x.Value));
 				}
 				else
 				{
-					modifiedFileList = String.Join("\n", ModifiedFiles.Take(100).Select(x => x.Value));
-					modifiedFileList += $"{Environment.NewLine}And {ModifiedFiles.Count - 100} more.";
+					modifiedFileList = String.Join("\n", modifiedFiles.Take(100).Select(x => x.Value));
+					modifiedFileList += $"{Environment.NewLine}And {modifiedFiles.Count - 100} more.";
 				}
 
 				logger.LogError("Build product(s) from a previous step have been modified:\n{FileList}", modifiedFileList);
 				return false;
 			}
-			/*
+
 			// Determine all the output files which are required to be copied to temp storage (because they're referenced by nodes in another agent)
-			HashSet<FileReference> ReferencedOutputFiles = new HashSet<FileReference>();
-			foreach (BgAgentDef Agent in Graph.Agents)
+			HashSet<FileReference> referencedOutputFiles = new HashSet<FileReference>();
+			foreach (int publishOutput in step.PublishOutputs)
 			{
-				bool bSameAgent = Agent.Nodes.Contains(Node);
-				foreach (BgNodeDef OtherNode in Agent.Nodes)
-				{
-					if (!bSameAgent)
-					{
-						foreach (BgNodeOutput Input in OtherNode.Inputs.Where(x => x.ProducingNode == Node))
-						{
-							ReferencedOutputFiles.UnionWith(TagNameToFileSet[Input.TagName]);
-						}
-					}
-				}
+				string tagName = step.OutputNames[publishOutput];
+				referencedOutputFiles.UnionWith(tagNameToFileSet[tagName]);
 			}
-			*/
+
 			// Find a block name for all new outputs
-			Dictionary<FileReference, string> FileToOutputName = new Dictionary<FileReference, string>();
-			foreach (string outputName in step.OutputNames)
+			Dictionary<FileReference, string> fileToBlockName = new Dictionary<FileReference, string>();
+			for(int idx = 0; idx < step.OutputNames.Count; idx++)
 			{
-				string outputNameWithoutHash = outputName.TrimStart('#');
+				string tagName = step.OutputNames[idx];
+
+				string outputNameWithoutHash = tagName.TrimStart('#');
 				bool isDefaultOutput = outputNameWithoutHash.Equals(step.Name, StringComparison.OrdinalIgnoreCase);
 
-				HashSet<FileReference> Files = TagNameToFileSet[outputName];
-				foreach (FileReference File in Files)
+				HashSet<FileReference> files = tagNameToFileSet[tagName];
+				foreach (FileReference file in files)
 				{
-					if (!FileToStorageBlock.ContainsKey(File) && File.IsUnderDirectory(workspaceDir))
+					if (!fileToStorageBlock.ContainsKey(file) && file.IsUnderDirectory(workspaceDir))
 					{
 						if (isDefaultOutput)
 						{
-							if (!FileToOutputName.ContainsKey(File))
+							if (!fileToBlockName.ContainsKey(file))
 							{
-								FileToOutputName[File] = "";
+								fileToBlockName[file] = "";
 							}
 						}
 						else
 						{
-							string? OutputName;
-							if (FileToOutputName.TryGetValue(File, out OutputName) && OutputName.Length > 0)
+							string? blockName;
+							if (fileToBlockName.TryGetValue(file, out blockName) && blockName.Length > 0)
 							{
-								FileToOutputName[File] = String.Format("{0}+{1}", OutputName, outputNameWithoutHash);
+								fileToBlockName[file] = $"{blockName}+{outputNameWithoutHash}";
 							}
 							else
 							{
-								FileToOutputName[File] = outputNameWithoutHash;
+								fileToBlockName[file] = outputNameWithoutHash;
 							}
 						}
 					}
@@ -788,47 +782,47 @@ namespace Horde.Agent.Execution
 			}
 
 			// Invert the dictionary to make a mapping of storage block to the files each contains
-			Dictionary<string, HashSet<FileReference>> OutputStorageBlockToFiles = new Dictionary<string, HashSet<FileReference>>();
-			foreach (KeyValuePair<FileReference, string> Pair in FileToOutputName)
+			Dictionary<string, HashSet<FileReference>> outputStorageBlockToFiles = new Dictionary<string, HashSet<FileReference>>();
+			foreach (KeyValuePair<FileReference, string> pair in fileToBlockName)
 			{
-				HashSet<FileReference>? Files;
-				if (!OutputStorageBlockToFiles.TryGetValue(Pair.Value, out Files))
+				HashSet<FileReference>? files;
+				if (!outputStorageBlockToFiles.TryGetValue(pair.Value, out files))
 				{
-					Files = new HashSet<FileReference>();
-					OutputStorageBlockToFiles.Add(Pair.Value, Files);
+					files = new HashSet<FileReference>();
+					outputStorageBlockToFiles.Add(pair.Value, files);
 				}
-				Files.Add(Pair.Key);
+				files.Add(pair.Key);
 			}
 
 			// Write all the storage blocks, and update the mapping from file to storage block
 			using (GlobalTracer.Instance.BuildSpan("TempStorage").WithTag("resource", "Write").StartActive())
 			{
-				foreach (KeyValuePair<string, HashSet<FileReference>> Pair in OutputStorageBlockToFiles)
+				foreach (KeyValuePair<string, HashSet<FileReference>> pair in outputStorageBlockToFiles)
 				{
-					TempStorageBlock OutputBlock = new TempStorageBlock(step.Name, Pair.Key);
-					foreach (FileReference File in Pair.Value)
+					TempStorageBlock outputBlock = new TempStorageBlock(step.Name, pair.Key);
+					foreach (FileReference file in pair.Value)
 					{
-						FileToStorageBlock.Add(File, OutputBlock);
+						fileToStorageBlock.Add(file, outputBlock);
 					}
-					Storage.Archive(step.Name, Pair.Key, Pair.Value.ToArray(), true/*Pair.Value.Any(x => ReferencedOutputFiles.Contains(x))*/, logger);
+					storage.Archive(step.Name, pair.Key, pair.Value.ToArray(), pair.Value.Any(x => referencedOutputFiles.Contains(x)), logger);
 				}
 
 				// Publish all the output tags
 				foreach (string outputName in step.OutputNames)
 				{
-					HashSet<FileReference> Files = TagNameToFileSet[outputName];
+					HashSet<FileReference> files = tagNameToFileSet[outputName];
 
-					HashSet<TempStorageBlock> StorageBlocks = new HashSet<TempStorageBlock>();
-					foreach (FileReference File in Files)
+					HashSet<TempStorageBlock> storageBlocks = new HashSet<TempStorageBlock>();
+					foreach (FileReference file in files)
 					{
-						TempStorageBlock? StorageBlock;
-						if (FileToStorageBlock.TryGetValue(File, out StorageBlock))
+						TempStorageBlock? storageBlock;
+						if (fileToStorageBlock.TryGetValue(file, out storageBlock))
 						{
-							StorageBlocks.Add(StorageBlock);
+							storageBlocks.Add(storageBlock);
 						}
 					}
 
-					Storage.WriteFileList(step.Name, outputName, Files, StorageBlocks.ToArray(), logger);
+					storage.WriteFileList(step.Name, outputName, files, storageBlocks.ToArray(), logger);
 				}
 			}
 
diff --git a/Engine/Source/Programs/Horde/Horde.Build/Server/RpcService.cs b/Engine/Source/Programs/Horde/Horde.Build/Server/RpcService.cs
--- a/Engine/Source/Programs/Horde/Horde.Build/Server/RpcService.cs
+++ b/Engine/Source/Programs/Horde/Horde.Build/Server/RpcService.cs
@@ -13,6 +13,9 @@ using System.Text.Json;
 using System.Text.Json.Serialization;
 using System.Threading;
 using System.Threading.Tasks;
+using EpicGames.BuildGraph.Expressions;
+using EpicGames.BuildGraph;
+using System.Xml.Linq;
 using EpicGames.Core;
 using Google.Protobuf.WellKnownTypes;
 using Grpc.Core;
@@ -632,6 +635,20 @@ namespace Horde.Build.Server
 
 				response.OutputNames.Add(node.OutputNames);
 
+				foreach (INodeGroup otherGroup in graph.Groups)
+				{
+					if (otherGroup != graph.Groups[batch.GroupIdx])
+					{
+						foreach (NodeOutputRef outputRef in otherGroup.Nodes.SelectMany(x => x.Inputs))
+						{
+							if (outputRef.NodeRef.GroupIdx == batch.GroupIdx && outputRef.NodeRef.NodeIdx == step.NodeIdx && !response.PublishOutputs.Contains(outputRef.OutputIdx))
+							{
+								response.PublishOutputs.Add(outputRef.OutputIdx);
+							}
+						}
+					}
+				}
+
 				string templateName = "<unknown>";
 				if (job.TemplateHash != null)
 				{
