Horde: Read output manifests from local workspace after job completes.

diff --git a/Engine/Source/Programs/Horde/Horde.Agent/Execution/BuildGraphExecutor.cs b/Engine/Source/Programs/Horde/Horde.Agent/Execution/BuildGraphExecutor.cs
--- a/Engine/Source/Programs/Horde/Horde.Agent/Execution/BuildGraphExecutor.cs
+++ b/Engine/Source/Programs/Horde/Horde.Agent/Execution/BuildGraphExecutor.cs
@@ -709,6 +709,18 @@ namespace Horde.Agent.Execution
 				return false;
 			}
 
+			// Read all the output manifests
+			foreach (string tagName in step.OutputNames)
+			{
+				TempStorageFileList? fileList = storage.ReadFileList(step.Name, tagName, logger);
+				if (fileList == null)
+				{
+					logger.LogError("Unable to read local file list for output {Tag}", tagName);
+					return false;
+				}
+				tagNameToFileSet[tagName] = fileList.ToFileSet(workspaceDir);
+			}
+
 			// Check that none of the inputs have been clobbered
 			Dictionary<string, string> modifiedFiles = new Dictionary<string, string>(StringComparer.OrdinalIgnoreCase);
 			foreach (TempStorageFile file in inputManifests.Values.SelectMany(x => x.Files))
