Horde: Add controller test for setting pool sizing and fleet mangers

diff --git a/Engine/Source/Programs/Horde/Horde.Build.Tests/PoolsControllerTest.cs b/Engine/Source/Programs/Horde/Horde.Build.Tests/PoolsControllerTest.cs
--- a/Engine/Source/Programs/Horde/Horde.Build.Tests/PoolsControllerTest.cs
+++ b/Engine/Source/Programs/Horde/Horde.Build.Tests/PoolsControllerTest.cs
@@ -2,7 +2,6 @@
 
 using System.Collections.Generic;
 using System.Threading.Tasks;
-using Horde.Build.Agents;
 using Horde.Build.Agents.Fleet;
 using Horde.Build.Agents.Pools;
 using Horde.Build.Utilities;
@@ -36,7 +35,8 @@ namespace Horde.Build.Tests
 		        ScaleInCooldown = 222,
 		        SizeStrategies = new List<PoolSizeStrategyMessage> { new () { Type = PoolSizeStrategy.JobQueue, Condition = "dayOfWeek == 'monday'", Config = @"{""ScaleOutFactor"": 22.0, ""ScaleInFactor"": 33.0}" }},
 		        SizeStrategy = PoolSizeStrategy.JobQueue,
-		        JobQueueSettings = new JobQueueSettingsMessage(new JobQueueSettings(0.35, 0.85))
+		        JobQueueSettings = new JobQueueSettingsMessage(new JobQueueSettings(0.35, 0.85)),
+		        FleetManagers = new List<FleetManagerMessage> { new () { Type = FleetManagerType.AwsReuse, Condition = "dayOfWeek == 'monday'", Config = "{}" }},
 	        };
 	        
 	        ActionResult<CreatePoolResponse> result = await PoolsController.CreatePoolAsync(request);
@@ -52,6 +52,10 @@ namespace Horde.Build.Tests
 	        Assert.AreEqual(request.SizeStrategies[0].Type, pool.SizeStrategies[0].Type);
 	        Assert.AreEqual(request.SizeStrategies[0].Condition!.Text, pool.SizeStrategies[0].Condition!.Text);
 	        Assert.AreEqual(request.SizeStrategies[0].Config, pool.SizeStrategies[0].Config);
+	        Assert.AreEqual(1, request.FleetManagers.Count);
+	        Assert.AreEqual(request.FleetManagers[0].Type, pool.FleetManagers[0].Type);
+	        Assert.AreEqual(request.FleetManagers[0].Condition!.Text, pool.FleetManagers[0].Condition!.Text);
+	        Assert.AreEqual(request.FleetManagers[0].Config, pool.FleetManagers[0].Config);
         }
         
         [TestMethod]
@@ -64,7 +68,11 @@ namespace Horde.Build.Tests
 		        Name = "Pool1Modified",
 		        SizeStrategies = new List<PoolSizeStrategyMessage>
 		        {
-			        new () { Type = PoolSizeStrategy.JobQueue, Condition = "true", Config = "{\"someConfig\": 123"}
+			        new () { Type = PoolSizeStrategy.JobQueue, Condition = "true", Config = "{\"someConfig\": 123}"}
+		        },
+		        FleetManagers = new List<FleetManagerMessage>
+		        {
+			        new () { Type = FleetManagerType.AwsReuse, Condition = "true", Config = "{}"}
 		        },
 		        SizeStrategy = PoolSizeStrategy.JobQueue,
 		        JobQueueSettings = new JobQueueSettingsMessage { ScaleOutFactor = 25.0, ScaleInFactor = 0.3 }
@@ -80,7 +88,12 @@ namespace Horde.Build.Tests
 	        Assert.AreEqual(1, response.SizeStrategies.Count);
 	        Assert.AreEqual(PoolSizeStrategy.JobQueue, response.SizeStrategies[0].Type);
 	        Assert.AreEqual("true", response.SizeStrategies[0].Condition!.Text);
-	        Assert.AreEqual("{\"someConfig\": 123", response.SizeStrategies[0].Config);
+	        Assert.AreEqual("{\"someConfig\": 123}", response.SizeStrategies[0].Config);
+	        
+	        Assert.AreEqual(1, response.FleetManagers.Count);
+	        Assert.AreEqual(FleetManagerType.AwsReuse, response.FleetManagers[0].Type);
+	        Assert.AreEqual("true", response.FleetManagers[0].Condition!.Text);
+	        Assert.AreEqual("{}", response.FleetManagers[0].Config);
         }
     }
 }
\ No newline at end of file
