Horde: Add job copy to change button when available
#rnx
#skipci

diff --git a/Engine/Source/Programs/Horde/HordeDashboard/src/components/ChangeButton.tsx b/Engine/Source/Programs/Horde/HordeDashboard/src/components/ChangeButton.tsx
--- a/Engine/Source/Programs/Horde/HordeDashboard/src/components/ChangeButton.tsx
+++ b/Engine/Source/Programs/Horde/HordeDashboard/src/components/ChangeButton.tsx
@@ -31,15 +31,21 @@ export const ChangeContextMenu: React.FC<{ target: ChangeContextMenuTarget, onDi
       change = `PF ${job.preflightChange}`;
    }
 
-   const copyToClipboard = () => {
+   const copyToClipboard = (value: string | undefined) => {
+
+      if (!value) {
+         return;
+      }
+
       const el = document.createElement('textarea');
-      el.value = change.replace("PF", "");
+      el.value = value;
       document.body.appendChild(el);
       el.select();
       document.execCommand('copy');
       document.body.removeChild(el);
    }
 
+
    const menuItems: IContextualMenuItem[] = [];
 
    if (commit) {
@@ -92,9 +98,16 @@ export const ChangeContextMenu: React.FC<{ target: ChangeContextMenuTarget, onDi
    menuItems.push({
       key: 'copy_to_clipboard',
       text: 'Copy CL to Clipboard',
-      onClick: () => copyToClipboard(),
+      onClick: () => copyToClipboard(change.replace("PF ", ""))
+   });
+
+   if (job) {
+      menuItems.push({
+         key: 'copy_job_to_clipboard',
+         text: 'Copy Job to Clipboard',
+         onClick: () => copyToClipboard(window.location.protocol + "//" + window.location.hostname + `/job/${job.id}`)
+      });
    }
-   );
 
    menuItems.push({ key: 'open_in_swarm', text: 'Open CL in Swarm', onClick: (ev) => { window.open(`${dashboard.swarmUrl}/changes/${change.replace("PF ", "")}`) } })
 
