Horde: Support using bundles for archiving temp storage.

diff --git a/Engine/Source/Programs/Horde/Horde.Agent/Execution/BuildGraphExecutor.cs b/Engine/Source/Programs/Horde/Horde.Agent/Execution/BuildGraphExecutor.cs
--- a/Engine/Source/Programs/Horde/Horde.Agent/Execution/BuildGraphExecutor.cs
+++ b/Engine/Source/Programs/Horde/Horde.Agent/Execution/BuildGraphExecutor.cs
@@ -12,7 +12,6 @@ using System.Text.Json;
 using System.Text.Json.Serialization;
 using System.Threading;
 using System.Threading.Tasks;
-using Amazon.EC2.Model;
 using EpicGames.Core;
 using Grpc.Core;
 using Horde.Agent.Execution.Interfaces;
@@ -641,7 +640,7 @@ namespace Horde.Agent.Execution
 
 		private async Task<bool> ExecuteWithTempStorageAsync(BeginStepResponse step, DirectoryReference workspaceDir, DirectoryReference sharedStorageDir, string arguments, ILogger logger, CancellationToken cancellationToken)
 		{
-			TempStorage storage = new TempStorage(workspaceDir, DirectoryReference.Combine(workspaceDir, "Engine", "Saved", "BuildGraph"), sharedStorageDir, true);
+			TempStorage storage = new TempStorage(workspaceDir, DirectoryReference.Combine(workspaceDir, "Engine", "Saved", "BuildGraph"), sharedStorageDir);
 			logger.LogInformation("Using Horde-managed shared storage via {SharedStorageDir}", sharedStorageDir);
 
 			// Create the mapping of tag names to file sets
@@ -679,7 +678,7 @@ namespace Horde.Agent.Execution
 				scope.Span.SetTag("blocks", inputStorageBlocks.Count);
 				foreach (TempStorageBlock inputStorageBlock in inputStorageBlocks)
 				{
-					TempStorageManifest manifest = storage.Retrieve(inputStorageBlock.NodeName, inputStorageBlock.OutputName, logger);
+					TempStorageManifest manifest = await storage.RetrieveAsync(inputStorageBlock.NodeName, inputStorageBlock.OutputName, logger, cancellationToken);
 					inputManifests[inputStorageBlock] = manifest;
 				}
 				scope.Span.SetTag("size", inputManifests.Sum(x => x.Value.GetTotalSize()));
@@ -816,7 +815,10 @@ namespace Horde.Agent.Execution
 					{
 						fileToStorageBlock.Add(file, outputBlock);
 					}
-					storage.Archive(step.Name, pair.Key, pair.Value.ToArray(), pair.Value.Any(x => referencedOutputFiles.Contains(x)), logger);
+					if (pair.Value.Any(x => referencedOutputFiles.Contains(x)))
+					{
+						await storage.ArchiveAsync(step.Name, pair.Key, pair.Value.ToArray(), logger, cancellationToken);
+					}
 				}
 
 				// Publish all the output tags
diff --git a/Engine/Source/Programs/Horde/Horde.Agent/Utility/TempStorage.cs b/Engine/Source/Programs/Horde/Horde.Agent/Utility/TempStorage.cs
--- a/Engine/Source/Programs/Horde/Horde.Agent/Utility/TempStorage.cs
+++ b/Engine/Source/Programs/Horde/Horde.Agent/Utility/TempStorage.cs
@@ -1,16 +1,16 @@
 // Copyright Epic Games, Inc. All Rights Reserved.
 
 using EpicGames.Core;
+using EpicGames.Horde.Storage;
+using EpicGames.Horde.Storage.Backends;
+using EpicGames.Horde.Storage.Nodes;
+using Microsoft.Extensions.Caching.Memory;
 using Microsoft.Extensions.Logging;
-using OpenTracing;
-using OpenTracing.Util;
 using System;
-using System.Collections.Concurrent;
 using System.Collections.Generic;
 using System.Diagnostics;
 using System.Diagnostics.CodeAnalysis;
 using System.IO;
-using System.IO.Compression;
 using System.Linq;
 using System.Threading;
 using System.Threading.Tasks;
@@ -564,134 +564,17 @@ namespace Horde.Storage.Utility
 		/// </summary>
 		readonly DirectoryReference _sharedDir;
 
-		/// <summary>
-		/// Whether to allow writes to shared storage
-		/// </summary>
-		readonly bool _writeToSharedStorage;
-
 		/// <summary>
 		/// Constructor
 		/// </summary>
 		/// <param name="rootDir">Root directory for this branch</param>
 		/// <param name="localDir">The local temp storage directory.</param>
 		/// <param name="sharedDir">The shared temp storage directory. May be null.</param>
-		/// <param name="writeToSharedStorage">Whether to write to shared storage, or only permit reads from it</param>
-		public TempStorage(DirectoryReference rootDir, DirectoryReference localDir, DirectoryReference sharedDir, bool writeToSharedStorage)
+		public TempStorage(DirectoryReference rootDir, DirectoryReference localDir, DirectoryReference sharedDir)
 		{
 			_rootDir = rootDir;
 			_localDir = localDir;
 			_sharedDir = sharedDir;
-			_writeToSharedStorage = writeToSharedStorage;
-		}
-
-		/// <summary>
-		/// Mark the given node as complete
-		/// </summary>
-		/// <param name="nodeName">Name of the node</param>
-		public void MarkAsComplete(string nodeName)
-		{
-			// Create the marker locally
-			FileReference localFile = GetCompleteMarkerFile(_localDir, nodeName);
-			DirectoryReference.CreateDirectory(localFile.Directory);
-			File.OpenWrite(localFile.FullName).Close();
-
-			// Create the marker in the shared directory
-			if(_sharedDir != null && _writeToSharedStorage)
-			{
-				FileReference sharedFile = GetCompleteMarkerFile(_sharedDir, nodeName);
-				DirectoryReference.CreateDirectory(sharedFile.Directory);
-				File.OpenWrite(sharedFile.FullName).Close();
-			}
-		}
-
-		/// <summary>
-		/// Checks the integrity of the give node's local build products.
-		/// </summary>
-		/// <param name="nodeName">The node to retrieve build products for</param>
-		/// <param name="tagNames">List of tag names from this node.</param>
-		/// <param name="logger">Logger for output</param>
-		/// <returns>True if the node is complete and valid, false if not (and typically followed by a call to CleanNode()).</returns>
-		public bool CheckLocalIntegrity(string nodeName, IEnumerable<string> tagNames, ILogger logger)
-		{
-			// If the node is not locally complete, fail immediately.
-			FileReference completeMarkerFile = GetCompleteMarkerFile(_localDir, nodeName);
-			if(!FileReference.Exists(completeMarkerFile))
-			{
-				return false;
-			}
-
-			// Check that each of the tags exist
-			HashSet<TempStorageBlock> blocks = new HashSet<TempStorageBlock>();
-			foreach(string tagName in tagNames)
-			{
-				// Check the local manifest exists
-				FileReference LocalFileListLocation = GetTaggedFileListLocation(_localDir, nodeName, tagName);
-				if(!FileReference.Exists(LocalFileListLocation))
-				{
-					return false;
-				}
-
-				// Check the local manifest matches the shared manifest
-				if(_sharedDir != null)
-				{
-					// Check the shared manifest exists
-					FileReference sharedFileListLocation = GetManifestLocation(_sharedDir, nodeName, tagName);
-					if(!FileReference.Exists(sharedFileListLocation))
-					{
-						return false;
-					}
-
-					// Check the manifests are identical, byte by byte
-					byte[] localManifestBytes = FileReference.ReadAllBytes(LocalFileListLocation);
-					byte[] sharedManifestBytes = FileReference.ReadAllBytes(sharedFileListLocation);
-					if (!localManifestBytes.SequenceEqual(sharedManifestBytes))
-					{
-						return false;
-					}
-				}
-
-				// Read the manifest and add the referenced blocks to be checked
-				TempStorageFileList localFileList = TempStorageFileList.Load(LocalFileListLocation);
-				blocks.UnionWith(localFileList.Blocks);
-			}
-
-			// Check that each of the outputs match
-			foreach(TempStorageBlock block in blocks)
-			{
-				// Check the local manifest exists
-				FileReference localManifestFile = GetManifestLocation(_localDir, block.NodeName, block.OutputName);
-				if(!FileReference.Exists(localManifestFile))
-				{
-					return false;
-				}
-
-				// Check the local manifest matches the shared manifest
-				if(_sharedDir != null)
-				{
-					// Check the shared manifest exists
-					FileReference sharedManifestFile = GetManifestLocation(_sharedDir, block.NodeName, block.OutputName);
-					if(!FileReference.Exists(sharedManifestFile))
-					{
-						return false;
-					}
-
-					// Check the manifests are identical, byte by byte
-					byte[] localManifestBytes = FileReference.ReadAllBytes(localManifestFile);
-					byte[] sharedManifestBytes = FileReference.ReadAllBytes(sharedManifestFile);
-					if(!localManifestBytes.SequenceEqual(sharedManifestBytes))
-					{
-						return false;
-					}
-				}
-
-				// Read the manifest and check the files
-				TempStorageManifest localManifest = TempStorageManifest.Load(localManifestFile);
-				if(localManifest.Files.Any(x => !x.Compare(_rootDir, logger)))
-				{
-					return false;
-				}
-			}
-			return true;
 		}
 
 		/// <summary>
@@ -760,293 +643,142 @@ namespace Horde.Storage.Utility
 			TempStorageFileList fileList = new TempStorageFileList(files, _rootDir, blocks);
 
 			// Save the set of files to the local and shared locations
-			FileReference localFileListLocation = GetTaggedFileListLocation(_localDir, nodeName, tagName);
-			if(_sharedDir != null && _writeToSharedStorage)
+			FileReference sharedFileListLocation = GetTaggedFileListLocation(_sharedDir, nodeName, tagName);
+			try
 			{
-				FileReference sharedFileListLocation = GetTaggedFileListLocation(_sharedDir, nodeName, tagName);
-
-				try
-				{
-					logger.LogInformation("Saving file list to {File} and {OtherFile}", localFileListLocation.FullName, sharedFileListLocation.FullName);
-					DirectoryReference.CreateDirectory(sharedFileListLocation.Directory);
-					fileList.Save(sharedFileListLocation);
-				}
-				catch (Exception ex)
-				{
-					throw new TempStorageException($"Failed to save file list {localFileListLocation} to {sharedFileListLocation}, exception: {ex}");
-				}
+				logger.LogInformation("Saving file list to {File}", sharedFileListLocation.FullName);
+				DirectoryReference.CreateDirectory(sharedFileListLocation.Directory);
+				fileList.Save(sharedFileListLocation);
 			}
-			else
+			catch (Exception ex)
 			{
-				logger.LogInformation("Saving file list to {File}", localFileListLocation.FullName);
+				throw new TempStorageException($"Failed to save file list {sharedFileListLocation}, exception: {ex}");
 			}
+		}
 
-			// Save the local file list
-			DirectoryReference.CreateDirectory(localFileListLocation.Directory);
-			fileList.Save(localFileListLocation);
+		static DirectoryNode FindOrAddDirNode(DirectoryReference dirRef, Dictionary<DirectoryReference, DirectoryNode> dirRefToNode)
+		{
+			DirectoryNode? dirNode;
+			if (!dirRefToNode.TryGetValue(dirRef, out dirNode))
+			{
+				DirectoryNode parentNode = FindOrAddDirNode(dirRef.ParentDirectory!, dirRefToNode);
+				dirNode = parentNode.AddDirectory(dirRef.GetDirectoryName());
+				dirRefToNode.Add(dirRef, dirNode);
+			}
+			return dirNode;
 		}
 
 		/// <summary>
 		/// Saves the given files (that should be rooted at the branch root) to a shared temp storage manifest with the given temp storage node and game.
 		/// </summary>
-		/// <param name="NodeName">The node which created the storage block</param>
-		/// <param name="BlockName">Name of the block to retrieve. May be null or empty.</param>
-		/// <param name="BuildProducts">Array of build products to be archived</param>
-		/// <param name="bPushToRemote">Allow skipping the copying of this manifest to shared storage, because it's not required by any other agent</param>
+		/// <param name="nodeName">The node which created the storage block</param>
+		/// <param name="blockName">Name of the block to retrieve. May be null or empty.</param>
+		/// <param name="buildProducts">Array of build products to be archived</param>
 		/// <param name="logger">Logger for output</param>
+		/// <param name="cancellationToken"></param>
 		/// <returns>The created manifest instance (which has already been saved to disk).</returns>
-		public TempStorageManifest Archive(string NodeName, string BlockName, FileReference[] BuildProducts, bool bPushToRemote, ILogger logger)
+		public async Task ArchiveAsync(string nodeName, string blockName, FileReference[] buildProducts, ILogger logger, CancellationToken cancellationToken)
 		{
-			using (IScope scope = GlobalTracer.Instance.BuildSpan("StoreToTempStorage").StartActive())
-			{
-				// Create a manifest for the given build products
-				FileInfo[] files = BuildProducts.Select(x => new FileInfo(x.FullName)).ToArray();
-				TempStorageManifest manifest = new TempStorageManifest(files, _rootDir);
+			// Create a manifest for the given build products
+			FileInfo[] files = buildProducts.Select(x => new FileInfo(x.FullName)).ToArray();
+			TempStorageManifest manifest = new TempStorageManifest(files, _rootDir);
 
-				// Compress the files and copy to shared storage if necessary
-				bool remote = _sharedDir != null && bPushToRemote && _writeToSharedStorage;
-				if(remote)
-				{
-					// Create the shared directory for this node
-					FileReference sharedManifestFile = GetManifestLocation(_sharedDir!, NodeName, BlockName);
-					DirectoryReference.CreateDirectory(sharedManifestFile.Directory);
+			// Compress the files and copy to shared storage if necessary
+			DirectoryReference sharedNodeDir = DirectoryReference.Combine(_sharedDir!, nodeName);
 
-					// Zip all the build products
-					FileInfo[] zipFiles = ParallelZipFiles(files, _rootDir, sharedManifestFile.Directory, sharedManifestFile.GetFileNameWithoutExtension(), logger);
-					manifest.ZipFiles = zipFiles.Select(x => new TempStorageZipFile(x)).ToArray();
+			// Create the storage client
+			using MemoryCache cache = new MemoryCache(new MemoryCacheOptions { });
+			FileStorageClient store = new FileStorageClient(sharedNodeDir, cache, logger);
 
-					// Save the shared manifest
-					logger.LogInformation("Saving shared manifest to {File}", sharedManifestFile.FullName);
-					manifest.Save(sharedManifestFile);
-				}
+			// Create the directory tree
+			DirectoryNode root = new DirectoryNode(DirectoryFlags.None);
 
-				// Save the local manifest
-				FileReference localManifestFile = GetManifestLocation(_localDir, NodeName, BlockName);
-				logger.LogInformation("Saving local manifest to {File}", localManifestFile.FullName);
-				manifest.Save(localManifestFile);
+			Dictionary<DirectoryReference, DirectoryNode> dirRefToNode = new Dictionary<DirectoryReference, DirectoryNode>();
+			dirRefToNode.Add(_rootDir, root);
 
-				// Update the stats
-				long zipFilesTotalSize = (manifest.ZipFiles == null)? 0 : manifest.ZipFiles.Sum(x => x.Length);
-				scope.Span.SetTag("numFiles", files.Length);
-				scope.Span.SetTag("manifestSize", manifest.GetTotalSize());
-				scope.Span.SetTag("manifestZipFilesSize", zipFilesTotalSize);
-				scope.Span.SetTag("isRemote", remote);
-				scope.Span.SetTag("blockName", BlockName);
-				return manifest;
+			List<(DirectoryNode, FileInfo)> filesToAdd = new List<(DirectoryNode, FileInfo)>();
+
+			foreach (FileInfo file in files)
+			{
+				DirectoryReference dirRef = new FileReference(file).Directory;
+				DirectoryNode dirNode = FindOrAddDirNode(dirRef, dirRefToNode);
+				filesToAdd.Add((dirNode, file));
 			}
+
+			// Add all the files and flush the ref
+			ChunkingOptions options = new ChunkingOptions();
+			TreeWriter writer = new TreeWriter(store, new TreeOptions(), prefix: blockName);
+			await DirectoryNode.CopyFromDirectoryAsync(filesToAdd, options, writer, cancellationToken);
+			await writer.WriteRefAsync(new RefName(blockName), root, cancellationToken);
+
+			// Save the shared manifest
+			logger.LogInformation("Written bundle for {BlockName} to {NodeDir}", blockName, sharedNodeDir);
 		}
 
 		/// <summary>
 		/// Retrieve an output of the given node. Fetches and decompresses the files from shared storage if necessary, or validates the local files.
 		/// </summary>
 		/// <param name="nodeName">The node which created the storage block</param>
-		/// <param name="outputName">Name of the block to retrieve. May be null or empty.</param>
+		/// <param name="blockName">Name of the block to retrieve.</param>
 		/// <param name="logger">Logger for output</param>
+		/// <param name="cancellationToken"></param>
 		/// <returns>Manifest of the files retrieved</returns>
-		public TempStorageManifest Retrieve(string nodeName, string outputName, ILogger logger)
+		public async Task<TempStorageManifest> RetrieveAsync(string nodeName, string blockName, ILogger logger, CancellationToken cancellationToken)
 		{
-			using (IScope scope = GlobalTracer.Instance.BuildSpan("RetrieveFromTempStorage").StartActive())
+			// Get the path to the local manifest
+			FileReference localManifestFile = GetManifestLocation(_localDir, nodeName, blockName);
+			bool local = FileReference.Exists(localManifestFile);
+
+			// Read the manifest, either from local storage or shared storage
+			TempStorageManifest? manifest;
+			if(local)
+			{
+				logger.LogInformation("Reading shared manifest from {File}", localManifestFile.FullName);
+				manifest = TempStorageManifest.Load(localManifestFile);
+			}
+			else
 			{
-				// Get the path to the local manifest
-				FileReference localManifestFile = GetManifestLocation(_localDir, nodeName, outputName);
-				bool local = FileReference.Exists(localManifestFile);
+				// Get the shared directory for this node
+				FileReference sharedManifestFile = GetManifestLocation(_sharedDir, nodeName, blockName);
 
-				// Read the manifest, either from local storage or shared storage
-				TempStorageManifest? manifest = null;
-				if(local)
+				// Make sure the manifest exists
+				if(!FileReference.Exists(sharedManifestFile))
 				{
-					logger.LogInformation("Reading shared manifest from {File}", localManifestFile.FullName);
-					manifest = TempStorageManifest.Load(localManifestFile);
+					throw new TempStorageException($"Missing local or shared manifest for node - {sharedManifestFile.FullName}");
 				}
-				else
-				{
-					// Check we have shared storage
-					if(_sharedDir == null)
-					{
-						throw new TempStorageException($"Missing local manifest for node - {localManifestFile.FullName}");
-					}
 
-					// Get the shared directory for this node
-					FileReference sharedManifestFile = GetManifestLocation(_sharedDir, nodeName, outputName);
+				// Read the shared manifest
+				logger.LogInformation("Copying shared manifest from {Source} to {Target}", sharedManifestFile.FullName, localManifestFile.FullName);
+				manifest = TempStorageManifest.Load(sharedManifestFile);
 
-					// Make sure the manifest exists
-					if(!FileReference.Exists(sharedManifestFile))
-					{
-						throw new TempStorageException($"Missing local or shared manifest for node - {sharedManifestFile.FullName}");
-					}
+				// Compress the files and copy to shared storage if necessary
+				DirectoryReference sharedNodeDir = DirectoryReference.Combine(_sharedDir!, nodeName);
 
-					// Read the shared manifest
-					logger.LogInformation("Copying shared manifest from {Source} to {Target}", sharedManifestFile.FullName, localManifestFile.FullName);
-					manifest = TempStorageManifest.Load(sharedManifestFile);
-
-					// Unzip all the build products
-					DirectoryReference sharedNodeDir = GetDirectoryForNode(_sharedDir, nodeName);
-					FileInfo[] zipFiles = manifest!.ZipFiles.Select(x => new FileInfo(FileReference.Combine(sharedNodeDir, x.Name).FullName)).ToArray();
-					ParallelUnzipFiles(zipFiles, _rootDir);
-
-					// Update the timestamps to match the manifest. Zip files only use local time, and there's no guarantee it matches the local clock.
-					foreach(TempStorageFile ManifestFile in manifest.Files)
-					{
-						FileReference File = ManifestFile.ToFileReference(_rootDir);
-						System.IO.File.SetLastWriteTimeUtc(File.FullName, new DateTime(ManifestFile.LastWriteTimeUtcTicks, DateTimeKind.Utc));
-					}
-
-					// Save the manifest locally
-					DirectoryReference.CreateDirectory(localManifestFile.Directory);
-					manifest.Save(localManifestFile);
-				}
+				// Create the storage client
+				using MemoryCache cache = new MemoryCache(new MemoryCacheOptions { });
+				FileStorageClient store = new FileStorageClient(sharedNodeDir, cache, logger);
 
-				// Check all the local files are as expected
-				bool bAllMatch = true;
-				foreach(TempStorageFile File in manifest.Files)
-				{
-					bAllMatch &= File.Compare(_rootDir, logger);
-				}
-				if(!bAllMatch)
-				{
-					throw new TempStorageException("Files have been modified");
-				}
+				// Add all the files and flush the ref
+				ChunkingOptions options = new ChunkingOptions();
+				DirectoryNode directory = await store.ReadNodeAsync<DirectoryNode>(new RefName(blockName), cancellationToken: cancellationToken);
+				await directory.CopyToDirectoryAsync(_rootDir.ToDirectoryInfo(), logger, cancellationToken);
 
-				// Update the stats and return
-				scope.Span.SetTag("numFiles", manifest.Files.Length);
-				scope.Span.SetTag("manifestSize", manifest.Files.Sum(x => x.Length));
-				scope.Span.SetTag("manifestZipFilesSize", local? 0 : manifest.ZipFiles.Sum(x => x.Length));
-				scope.Span.SetTag("isRemote", !local);
-				scope.Span.SetTag("outputName", outputName);
-				return manifest;
+				// Save the manifest locally
+				DirectoryReference.CreateDirectory(localManifestFile.Directory);
+				manifest.Save(localManifestFile);
 			}
-		}
 
-		/// <summary>
-		/// Zips a set of files (that must be rooted at the given RootDir) to a set of zip files in the given OutputDir. The files will be prefixed with the given basename.
-		/// </summary>
-		/// <param name="InputFiles">Fully qualified list of files to zip (must be rooted at RootDir).</param>
-		/// <param name="RootDir">Root Directory where all files will be extracted.</param>
-		/// <param name="OutputDir">Location to place the set of zip files created.</param>
-		/// <param name="ZipBaseName">The basename of the set of zip files.</param>
-		/// <param name="logger">Logger for output</param>
-		/// <returns>Some metrics about the zip process.</returns>
-		/// <remarks>
-		/// This function tries to zip the files in parallel as fast as it can. It makes no guarantees about how many zip files will be created or which files will be in which zip,
-		/// but it does try to reasonably balance the file sizes.
-		/// </remarks>
-		static FileInfo[] ParallelZipFiles(FileInfo[] InputFiles, DirectoryReference RootDir, DirectoryReference OutputDir, string ZipBaseName, ILogger logger)
-		{
-			// First get the sizes of all the files. We won't parallelize if there isn't enough data to keep the number of zips down.
-			var FilesInfo = InputFiles
-				.Select(InputFile => new { File = new FileReference(InputFile), FileSize = InputFile.Length })
-				.ToList();
-
-			// Profiling results show that we can zip 100MB quite fast and it is not worth parallelizing that case and creating a bunch of zips that are relatively small.
-			const long MinFileSizeToZipInParallel = 1024 * 1024 * 100L;
-			bool bZipInParallel = FilesInfo.Sum(FileInfo => FileInfo.FileSize) >= MinFileSizeToZipInParallel;
-
-			// order the files in descending order so our threads pick up the biggest ones first.
-			// We want to end with the smaller files to more effectively fill in the gaps
-			ConcurrentQueue<FileReference> FilesToZip = new(FilesInfo.OrderByDescending(FileInfo => FileInfo.FileSize).Select(FileInfo => FileInfo.File));
-
-			ConcurrentBag<FileInfo> ZipFiles = new ConcurrentBag<FileInfo>();
-
-			DirectoryReference ZipDir = OutputDir;
-
-			// We deliberately avoid Parallel.ForEach here because profiles have shown that dynamic partitioning creates
-			// too many zip files, and they can be of too varying size, creating uneven work when unzipping later,
-			// as ZipFile cannot unzip files in parallel from a single archive.
-			// We can safely assume the build system will not be doing more important things at the same time, so we simply use all our logical cores,
-			// which has shown to be optimal via profiling, and limits the number of resulting zip files to the number of logical cores.
-			List<Thread> ZipThreads = (
-				from CoreNum in Enumerable.Range(0, bZipInParallel ? Environment.ProcessorCount : 1)
-				select new Thread((object? indexObject) =>
-				{
-					int index = (int)indexObject!;
-					FileReference ZipFileName = FileReference.Combine(ZipDir, String.Format("{0}{1}.zip", ZipBaseName, bZipInParallel ? "-" + index.ToString("00") : ""));
-					// don't create the zip unless we have at least one file to add
-					FileReference? File;
-					if (FilesToZip.TryDequeue(out File))
-					{
-						try
-						{
-							// Create one zip per thread using the given basename
-							using (ZipArchive ZipArchive = ZipFile.Open(ZipFileName.FullName, ZipArchiveMode.Create))
-							{
-								// pull from the queue until we are out of files.
-								do
-								{
-									// use fastest compression. In our best case we are CPU bound, so this is a good tradeoff,
-									// cutting overall time by 2/3 while only modestly increasing the compression ratio (22.7% -> 23.8% for RootEditor PDBs).
-									// This is in cases of a super hot cache, so the operation was largely CPU bound.
-									ZipArchiveExtensions.CreateEntryFromFile_CrossPlatform(ZipArchive, File.FullName, File.MakeRelativeTo(RootDir).Replace(Path.DirectorySeparatorChar, '/'), CompressionLevel.Fastest);
-								} while (FilesToZip.TryDequeue(out File));
-							}
-						}
-						catch (IOException ex)
-						{
-							logger.LogError(ex, "Unable to open file for TempStorage zip: \"{File}\"", ZipFileName);
-							throw new TempStorageException($"Unable to open file {ZipFileName.FullName}");
-						}
-
-						ZipFiles.Add(new FileInfo(ZipFileName.FullName));
-					}
-				})).ToList();
-
-			for (int index = 0; index < ZipThreads.Count; index++)
+			// Check all the local files are as expected
+			bool bAllMatch = true;
+			foreach(TempStorageFile File in manifest.Files)
 			{
-				Thread thread = ZipThreads[index];
-				thread.Start(index);
+				bAllMatch &= File.Compare(_rootDir, logger);
 			}
-
-			ZipThreads.ForEach(thread => thread.Join());
-
-			return ZipFiles.OrderBy(x => x.Name).ToArray();
-		}
-
-		/// <summary>
-		/// Unzips a set of zip files with a given basename in a given folder to a given RootDir.
-		/// </summary>
-		/// <param name="zipFiles">Files to extract</param>
-		/// <param name="rootDir">Root Directory where all files will be extracted.</param>
-		/// <returns>Some metrics about the unzip process.</returns>
-		/// <remarks>
-		/// The code is expected to be the used as the symmetrical inverse of <see cref="ParallelZipFiles"/>, but could be used independently, as long as the files in the zip do not overlap.
-		/// </remarks>
-		private static void ParallelUnzipFiles(FileInfo[] zipFiles, DirectoryReference rootDir)
-		{
-			Parallel.ForEach(zipFiles,
-				(zipFile) =>
-				{
-					string localZipFile = zipFile.FullName;
-
-					// unzip the files manually instead of caling ZipFile.ExtractToDirectory() because we need to overwrite readonly files. Because of this, creating the directories is up to us as well.
-					using (ZipArchive zipArchive = ZipFile.OpenRead(localZipFile))
-					{
-						foreach (ZipArchiveEntry entry in zipArchive.Entries)
-						{
-							// We must delete any existing file, even if it's readonly. .Net does not do this by default.
-							FileReference extractedFileName = FileReference.Combine(rootDir, entry.FullName);
-							if (FileReference.Exists(extractedFileName))
-							{
-								FileUtils.ForceDeleteFile(extractedFileName);
-							}
-							else
-							{
-								DirectoryReference.CreateDirectory(extractedFileName.Directory!);
-							}
-
-							entry.ExtractToFile_CrossPlatform(extractedFileName.FullName, true);
-						}
-					}
-				});
-		}
-
-		/// <summary>
-		/// Gets the directory used to store data for the given node
-		/// </summary>
-		/// <param name="baseDir">A local or shared temp storage root directory.</param>
-		/// <param name="nodeName">Name of the node</param>
-		/// <returns>Directory to contain a node's data</returns>
-		static DirectoryReference GetDirectoryForNode(DirectoryReference baseDir, string nodeName)
-		{
-			return DirectoryReference.Combine(baseDir, nodeName);
+			if(!bAllMatch)
+			{
+				throw new TempStorageException("Files have been modified");
+			}
+			return manifest;
 		}
 
 		/// <summary>
@@ -1072,16 +804,6 @@ namespace Horde.Storage.Utility
 			return FileReference.Combine(baseDir, nodeName, String.Format("Tag-{0}.xml", tagName.Substring(1)));
 		}
 
-		/// <summary>
-		/// Gets the path to a file created to indicate that a node is complete, under the given base directory.
-		/// </summary>
-		/// <param name="baseDir">A local or shared temp storage root directory.</param>
-		/// <param name="nodeName">Name of the node to get the file for</param>
-		static FileReference GetCompleteMarkerFile(DirectoryReference baseDir, string nodeName)
-		{
-			return FileReference.Combine(GetDirectoryForNode(baseDir, nodeName), "Complete");
-		}
-
 		/// <summary>
 		/// Checks whether the given path is allowed as a build product that can be produced by more than one node (timestamps may be modified, etc..). Used to suppress
 		/// warnings about build products being overwritten.
