Horde: Sanitize ref names before writing blocks to temp storage.

diff --git a/Engine/Source/Programs/Horde/Horde.Agent/Utility/TempStorage.cs b/Engine/Source/Programs/Horde/Horde.Agent/Utility/TempStorage.cs
--- a/Engine/Source/Programs/Horde/Horde.Agent/Utility/TempStorage.cs
+++ b/Engine/Source/Programs/Horde/Horde.Agent/Utility/TempStorage.cs
@@ -12,6 +12,7 @@ using System.Diagnostics;
 using System.Diagnostics.CodeAnalysis;
 using System.IO;
 using System.Linq;
+using System.Text;
 using System.Threading;
 using System.Threading.Tasks;
 using System.Xml;
@@ -668,6 +669,33 @@ namespace Horde.Storage.Utility
 			return dirNode;
 		}
 
+		/// <summary>
+		/// Gets a ref name from the given node and block name
+		/// </summary>
+		static RefName GetRefName(string nodeName, string blockName)
+		{
+			string inputName = $"{nodeName}/{blockName}".TrimEnd('/');
+
+			StringBuilder result = new StringBuilder();
+			for (int idx = 0; idx < inputName.Length; idx++)
+			{
+				if (inputName[idx] >= 'A' && inputName[idx] <= 'Z')
+				{
+					result.Append((char)(inputName[idx] + 'a' - 'A'));
+				}
+				else if ((inputName[idx] >= 'a' && inputName[idx] <= 'z') || inputName[idx] == '+' || inputName[idx] == '/')
+				{
+					result.Append(inputName[idx]);
+				}
+				else if (result.Length > 0 && result[result.Length - 1] != '-')
+				{
+					result.Append('-');
+				}
+			}
+
+			return new RefName(result.ToString());
+		}
+
 		/// <summary>
 		/// Saves the given files (that should be rooted at the branch root) to a shared temp storage manifest with the given temp storage node and game.
 		/// </summary>
@@ -709,10 +737,12 @@ namespace Horde.Storage.Utility
 			ChunkingOptions options = new ChunkingOptions();
 			TreeWriter writer = new TreeWriter(store, new TreeOptions(), prefix: blockName);
 			await DirectoryNode.CopyFromDirectoryAsync(filesToAdd, options, writer, cancellationToken);
-			await writer.WriteRefAsync(new RefName(blockName), root, cancellationToken);
+
+			RefName refName = GetRefName(nodeName, blockName);
+			await writer.WriteRefAsync(refName, root, cancellationToken);
 
 			// Save the shared manifest
-			logger.LogInformation("Written bundle for {BlockName} to {NodeDir}", blockName, sharedNodeDir);
+			logger.LogInformation("Written {RefName} to {NodeDir}", refName, sharedNodeDir);
 		}
 
 		/// <summary>
@@ -759,8 +789,8 @@ namespace Horde.Storage.Utility
 				FileStorageClient store = new FileStorageClient(sharedNodeDir, cache, logger);
 
 				// Add all the files and flush the ref
-				ChunkingOptions options = new ChunkingOptions();
-				DirectoryNode directory = await store.ReadNodeAsync<DirectoryNode>(new RefName(blockName), cancellationToken: cancellationToken);
+				RefName refName = GetRefName(nodeName, blockName);
+				DirectoryNode directory = await store.ReadNodeAsync<DirectoryNode>(refName, cancellationToken: cancellationToken);
 				await directory.CopyToDirectoryAsync(_rootDir.ToDirectoryInfo(), logger, cancellationToken);
 
 				// Save the manifest locally
