Horde: Fix temp storage manifest not being written to shared drive.

diff --git a/Engine/Source/Programs/Horde/Horde.Agent/Utility/TempStorage.cs b/Engine/Source/Programs/Horde/Horde.Agent/Utility/TempStorage.cs
--- a/Engine/Source/Programs/Horde/Horde.Agent/Utility/TempStorage.cs
+++ b/Engine/Source/Programs/Horde/Horde.Agent/Utility/TempStorage.cs
@@ -693,6 +693,11 @@ namespace Horde.Storage.Utility
 				}
 			}
 
+			if (result.Length > 0 && result[result.Length - 1] == '-')
+			{
+				result.Remove(result.Length - 1, 1);
+			}
+
 			return new RefName(result.ToString());
 		}
 
@@ -709,14 +714,13 @@ namespace Horde.Storage.Utility
 		{
 			// Create a manifest for the given build products
 			FileInfo[] files = buildProducts.Select(x => new FileInfo(x.FullName)).ToArray();
-			TempStorageManifest manifest = new TempStorageManifest(files, _rootDir);
 
 			// Compress the files and copy to shared storage if necessary
-			DirectoryReference sharedNodeDir = DirectoryReference.Combine(_sharedDir!, nodeName);
+			DirectoryReference sharedNodeDir = DirectoryReference.Combine(_sharedDir!, "data");
 
 			// Create the storage client
 			using MemoryCache cache = new MemoryCache(new MemoryCacheOptions { });
-			FileStorageClient store = new FileStorageClient(sharedNodeDir, cache, logger);
+			FileStorageClient store = new FileStorageClient(_sharedDir!, cache, logger);
 
 			// Create the directory tree
 			DirectoryNode root = new DirectoryNode(DirectoryFlags.None);
@@ -734,15 +738,26 @@ namespace Horde.Storage.Utility
 			}
 
 			// Add all the files and flush the ref
-			ChunkingOptions options = new ChunkingOptions();
-			TreeWriter writer = new TreeWriter(store, new TreeOptions(), prefix: blockName);
-			await DirectoryNode.CopyFromDirectoryAsync(filesToAdd, options, writer, cancellationToken);
-
 			RefName refName = GetRefName(nodeName, blockName);
+
+			TreeOptions treeOptions = new TreeOptions();
+			TreeWriter writer = new TreeWriter(store, treeOptions, prefix: refName.Text);
+
+			ChunkingOptions chunkingOptions = new ChunkingOptions();
+			await DirectoryNode.CopyFromDirectoryAsync(filesToAdd, chunkingOptions, writer, cancellationToken);
+
 			await writer.WriteRefAsync(refName, root, cancellationToken);
 
 			// Save the shared manifest
 			logger.LogInformation("Written {RefName} to {NodeDir}", refName, sharedNodeDir);
+
+			// Create the shared directory for this node
+			FileReference sharedManifestFile = GetManifestLocation(_sharedDir!, nodeName, blockName);
+			logger.LogInformation("Saving shared manifest to {File}", sharedManifestFile.FullName);
+
+			DirectoryReference.CreateDirectory(sharedManifestFile.Directory);
+			TempStorageManifest manifest = new TempStorageManifest(files, _rootDir);
+			manifest.Save(sharedManifestFile);
 		}
 
 		/// <summary>
