Horde: Allow numbers in temp storage ref names.

diff --git a/Engine/Source/Programs/Horde/Horde.Agent/Utility/TempStorage.cs b/Engine/Source/Programs/Horde/Horde.Agent/Utility/TempStorage.cs
--- a/Engine/Source/Programs/Horde/Horde.Agent/Utility/TempStorage.cs
+++ b/Engine/Source/Programs/Horde/Horde.Agent/Utility/TempStorage.cs
@@ -683,7 +683,7 @@ namespace Horde.Storage.Utility
 				{
 					result.Append((char)(inputName[idx] + 'a' - 'A'));
 				}
-				else if ((inputName[idx] >= 'a' && inputName[idx] <= 'z') || inputName[idx] == '+' || inputName[idx] == '/')
+				else if ((inputName[idx] >= 'a' && inputName[idx] <= 'z') || (inputName[idx] >= '0' && inputName[idx] <= '9') || inputName[idx] == '+' || inputName[idx] == '/')
 				{
 					result.Append(inputName[idx]);
 				}
