Horde: Fix incorrect directory for reading temp storage.

diff --git a/Engine/Source/Programs/Horde/Horde.Agent/Utility/TempStorage.cs b/Engine/Source/Programs/Horde/Horde.Agent/Utility/TempStorage.cs
--- a/Engine/Source/Programs/Horde/Horde.Agent/Utility/TempStorage.cs
+++ b/Engine/Source/Programs/Horde/Horde.Agent/Utility/TempStorage.cs
@@ -701,6 +701,8 @@ namespace Horde.Storage.Utility
 			return new RefName(result.ToString());
 		}
 
+		DirectoryReference GetBundleDir() => DirectoryReference.Combine(_sharedDir!, "data");
+
 		/// <summary>
 		/// Saves the given files (that should be rooted at the branch root) to a shared temp storage manifest with the given temp storage node and game.
 		/// </summary>
@@ -716,7 +718,7 @@ namespace Horde.Storage.Utility
 			FileInfo[] files = buildProducts.Select(x => new FileInfo(x.FullName)).ToArray();
 
 			// Compress the files and copy to shared storage if necessary
-			DirectoryReference sharedNodeDir = DirectoryReference.Combine(_sharedDir!, "data");
+			DirectoryReference bundleDir = GetBundleDir();
 
 			// Create the storage client
 			using MemoryCache cache = new MemoryCache(new MemoryCacheOptions { });
@@ -749,7 +751,7 @@ namespace Horde.Storage.Utility
 			await writer.WriteRefAsync(refName, root, cancellationToken);
 
 			// Save the shared manifest
-			logger.LogInformation("Written {RefName} to {NodeDir}", refName, sharedNodeDir);
+			logger.LogInformation("Written {RefName} to {NodeDir}", refName, bundleDir);
 
 			// Create the shared directory for this node
 			FileReference sharedManifestFile = GetManifestLocation(_sharedDir!, nodeName, blockName);
@@ -797,11 +799,11 @@ namespace Horde.Storage.Utility
 				manifest = TempStorageManifest.Load(sharedManifestFile);
 
 				// Compress the files and copy to shared storage if necessary
-				DirectoryReference sharedNodeDir = DirectoryReference.Combine(_sharedDir!, nodeName);
+				DirectoryReference bundleDir = GetBundleDir();
 
 				// Create the storage client
 				using MemoryCache cache = new MemoryCache(new MemoryCacheOptions { });
-				FileStorageClient store = new FileStorageClient(sharedNodeDir, cache, logger);
+				FileStorageClient store = new FileStorageClient(bundleDir, cache, logger);
 
 				// Add all the files and flush the ref
 				RefName refName = GetRefName(nodeName, blockName);
