Horde: Fix temp storage directory.

diff --git a/Engine/Source/Programs/Horde/Horde.Agent/Utility/TempStorage.cs b/Engine/Source/Programs/Horde/Horde.Agent/Utility/TempStorage.cs
--- a/Engine/Source/Programs/Horde/Horde.Agent/Utility/TempStorage.cs
+++ b/Engine/Source/Programs/Horde/Horde.Agent/Utility/TempStorage.cs
@@ -722,7 +722,7 @@ namespace Horde.Storage.Utility
 
 			// Create the storage client
 			using MemoryCache cache = new MemoryCache(new MemoryCacheOptions { });
-			FileStorageClient store = new FileStorageClient(_sharedDir!, cache, logger);
+			FileStorageClient store = new FileStorageClient(bundleDir, cache, logger);
 
 			// Create the directory tree
 			DirectoryNode root = new DirectoryNode(DirectoryFlags.None);
