Horde: Set last write timestamps of files extracted from temp storage back to their timestamps on the originating machine.

diff --git a/Engine/Source/Programs/Horde/Horde.Agent/Utility/TempStorage.cs b/Engine/Source/Programs/Horde/Horde.Agent/Utility/TempStorage.cs
--- a/Engine/Source/Programs/Horde/Horde.Agent/Utility/TempStorage.cs
+++ b/Engine/Source/Programs/Horde/Horde.Agent/Utility/TempStorage.cs
@@ -810,6 +810,13 @@ namespace Horde.Storage.Utility
 				DirectoryNode directory = await store.ReadNodeAsync<DirectoryNode>(refName, cancellationToken: cancellationToken);
 				await directory.CopyToDirectoryAsync(_rootDir.ToDirectoryInfo(), logger, cancellationToken);
 
+				// Update the timestamps to match the manifest.
+				foreach (TempStorageFile ManifestFile in manifest.Files)
+				{
+					FileReference File = ManifestFile.ToFileReference(_rootDir);
+					System.IO.File.SetLastWriteTimeUtc(File.FullName, new DateTime(ManifestFile.LastWriteTimeUtcTicks, DateTimeKind.Utc));
+				}
+
 				// Save the manifest locally
 				DirectoryReference.CreateDirectory(localManifestFile.Directory);
 				manifest.Save(localManifestFile);
