Horde: Ensure files are writable before extracting temp storage into workspace.

diff --git a/Engine/Source/Programs/Horde/Horde.Agent/Utility/TempStorage.cs b/Engine/Source/Programs/Horde/Horde.Agent/Utility/TempStorage.cs
--- a/Engine/Source/Programs/Horde/Horde.Agent/Utility/TempStorage.cs
+++ b/Engine/Source/Programs/Horde/Horde.Agent/Utility/TempStorage.cs
@@ -798,6 +798,13 @@ namespace Horde.Storage.Utility
 				logger.LogInformation("Copying shared manifest from {Source} to {Target}", sharedManifestFile.FullName, localManifestFile.FullName);
 				manifest = TempStorageManifest.Load(sharedManifestFile);
 
+				// Delete all the existing files. They may be read-only.
+				foreach (TempStorageFile ManifestFile in manifest.Files)
+				{
+					FileReference File = ManifestFile.ToFileReference(_rootDir);
+					FileUtils.ForceDeleteFile(File);
+				}
+
 				// Compress the files and copy to shared storage if necessary
 				DirectoryReference bundleDir = GetBundleDir();
 
