Horde: Add simple VCS prototype using Horde storage library.

* Core data model is Git-like, with a merkle tree of commits and filesystem objects, but usage is designed around having a centralized, shared repository used by everyone.
* Commits are numbered to allow ordering for versioning (and potentially - in the future - allow the possibility of commit rewriting and per-user permissions for paths without relying on a universally shared hash to identify each change).
* Multiple branches are supported, but commit history in each is linear and does not currently contain any annotations for merges. (Thinking that merges are probably best implemented as weak references, to prevent sprawling history and to allow coarse grained security along branch boundaries.)
* File data is chunked by slicing along boundaries defined by a rolling hash, to more easily support large binaries. Chunks are hashed for deduplication using IoHash and packed together into bundles of a bounded size.
* Supports the normal array of storage backends (filesystem, HTTP, etc...) as IStorageClient, thus supporting use against a centralized server or local directory. Data can be pulled directly from the storage backend rather than requiring a full copy locally.
* Added horde.bat batch file to run HordeAgent, to allow using from command line as "horde vcs init" etc...

A small subset of commands are currently supported:
* horde vcs init (Initialize a repository in the current directory)
* horde vcs commit -message="..." (Commit the current directory state to the repository)
* horde vcs status (Show the current repository status)
* horde vcs checkout -change=N [-clean] [-force] (Checkout the repository at the given change number).
* horde vcs branch -name=XYZ (Create a new, empty branch. Leaves the current working directory intact.)
* horde vcs log (Shows the history of commits)

diff --git a/Engine/Source/Programs/Horde/Horde.Agent/Commands/Vcs/VcsCommand.cs b/Engine/Source/Programs/Horde/Horde.Agent/Commands/Vcs/VcsCommand.cs
new file mode 100644
--- /dev/null
+++ b/Engine/Source/Programs/Horde/Horde.Agent/Commands/Vcs/VcsCommand.cs
@@ -0,0 +1,772 @@
+// Copyright Epic Games, Inc. All Rights Reserved.
+
+using System;
+using System.Collections.Generic;
+using System.Diagnostics;
+using System.IO;
+using System.Linq;
+using System.Threading;
+using System.Threading.Tasks;
+using EpicGames.Core;
+using EpicGames.Horde.Storage;
+using EpicGames.Horde.Storage.Nodes;
+using Microsoft.CodeAnalysis;
+using Microsoft.Extensions.Logging;
+
+namespace Horde.Agent.Commands.Bundles
+{
+	abstract class VcsCommand : BundleCommandBase
+	{
+		[CommandLine("-BaseDir=")]
+		public DirectoryReference? BaseDir { get; set; }
+
+		protected class CommitNode : TreeNode
+		{
+			public int Number { get; set; }
+			public TreeNodeRef<CommitNode>? Parent { get; set; }
+			public string Message { get; set; }
+			public TreeNodeRef<DirectoryNode> Contents { get; set; }
+
+			public CommitNode(ITreeNodeReader reader)
+			{
+				Number = reader.ReadInt32();
+				if (reader.ReadBoolean())
+				{
+					Parent = reader.ReadRef<CommitNode>();
+				}
+				Message = reader.ReadString();
+				Contents = reader.ReadRef<DirectoryNode>();
+			}
+
+			public CommitNode(int number, TreeNodeRef<CommitNode>? parent, string message, TreeNodeRef<DirectoryNode> contents)
+			{
+				Number = number;
+				Parent = parent;
+				Message = message;
+				Contents = contents;
+			}
+
+			public override void Serialize(ITreeNodeWriter writer)
+			{
+				writer.WriteInt32(Number);
+				writer.WriteBoolean(Parent != null);
+				if (Parent != null)
+				{
+					writer.WriteRef(Parent);
+				}
+				writer.WriteString(Message);
+				writer.WriteRef(Contents);
+			}
+
+			public override IEnumerable<TreeNodeRef> EnumerateRefs()
+			{
+				if (Parent != null)
+				{
+					yield return Parent;
+				}
+				yield return Contents;
+			}
+		}
+
+		protected class FileState
+		{
+			public IoHash Hash { get; set; }
+			public long Length { get; private set; }
+			public long LastModifiedTimeUtc { get; private set; }
+
+			public FileState(FileInfo fileInfo, IoHash hash = default)
+			{
+				Update(fileInfo);
+				Hash = hash;
+			}
+
+			public FileState(FileState other)
+			{
+				Length = other.Length;
+				LastModifiedTimeUtc = other.LastModifiedTimeUtc;
+			}
+
+			public FileState(IMemoryReader reader)
+			{
+				Hash = reader.ReadIoHash();
+				Length = reader.ReadInt64();
+				LastModifiedTimeUtc = reader.ReadInt64();
+			}
+
+			public bool Modified(FileInfo fileInfo) => Length != fileInfo.Length || LastModifiedTimeUtc != fileInfo.LastWriteTimeUtc.Ticks;
+
+			public void Update(FileInfo fileInfo)
+			{
+				Length = fileInfo.Length;
+				LastModifiedTimeUtc = fileInfo.LastWriteTimeUtc.Ticks;
+			}
+
+			public void Write(IMemoryWriter writer)
+			{
+				Debug.Assert(Hash != IoHash.Zero);
+				writer.WriteIoHash(Hash);
+				writer.WriteInt64(Length);
+				writer.WriteInt64(LastModifiedTimeUtc);
+			}
+		}
+
+		protected class DirectoryState
+		{
+			public IoHash Hash { get; set; }
+			public Dictionary<Utf8String, DirectoryState> Directories { get; } = new Dictionary<Utf8String, DirectoryState>(Utf8StringComparer.OrdinalIgnoreCase);
+			public Dictionary<Utf8String, FileState> Files { get; } = new Dictionary<Utf8String, FileState>(Utf8StringComparer.OrdinalIgnoreCase);
+
+			public DirectoryState()
+			{
+				Directories = new Dictionary<Utf8String, DirectoryState>(Utf8StringComparer.OrdinalIgnoreCase);
+				Files = new Dictionary<Utf8String, FileState>(Utf8StringComparer.OrdinalIgnoreCase);
+			}
+
+			public DirectoryState(IMemoryReader reader)
+			{
+				Hash = reader.ReadIoHash();
+
+				int numDirectories = reader.ReadInt32();
+				Directories = new Dictionary<Utf8String, DirectoryState>(numDirectories, Utf8StringComparer.OrdinalIgnoreCase);
+
+				for (int idx = 0; idx < numDirectories; idx++)
+				{
+					Utf8String name = reader.ReadUtf8String();
+					Directories[name] = new DirectoryState(reader);
+				}
+
+				int numFiles = reader.ReadInt32();
+				Files = new Dictionary<Utf8String, FileState>(numFiles, Utf8StringComparer.OrdinalIgnoreCase);
+
+				for (int idx = 0; idx < numFiles; idx++)
+				{
+					Utf8String name = reader.ReadUtf8String();
+					Files[name] = new FileState(reader);
+				}
+			}
+
+			public void Write(IMemoryWriter writer)
+			{
+				writer.WriteIoHash(Hash);
+
+				writer.WriteInt32(Directories.Count);
+				foreach ((Utf8String name, DirectoryState state) in Directories)
+				{
+					writer.WriteUtf8String(name);
+					state.Write(writer);
+				}
+
+				writer.WriteInt32(Files.Count);
+				foreach ((Utf8String name, FileState state) in Files)
+				{
+					writer.WriteUtf8String(name);
+					state.Write(writer);
+				}
+			}
+		}
+
+		protected class WorkspaceState
+		{
+			public RefName Branch { get; set; }
+			public int Change { get; set; }
+			public DirectoryState Tree { get; set; }
+
+			public WorkspaceState() 
+				: this(new RefName("main"), 0, new DirectoryState())
+			{
+			}
+
+			public WorkspaceState(RefName branch, int change, DirectoryState tree)
+			{
+				Branch = branch;
+				Change = change;
+				Tree = tree;
+			}
+
+			public WorkspaceState(IMemoryReader reader)
+			{
+				Branch = new RefName(reader.ReadUtf8String());
+				Change = reader.ReadInt32();
+				Tree = new DirectoryState(reader);
+			}
+
+			public void Write(IMemoryWriter writer)
+			{
+				writer.WriteUtf8String(Branch.Text);
+				writer.WriteInt32(Change);
+				Tree.Write(writer);
+			}
+		}
+
+		public static IEnumerable<(TKey, TOldValue?, TNewValue?)> Zip<TKey, TOldValue, TNewValue>(IReadOnlyDictionary<TKey, TOldValue>? oldDictionary, IReadOnlyDictionary<TKey, TNewValue>? newDictionary)
+			where TKey : notnull 
+			where TOldValue : class 
+			where TNewValue : class
+		{
+			if (newDictionary != null)
+			{
+				foreach ((TKey key, TNewValue newValue) in newDictionary)
+				{
+					TOldValue? oldValue;
+					if (oldDictionary == null || !oldDictionary.TryGetValue(key, out oldValue))
+					{
+						yield return (key, null, newValue);
+					}
+					else
+					{
+						yield return (key, oldValue, newValue);
+					}
+				}
+			}
+			if (oldDictionary != null)
+			{
+				foreach ((TKey key, TOldValue oldValue) in oldDictionary)
+				{
+					if (newDictionary == null || !newDictionary.ContainsKey(key))
+					{
+						yield return (key, oldValue, null);
+					}
+				}
+			}
+		}
+
+		const string DataDir = ".horde";
+		const string StateFileName = "state.dat";
+
+		protected IStorageClientOwner CreateStorageClient(DirectoryReference rootDir, ILogger logger)
+		{
+			StorageDir = DirectoryReference.Combine(rootDir, DataDir, "bundles");
+			return base.CreateStorageClient(logger);
+		}
+
+		protected static async Task<WorkspaceState> ReadStateAsync(DirectoryReference rootDir)
+		{
+			FileReference stateFile = FileReference.Combine(rootDir, DataDir, StateFileName);
+			byte[] data = await FileReference.ReadAllBytesAsync(stateFile);
+			return new WorkspaceState(new MemoryReader(data));
+		}
+
+		protected static async Task WriteStateAsync(DirectoryReference rootDir, WorkspaceState state)
+		{
+			ByteArrayBuilder builder = new ByteArrayBuilder();
+			state.Write(builder);
+
+			FileReference stateFile = FileReference.Combine(rootDir, DataDir, StateFileName);
+			DirectoryReference.CreateDirectory(stateFile.Directory);
+
+			using (FileStream stream = FileReference.Open(stateFile, FileMode.Create, FileAccess.Write))
+			{
+				foreach (ReadOnlyMemory<byte> chunk in builder.AsSequence())
+				{
+					await stream.WriteAsync(chunk);
+				}
+			}
+		}
+
+		protected static async Task InitAsync(DirectoryReference rootDir)
+		{
+			await WriteStateAsync(rootDir, new WorkspaceState());
+		}
+
+		protected DirectoryReference FindRootDir()
+		{
+			if (BaseDir != null)
+			{
+				return BaseDir;
+			}
+
+			DirectoryReference startDir = DirectoryReference.GetCurrentDirectory();
+			for (DirectoryReference? currentDir = startDir; currentDir != null; currentDir = currentDir.ParentDirectory)
+			{
+				FileReference markerFile = FileReference.Combine(currentDir, DataDir, StateFileName);
+				if (FileReference.Exists(markerFile))
+				{
+					return currentDir;
+				}
+			}
+			throw new InvalidDataException($"No root directory found under {startDir}");
+		}
+
+		protected Task<DirectoryState> GetCurrentDirectoryState(DirectoryReference rootDir, DirectoryState? oldState) => GetCurrentDirectoryState(rootDir.ToDirectoryInfo(), oldState);
+
+		protected async Task<DirectoryState> GetCurrentDirectoryState(DirectoryInfo directoryInfo, DirectoryState? oldState)
+		{
+			List<DirectoryInfo> subDirectoryInfos = new List<DirectoryInfo>();
+			foreach (DirectoryInfo subDirectoryInfo in directoryInfo.EnumerateDirectories().Where(x => !x.Name.StartsWith(".", StringComparison.Ordinal)))
+			{
+				subDirectoryInfos.Add(subDirectoryInfo);
+			}
+
+			Task<DirectoryState>[] tasks = new Task<DirectoryState>[subDirectoryInfos.Count];
+			for (int idx = 0; idx < subDirectoryInfos.Count; idx++)
+			{
+				DirectoryInfo subDirectoryInfo = subDirectoryInfos[idx];
+
+				DirectoryState? prevSubDirectoryState = null;
+				if (oldState != null)
+				{
+					oldState.Directories.TryGetValue(subDirectoryInfo.Name, out prevSubDirectoryState);
+				}
+
+				tasks[idx] = Task.Run(() => GetCurrentDirectoryState(subDirectoryInfo, prevSubDirectoryState));
+			}
+
+			DirectoryState newState = new DirectoryState();
+			for (int idx = 0; idx < subDirectoryInfos.Count; idx++)
+			{
+				newState.Directories[subDirectoryInfos[idx].Name] = await tasks[idx];
+			}
+			foreach (FileInfo fileInfo in directoryInfo.EnumerateFiles())
+			{
+				FileState? fileState;
+				if (oldState == null || !oldState.Files.TryGetValue(fileInfo.Name, out fileState) || fileState.Modified(fileInfo))
+				{
+					fileState = new FileState(fileInfo);
+				}
+				newState.Files[fileInfo.Name] = fileState;
+			}
+
+			if (oldState != null && oldState.Directories.Count == newState.Directories.Count && oldState.Files.Count == newState.Files.Count)
+			{
+				if (newState.Directories.Values.All(x => x.Hash != IoHash.Zero) && newState.Files.Values.All(x => x.Hash != IoHash.Zero))
+				{
+					return oldState;
+				}
+			}
+	
+			return newState;
+		}
+
+		protected static void RemoveAddedFiles(DirectoryState oldState, DirectoryState newState)
+		{
+			List<(Utf8String, DirectoryState?, DirectoryState?)> directoryDeltas = Zip(oldState.Directories, newState.Directories).ToList();
+			foreach ((Utf8String name, DirectoryState? oldSubDirState, _) in directoryDeltas)
+			{
+				if (oldSubDirState == null)
+				{
+					newState.Directories.Remove(name);
+				}
+			}
+
+			List<(Utf8String, FileState?, FileState?)> fileDeltas = Zip(oldState.Files, newState.Files).ToList();
+			foreach ((Utf8String name, FileState? oldFileState, _) in fileDeltas)
+			{
+				if (oldFileState == null)
+				{
+					newState.Files.Remove(name);
+				}
+			}
+		}
+
+		protected static DirectoryState DedupTrees(DirectoryState oldState, DirectoryState newState)
+		{
+			if (oldState == newState)
+			{
+				return newState;
+			}
+
+			foreach ((_, DirectoryState? oldSubDirState, DirectoryState? newSubDirState) in Zip(oldState.Directories, newState.Directories))
+			{
+				if (oldSubDirState == null || newSubDirState == null || oldSubDirState != DedupTrees(oldSubDirState, newSubDirState))
+				{
+					return newState;
+				}
+			}
+
+			foreach ((_, FileState? oldFileState, FileState? newFileState) in Zip(oldState.Files, newState.Files))
+			{
+				if (oldFileState != newFileState)
+				{
+					return newState;
+				}
+			}
+
+			return oldState;
+		}
+
+		protected static void PrintDelta(DirectoryState oldState, DirectoryState newState, ILogger logger) => PrintDelta("", oldState, newState, logger);
+
+		static void PrintDelta(string prefix, DirectoryState oldState, DirectoryState newState, ILogger logger)
+		{
+			foreach ((Utf8String name, DirectoryState? oldSubDirState, DirectoryState? newSubDirState) in Zip(oldState.Directories, newState.Directories))
+			{
+				if (oldSubDirState == null)
+				{
+					logger.LogInformation("{Path} (added)", $"{prefix}{name}/");
+				}
+				else if (newSubDirState == null)
+				{
+					logger.LogInformation("{Path} (removed)", $"{prefix}{name}/");
+				}
+				else if (oldSubDirState != newSubDirState)
+				{
+					PrintDelta($"{prefix}{name}/", oldSubDirState, newSubDirState, logger);
+				}
+			}
+
+			foreach ((Utf8String name, FileState? oldFileState, FileState? newFileState) in Zip(oldState.Files, newState.Files))
+			{
+				if (oldFileState == null)
+				{
+					logger.LogInformation("{Path} (added)", $"{prefix}{name}");
+				}
+				else if (newFileState == null)
+				{
+					logger.LogInformation("{Path} (deleted)", $"{prefix}{name}");
+				}
+				else if (oldFileState != newFileState)
+				{
+					logger.LogInformation("{Path} (modified)", $"{prefix}{name}");
+				}
+			}
+		}
+	}
+
+	[Command("Vcs", "Init", "Initialize a directory for VCS-like operations")]
+	class VcsInitCommand : VcsCommand
+	{
+		public override async Task<int> ExecuteAsync(ILogger logger)
+		{
+			BaseDir ??= DirectoryReference.GetCurrentDirectory();
+			await InitAsync(BaseDir);
+			logger.LogInformation("Initialized in {RootDir}", BaseDir);
+			return 0;
+		}
+	}
+
+	[Command("vcs", "branch", "Switch to a new branch")]
+	class VcsBranchCommand : VcsCommand
+	{
+		[CommandLine(Prefix = "-Name=", Required = true)]
+		public string Name { get; set; } = "";
+
+		public override async Task<int> ExecuteAsync(ILogger logger)
+		{
+			DirectoryReference rootDir = FindRootDir();
+
+			WorkspaceState workspaceState = await ReadStateAsync(rootDir);
+
+			using IStorageClientOwner owner = CreateStorageClient(rootDir, logger);
+			IStorageClient store = owner.Store;
+
+			RefName branchName = new RefName(Name);
+			if (await store.HasRefAsync(branchName))
+			{
+				logger.LogError("Branch {BranchName} already exists - use checkout instead.", branchName);
+				return 1;
+			}
+
+			logger.LogInformation("Starting work in new branch {BranchName}", branchName);
+
+			workspaceState.Branch = new RefName(Name);
+			workspaceState.Tree = new DirectoryState();
+			await WriteStateAsync(rootDir, workspaceState);
+
+			return 0;
+		}
+	}
+
+	[Command("vcs", "checkout", "Checkout a particular branch/change")]
+	class VcsCheckoutCommand : VcsCommand
+	{
+		[CommandLine("-Branch")]
+		public string? Branch { get; set; }
+
+		[CommandLine("-Change=")]
+		public int Change { get; set; }
+
+		[CommandLine("-Clean")]
+		public bool Clean { get; set; }
+
+		[CommandLine("-Force")]
+		public bool Force { get; set; }
+
+		public override async Task<int> ExecuteAsync(ILogger logger)
+		{
+			DirectoryReference rootDir = FindRootDir();
+
+			WorkspaceState workspaceState = await ReadStateAsync(rootDir);
+
+			DirectoryState oldState = workspaceState.Tree;
+			DirectoryState newState = await GetCurrentDirectoryState(rootDir, oldState);
+
+			if (!Clean)
+			{
+				RemoveAddedFiles(oldState, newState);
+				newState = DedupTrees(oldState, newState);
+			}
+
+			if (oldState != newState && !Force)
+			{
+				logger.LogInformation("Current workspace has modified files. Resolve before checking out a different change, or run with -Force.");
+				PrintDelta(oldState, newState, logger);
+				return 0;
+			}
+
+			RefName branchName = (Branch != null) ? new RefName(Branch) : workspaceState.Branch;
+
+			using IStorageClientOwner owner = CreateStorageClient(rootDir, logger);
+			IStorageClient store = owner.Store;
+
+			CommitNode tip = await store.ReadNodeAsync<CommitNode>(branchName);
+			if (Change != 0)
+			{
+				while (tip.Number != Change)
+				{
+					if (tip.Parent == null || tip.Number < Change)
+					{
+						logger.LogError("Unable to find change {Change}", Change);
+						return 1;
+					}
+					tip = await tip.Parent.ExpandAsync();
+				}
+			}
+
+			workspaceState.Tree = await RealizeAsync(tip.Contents, rootDir, newState, Clean, logger);
+			await WriteStateAsync(rootDir, workspaceState);
+
+			logger.LogInformation("Updated workspace to change {Number}", tip.Number);
+			return 0;
+		}
+
+		static async Task<DirectoryState> RealizeAsync(TreeNodeRef<DirectoryNode> directoryRef, DirectoryReference dirPath, DirectoryState? directoryState, bool clean, ILogger logger)
+		{
+			DirectoryReference.CreateDirectory(dirPath);
+
+			DirectoryState newState = new DirectoryState();
+
+			DirectoryNode directoryNode = await directoryRef.ExpandAsync();
+			foreach ((Utf8String name, DirectoryEntry? subDirEntry, DirectoryState? subDirState) in Zip(directoryNode.NameToDirectory, directoryState?.Directories))
+			{
+				DirectoryReference subDirPath = DirectoryReference.Combine(dirPath, name.ToString());
+				if (subDirEntry == null)
+				{
+					if (clean)
+					{
+						DirectoryReference.Delete(subDirPath);
+					}
+				}
+				else
+				{
+					newState.Directories[name] = await RealizeAsync(subDirEntry, subDirPath, subDirState, clean, logger);
+				}
+			}
+
+			foreach ((Utf8String name, FileEntry? fileEntry, FileState? fileState) in Zip(directoryNode.NameToFile, directoryState?.Files))
+			{
+				FileReference filePath = FileReference.Combine(dirPath, name.ToString());
+				if (fileEntry == null)
+				{
+					if (clean)
+					{
+						FileReference.Delete(filePath);
+					}
+				}
+				else if (fileState == null || fileState.Hash != fileEntry.Hash)
+				{
+					newState.Files[name] = await CheckoutFileAsync(fileEntry, filePath.ToFileInfo(), logger);
+				}
+				else
+				{
+					newState.Files[name] = fileState;
+				}
+			}
+
+			newState.Hash = directoryRef.Hash;
+			return newState;
+		}
+
+		static async Task<FileState> CheckoutFileAsync(TreeNodeRef<FileNode> fileRef, FileInfo fileInfo, ILogger logger)
+		{
+			logger.LogInformation("Updating {File} to {Hash}", fileInfo, fileRef.Hash);
+			FileNode fileNode = await fileRef.ExpandAsync();
+			await fileNode.CopyToFileAsync(fileInfo, CancellationToken.None);
+			fileInfo.Refresh();
+			return new FileState(fileInfo, fileRef.Hash);
+		}
+	}
+
+	[Command("vcs", "commit", "Commits data to the VCS store")]
+	class VcsCommitCommand : VcsCommand
+	{
+		[CommandLine("-Message=")]
+		public string? Message { get; set; }
+
+		public override async Task<int> ExecuteAsync(ILogger logger)
+		{
+			DirectoryReference rootDir = FindRootDir();
+
+			WorkspaceState workspaceState = await ReadStateAsync(rootDir);
+
+			DirectoryState oldState = workspaceState.Tree;
+			DirectoryState newState = await GetCurrentDirectoryState(rootDir, oldState);
+
+			if (oldState == newState)
+			{
+				logger.LogInformation("No files modified");
+				return 0;
+			}
+			
+			if (Message == null)
+			{
+				logger.LogError("Missing -Message=... argument");
+				return 1;
+			}
+
+			PrintDelta(oldState, newState, logger);
+
+			using IStorageClientOwner owner = CreateStorageClient(rootDir, logger);
+			IStorageClient store = owner.Store;
+
+			CommitNode? tip = await store.TryReadNodeAsync<CommitNode>(workspaceState.Branch);
+			TreeNodeRef<CommitNode>? tipRef = (tip == null) ? null : new TreeNodeRef<CommitNode>(tip);
+
+			DirectoryNode rootNode;
+			if (tip == null)
+			{
+				rootNode = new DirectoryNode();
+			}
+			else
+			{
+				rootNode = await tip.Contents.ExpandCopyAsync();
+			}
+
+			TreeNodeRef<DirectoryNode> rootRef = new TreeNodeRef<DirectoryNode>(rootNode);
+
+			List<(DirectoryNode, FileInfo, FileState)> files = new List<(DirectoryNode, FileInfo, FileState)>();
+			List<(TreeNodeRef<DirectoryNode>, DirectoryState)> directories = new List<(TreeNodeRef<DirectoryNode>, DirectoryState)>();
+			await UpdateTreeAsync(rootRef, rootDir, oldState, newState, files, directories);
+
+			TreeWriter writer = new TreeWriter(store);
+			await DirectoryNode.CopyFromDirectoryAsync(files.ConvertAll(x => (x.Item1, x.Item2)), new ChunkingOptions(), writer, CancellationToken.None);
+
+			CommitNode newTip;
+			if (tip == null)
+			{
+				newTip = new CommitNode(1, null, Message, rootRef);
+			}
+			else
+			{
+				newTip = new CommitNode(tip.Number + 1, tipRef, Message, rootRef);
+			}
+			await store.WriteNodeAsync(workspaceState.Branch, newTip);
+
+			foreach ((TreeNodeRef<DirectoryNode> directoryRef, DirectoryState directoryState) in directories)
+			{
+				directoryState.Hash = directoryRef.Hash;
+			}
+
+			foreach ((DirectoryNode directoryNode, FileInfo fileInfo, FileState fileState) in files)
+			{
+				FileEntry fileEntry = directoryNode.GetFileEntry(fileInfo.Name);
+				fileState.Hash = fileEntry.Hash;
+			}
+
+			workspaceState.Change = newTip.Number;
+			workspaceState.Tree = newState;
+			await WriteStateAsync(rootDir, workspaceState);
+
+			logger.LogInformation("Commited in change {Number}", newTip.Number);
+			return 0;
+		}
+
+		private async Task UpdateTreeAsync(TreeNodeRef<DirectoryNode> rootRef, DirectoryReference rootDir, DirectoryState? oldState, DirectoryState newState, List<(DirectoryNode, FileInfo, FileState)> files, List<(TreeNodeRef<DirectoryNode>, DirectoryState)> directories)
+		{
+			directories.Add((rootRef, newState));
+
+			DirectoryNode root = await rootRef.ExpandAsync();
+			foreach ((Utf8String name, DirectoryState? oldSubDirState, DirectoryState? newSubDirState) in Zip(oldState?.Directories, newState.Directories))
+			{
+				if (newSubDirState == null)
+				{
+					root.DeleteDirectory(name);
+				}
+				else if(oldSubDirState != newSubDirState)
+				{
+					await UpdateTreeAsync(root.FindOrAddDirectoryEntry(name), DirectoryReference.Combine(rootDir, name.ToString()), oldSubDirState, newSubDirState, files, directories);
+				}
+			}
+
+			foreach ((Utf8String name, FileState? oldFileState, FileState? newFileState) in Zip(oldState?.Files, newState.Files))
+			{
+				if(newFileState == null)
+				{
+					root.DeleteFile(name);
+				}
+				else if(oldFileState != newFileState)
+				{
+					files.Add((root, FileReference.Combine(rootDir, name.ToString()).ToFileInfo(), newFileState));
+				}
+			}
+		}
+	}
+
+	[Command("Vcs", "Log", "Print a history of commits")]
+	class VcsLogCommand : VcsCommand
+	{
+		[CommandLine("-Count")]
+		public int Count { get; set; } = 20;
+
+		public override async Task<int> ExecuteAsync(ILogger logger)
+		{
+			DirectoryReference rootDir = FindRootDir();
+
+			WorkspaceState workspaceState = await ReadStateAsync(rootDir);
+
+			using IStorageClientOwner owner = CreateStorageClient(rootDir, logger);
+			IStorageClient store = owner.Store;
+
+			List<CommitNode> commits = new List<CommitNode>();
+
+			CommitNode? tip = await store.TryReadNodeAsync<CommitNode>(workspaceState.Branch);
+			if (tip != null)
+			{
+				for (int idx = 0; idx < Count; idx++)
+				{
+					commits.Add(tip);
+
+					if (tip.Parent == null)
+					{
+						break;
+					}
+
+					tip = await tip.Parent.ExpandAsync();
+				}
+			}
+
+			foreach (CommitNode commit in commits)
+			{
+				logger.LogInformation("Commit {Number}: {Message}", commit.Number, commit.Message);
+			}
+
+			return 0;
+		}
+	}
+
+	[Command("Vcs", "Status", "Find status of local files")]
+	class VcsStatusCommand : VcsCommand
+	{
+		public override async Task<int> ExecuteAsync(ILogger logger)
+		{
+			DirectoryReference rootDir = FindRootDir();
+
+			WorkspaceState workspaceState = await ReadStateAsync(rootDir);
+
+			DirectoryState oldState = workspaceState.Tree;
+			DirectoryState newState = await GetCurrentDirectoryState(rootDir, oldState);
+
+			logger.LogInformation("On branch {BranchName}", workspaceState.Branch);
+
+			if (oldState == newState)
+			{
+				logger.LogInformation("No files modified");
+			}
+			else
+			{
+				PrintDelta(oldState, newState, logger);
+			}
+
+			return 0;
+		}
+	}
+}
diff --git a/Engine/Source/Programs/Horde/Horde.Agent/Horde.Agent.csproj b/Engine/Source/Programs/Horde/Horde.Agent/Horde.Agent.csproj
--- a/Engine/Source/Programs/Horde/Horde.Agent/Horde.Agent.csproj
+++ b/Engine/Source/Programs/Horde/Horde.Agent/Horde.Agent.csproj
@@ -80,6 +80,9 @@
     <Content Include="..\Horde.Build\Protos\horde\horde_rpc_messages.proto">
       <Link>Protos\horde_rpc_messages.proto</Link>
     </Content>
+    <Content Include="Horde.bat">
+      <CopyToOutputDirectory>PreserveNewest</CopyToOutputDirectory>
+    </Content>
     <Content Include="..\Horde.Build\Protos\horde\common\**\*.proto">
       <Link>Protos\horde\common\%(RecursiveDir)%(FileName)%(Extension)</Link>
     </Content>
diff --git a/Engine/Source/Programs/Horde/Horde.Agent/Horde.bat b/Engine/Source/Programs/Horde/Horde.Agent/Horde.bat
new file mode 100755
--- /dev/null
+++ b/Engine/Source/Programs/Horde/Horde.Agent/Horde.bat
@@ -0,0 +1,3 @@
+@echo off
+dotnet "%~dp0\HordeAgent.dll" %*
+
