Horde Storage - Removed references to object_last_access table that doesnt exist anymore

diff --git a/Engine/Source/Programs/Horde/HordeStorage/Horde.Storage/Implementation/Objects/ScyllaReferencesStore.cs b/Engine/Source/Programs/Horde/HordeStorage/Horde.Storage/Implementation/Objects/ScyllaReferencesStore.cs
--- a/Engine/Source/Programs/Horde/HordeStorage/Horde.Storage/Implementation/Objects/ScyllaReferencesStore.cs
+++ b/Engine/Source/Programs/Horde/HordeStorage/Horde.Storage/Implementation/Objects/ScyllaReferencesStore.cs
@@ -389,28 +389,4 @@ namespace Horde.Storage.Implementation
         [System.Diagnostics.CodeAnalysis.SuppressMessage("Usage", "CA2227:Collection properties should be read only", Justification = "Used by serialization")]
         public List<string> Buckets { get; set; } = new List<string>();
     }
-
-    [Cassandra.Mapping.Attributes.Table("object_last_access")]
-    public class ScyllaObjectLastAccessAt
-    {
-        public ScyllaObjectLastAccessAt()
-        {
-
-        }
-
-        [Cassandra.Mapping.Attributes.PartitionKey(0)]
-        public string? Namespace { get; set; }
-
-        [Cassandra.Mapping.Attributes.PartitionKey(1)]
-        [Cassandra.Mapping.Attributes.Column("partition_index")]
-        public sbyte? PartitionIndex { get;set; }
-
-        [Cassandra.Mapping.Attributes.PartitionKey(2)]
-        [Cassandra.Mapping.Attributes.Column("accessed_at")]
-        public LocalDate? AccessedAt { get; set; }
-
-        [System.Diagnostics.CodeAnalysis.SuppressMessage("Usage", "CA2227:Collection properties should be read only", Justification = "Used by serialization")]
-        public List<string> Objects { get; set; } = new List<string>();
-
-    }
 }
diff --git a/Engine/Source/Programs/Horde/HordeStorage/Horde.Storage/Tests/Functional/Objects/ReferencesTests.cs b/Engine/Source/Programs/Horde/HordeStorage/Horde.Storage/Tests/Functional/Objects/ReferencesTests.cs
--- a/Engine/Source/Programs/Horde/HordeStorage/Horde.Storage/Tests/Functional/Objects/ReferencesTests.cs
+++ b/Engine/Source/Programs/Horde/HordeStorage/Horde.Storage/Tests/Functional/Objects/ReferencesTests.cs
@@ -75,7 +75,6 @@ namespace Horde.Storage.FunctionalTests.References
 
             ISession replicatedKeyspace = scyllaSessionManager.GetSessionForReplicatedKeyspace();
             await replicatedKeyspace.ExecuteAsync(new SimpleStatement("DROP TABLE IF EXISTS objects"));
-            await replicatedKeyspace.ExecuteAsync(new SimpleStatement("DROP TABLE IF EXISTS object_last_access"));
             await replicatedKeyspace.ExecuteAsync(new SimpleStatement("DROP TABLE IF EXISTS content_id"));
 
             // we need to clear out the state we modified in the replication log, otherwise the replication log tests will fail
