Horde: Catch errors while attempting to escalate issues, to prevent blocking escalation of other issues.

diff --git a/Engine/Source/Programs/Horde/Horde.Build/Notifications/Sinks/SlackNotificationSink.cs b/Engine/Source/Programs/Horde/Horde.Build/Notifications/Sinks/SlackNotificationSink.cs
--- a/Engine/Source/Programs/Horde/Horde.Build/Notifications/Sinks/SlackNotificationSink.cs
+++ b/Engine/Source/Programs/Horde/Horde.Build/Notifications/Sinks/SlackNotificationSink.cs
@@ -2066,17 +2066,24 @@ namespace Horde.Build.Notifications.Sinks
 				foreach (int issueId in issueIds)
 				{
 					cancellationToken.ThrowIfCancellationRequested();
-
-					double? nextTime = await EscalateSingleIssueAsync(issueId, utcNow);
-					if (nextTime == null)
+					try
 					{
-						_logger.LogInformation("Cancelling escalation for issue {IssueId}", issueId);
-						await _escalateIssues.RemoveAsync(issueId);
+						double? nextTime = await EscalateSingleIssueAsync(issueId, utcNow);
+						if (nextTime == null)
+						{
+							_logger.LogInformation("Cancelling escalation for issue {IssueId}", issueId);
+							await _escalateIssues.RemoveAsync(issueId);
+						}
+						else
+						{
+							_logger.LogInformation("Next escalation for issue {IssueId} is at timestamp {Time}", issueId, nextTime.Value);
+							await _escalateIssues.AddAsync(issueId, nextTime.Value);
+						}
 					}
-					else
+					catch (SlackException ex)
 					{
-						_logger.LogInformation("Next escalation for issue {IssueId} is at timestamp {Time}", issueId, nextTime.Value);
-						await _escalateIssues.AddAsync(issueId, nextTime.Value);
+						_logger.LogError(ex, "Slack exception while escalating issue {IssueId}; cancelling.", issueId);
+						await _escalateIssues.RemoveAsync(issueId);
 					}
 				}
 			}
