Horde: Don't reopen verified issues, hash issue handler improvements, multistream issue test
#rnx

diff --git a/Engine/Source/Programs/Horde/Horde.Build.Tests/IssueServiceTests.cs b/Engine/Source/Programs/Horde/Horde.Build.Tests/IssueServiceTests.cs
--- a/Engine/Source/Programs/Horde/Horde.Build.Tests/IssueServiceTests.cs
+++ b/Engine/Source/Programs/Horde/Horde.Build.Tests/IssueServiceTests.cs
@@ -2136,6 +2136,140 @@ namespace Horde.Build.Tests
 
 		}
 
+		[TestMethod]
+		public async Task MultipleStreamIssueTest()
+		{
+
+			int hours = 0;
+
+			IUser bob = await UserCollection.FindOrAddUserByLoginAsync("Bob");
+			IUser anne = await UserCollection.FindOrAddUserByLoginAsync("Anne");
+
+			// Main Stream
+
+			// last success
+			_perforce.AddChange(_mainStreamId, 22136421, bob, "Description", new string[] { "a/b.cpp" });
+
+			// unrelated change, no job run on it, merges to release between breakages in that stream
+			_perforce.AddChange(_mainStreamId, 22136521, anne, "Description", new string[] { "a/b.cpp" });
+
+			// change with a breakage in main
+			_perforce.AddChange(_mainStreamId, 22145160, bob, "Description", new string[] { "a/b.cpp" });
+
+			// success
+			_perforce.AddChange(_mainStreamId, 22151893, bob, "Description", new string[] { "a/b.cpp" });
+
+			// breaking change merged from release
+			_perforce.AddChange(_mainStreamId, 22166000, 22165119, bob, bob, "Description", new string[] { "a/b.cpp" });
+
+			// Release Stream
+
+			// last success
+			_perforce.AddChange(_releaseStreamId, 22133008, bob, "Description", new string[] { "a/b.cpp" });
+
+			// unrelated change originating in main, no job run on it
+			_perforce.AddChange(_releaseStreamId, 22152050, 22136521, anne, anne, "Description", new string[] { "a/b.cpp" });
+
+			// unrelated breaking change
+			_perforce.AddChange(_releaseStreamId, 22165119, bob, "Description", new string[] { "a/b.cpp" });
+
+			string[] breakage1 =
+			{
+					"Error executing d:\\build\\AutoSDK\\Sync\\HostWin64\\Win64\\VS2019\\14.29.30146\\bin\\HostX64\\x64\\link.exe (tool returned code: 1123)",
+					"LINK : fatal error LNK1123: failure during conversion to COFF: file invalid or corrupt"
+			};
+
+			string[] breakage2 =
+			{
+					"Error executing d:\\build\\AutoSDK\\Sync\\HostWin64\\Win64\\VS2019\\14.29.30146\\bin\\HostX64\\x64\\link.exe (tool returned code: 1169)",
+					"D:\\build\\++UE5\\Sync\\QAGame\\Binaries\\Win64\\QAGameEditor.exe : fatal error LNK1169: one or more multiply defined symbols found",
+					"Module.Core.1_of_20.cpp.obj : error LNK2005: \"int GNumForegroundWorkers\" (?GNumForegroundWorkers@@3HA) already defined in Module.UnrealEd.20_of_42.cpp.obj"
+			};
+
+			// #1
+			// Job runs successfully in release stream
+			{
+				IJob job = CreateJob(_releaseStreamId, 22133008, "Test Build", _graph, TimeSpan.FromHours(hours++));
+				await UpdateCompleteStep(job, 0, 0, JobStepOutcome.Success);
+				await UpdateCompleteStep(job, 0, 1, JobStepOutcome.Success);
+				List<IIssue> issues = await IssueCollection.FindIssuesAsync();
+				Assert.AreEqual(0, issues.Count);
+			}
+
+			// #2
+			// Job runs successfully in main stream at a latest CL
+			{
+				IJob job = CreateJob(_mainStreamId, 22136421, "Test Build", _graph, TimeSpan.FromHours(hours++));
+				await UpdateCompleteStep(job, 0, 0, JobStepOutcome.Success);
+				await UpdateCompleteStep(job, 0, 1, JobStepOutcome.Success);
+				List<IIssue> issues = await IssueCollection.FindIssuesAsync();
+				Assert.AreEqual(0, issues.Count);
+			}
+
+
+			// #3
+			// Scenario: Job step encounters a hashed issue at CL 22145160
+			// Expected: Hashed issue type is created
+			{
+				IJob job = CreateJob(_mainStreamId, 22145160, "Test Build", _graph, TimeSpan.FromHours(hours++));				
+				await ParseEventsAsync(job, 0, 0, breakage1);
+				await UpdateCompleteStep(job, 0, 0, JobStepOutcome.Failure);
+
+
+				List<IIssue> issues = await IssueCollection.FindIssuesAsync();
+				Assert.AreEqual(1, issues.Count);
+				Assert.AreEqual(1, issues[0].Id);
+				Assert.AreEqual("Errors in Update Version Files", issues[0].Summary);
+			}
+
+			// #4
+			// Scenario: Job succeeds an hour later in CL 22151893
+			// Expected: Existing issue is closed
+			{
+				IJob job = CreateJob(_mainStreamId, 22151893, "Test Build", _graph, TimeSpan.FromHours(hours++));
+				await UpdateCompleteStep(job, 0, 0, JobStepOutcome.Success);
+
+				List<IIssue> issues = await IssueCollection.FindIssuesAsync();
+				Assert.AreEqual(0, issues.Count);
+			}
+
+			// #5
+			// Scenario: Job step encounters a failure on CL 22165119
+			// Expected:New issue is created, and does not reopen the one in main
+			{
+
+				IJob job = CreateJob(_releaseStreamId, 22165119, "Test Build", _graph, TimeSpan.FromHours(hours++));
+				await ParseEventsAsync(job, 0, 0, breakage2);
+				await UpdateCompleteStep(job, 0, 0, JobStepOutcome.Failure);
+
+				List<IIssue> issues = await IssueCollection.FindIssuesAsync();
+				Assert.AreEqual(1, issues.Count);
+				
+				// Check that new issue was created
+				Assert.AreEqual(2, issues[0].Id);
+				Assert.AreEqual("Errors in Update Version Files", issues[0].Summary);
+
+			}
+
+			// #6
+			// Scenario: Job step in main encounters a failure on CL 22166000 originating from 22165119 breakage
+			// Expected: Step is added to existing issue and summary is updated
+			{
+
+				IJob job = CreateJob(_mainStreamId, 22166000, "Test Build", _graph, TimeSpan.FromHours(hours++));
+				await ParseEventsAsync(job, 0, 1, breakage2);
+				await UpdateCompleteStep(job, 0,1, JobStepOutcome.Failure);
+
+				List<IIssue> issues = await IssueCollection.FindIssuesAsync();
+				Assert.AreEqual(1, issues.Count);
+
+				// Check that new issue was created
+				Assert.AreEqual(2, issues[0].Id);
+				Assert.AreEqual("Errors in Update Version Files and Compile UnrealHeaderTool Win64", issues[0].Summary);
+
+			}
+
+		}
 
 
 		private async Task ParseAsync(LogId logId, string[] lines)
diff --git a/Engine/Source/Programs/Horde/Horde.Build/Issues/Handlers/HashedIssueHandler.cs b/Engine/Source/Programs/Horde/Horde.Build/Issues/Handlers/HashedIssueHandler.cs
--- a/Engine/Source/Programs/Horde/Horde.Build/Issues/Handlers/HashedIssueHandler.cs
+++ b/Engine/Source/Programs/Horde/Horde.Build/Issues/Handlers/HashedIssueHandler.cs
@@ -7,6 +7,7 @@ using System.Text.RegularExpressions;
 using EpicGames.Core;
 using Horde.Build.Jobs;
 using Horde.Build.Jobs.Graphs;
+using Microsoft.Extensions.Logging;
 
 namespace Horde.Build.Issues.Handlers
 {
@@ -23,6 +24,21 @@ namespace Horde.Build.Issues.Handlers
 		/// <inheritdoc/>
 		public override int Priority => 1;
 
+		/// <summary>
+		///  Known general events
+		/// </summary>
+		static readonly HashSet<EventId> s_knownGeneralEvents = new HashSet<EventId> { KnownLogEvents.Generic, KnownLogEvents.ExitCode, KnownLogEvents.Horde, KnownLogEvents.Horde_InvalidPreflight };
+
+		/// <summary>
+		/// Determines if the given event is general and should be salted to make it unique
+		/// </summary>
+		/// <param name="eventId">The event id to compare</param>
+		/// <returns>True if the given event id matches</returns>
+		public static bool IsGeneralEventId(EventId eventId)
+		{
+			return s_knownGeneralEvents.Contains(eventId) || (eventId.Id >= KnownLogEvents.Systemic.Id && eventId.Id <= KnownLogEvents.Systemic_Max.Id);
+		}
+
 		/// <inheritdoc/>
 		public override void TagEvents(IJob job, INode node, IReadOnlyNodeAnnotations annotations, IReadOnlyList<IssueEvent> stepEvents)
 		{
@@ -33,7 +49,18 @@ namespace Horde.Build.Issues.Handlers
 
 			foreach (IssueEvent stepEvent in stepEvents)
 			{
-				if (hashes.Count < 25 && TryGetHash(stepEvent.Message, out Md5Hash hash))
+				string hashSource = stepEvent.Message;
+				
+				if (stepEvent.EventId != null)
+				{
+					// If the event is general, salt the hash with the stream id, template, and node name, otherwise it will be aggressively matched
+					if (IsGeneralEventId(stepEvent.EventId.Value))
+					{
+						hashSource += $"step:{job.StreamId}:{job.TemplateId}:{node.Name}";
+					}					
+				}
+
+				if (hashes.Count < 25 && TryGetHash(hashSource, out Md5Hash hash))
 				{
 					hashes.Add(hash);
 					stepEvent.Fingerprint = new NewIssueFingerprint(Type, new[] { $"hash:{hash}" }, null, metadata);
diff --git a/Engine/Source/Programs/Horde/Horde.Build/Issues/IssueService.cs b/Engine/Source/Programs/Horde/Horde.Build/Issues/IssueService.cs
--- a/Engine/Source/Programs/Horde/Horde.Build/Issues/IssueService.cs
+++ b/Engine/Source/Programs/Horde/Horde.Build/Issues/IssueService.cs
@@ -1213,8 +1213,8 @@ namespace Horde.Build.Issues
 				existingIssues = await _issueCollection.FindIssuesForChangesAsync(span.Suspects.ConvertAll(x => x.OriginatingChange ?? x.Change));
 				_logger.LogDebug("Found {NumIssues} open issues in {StreamId} from [{ChangeList}]", existingIssues.Count, span.StreamId, String.Join(", ", span.Suspects.ConvertAll(x => (x.OriginatingChange ?? x.Change).ToString())));
 			}
-
-			IIssue? issue = existingIssues.FirstOrDefault(x => x.Fingerprints.Any(y => y.IsMatch(span.Fingerprint)));
+			
+			IIssue? issue = existingIssues.FirstOrDefault(x => x.VerifiedAt == null && x.Fingerprints.Any(y => y.IsMatch(span.Fingerprint)));
 			if (issue == null)
 			{
 				string summary = GetSummary(span.Fingerprint, span.FirstFailure.Severity);
