Horde: Limit number of retries for log chunk retrieval and completion

diff --git a/Engine/Source/Programs/Horde/Horde.Build/Logs/Builder/RedisLogBuilder.cs b/Engine/Source/Programs/Horde/Horde.Build/Logs/Builder/RedisLogBuilder.cs
--- a/Engine/Source/Programs/Horde/Horde.Build/Logs/Builder/RedisLogBuilder.cs
+++ b/Engine/Source/Programs/Horde/Horde.Build/Logs/Builder/RedisLogBuilder.cs
@@ -134,7 +134,10 @@ namespace Horde.Build.Logs.Builder
 		{
 			IDatabase redisDb = _redisConnectionPool.GetDatabase();
 			ChunkKeys keys = new ChunkKeys(logId, offset);
-			for (; ; )
+			
+			const int MaxRetries = 10;
+			int numTries;
+			for (numTries = 0; numTries < MaxRetries; numTries++)
 			{
 				using IScope scope = GlobalTracer.Instance.BuildSpan("Redis.CompleteSubChunk").StartActive();
 				scope.Span.SetTag("LogId", logId.ToString());
@@ -148,7 +151,7 @@ namespace Horde.Build.Logs.Builder
 				RedisValue subChunkTextValue = await redisDb.StringGetAsync(keys.SubChunkData);
 				if (subChunkTextValue.IsNullOrEmpty)
 				{
-					break;
+					return;
 				}
 
 				ReadOnlyLogText subChunkText = new ReadOnlyLogText(subChunkTextValue);
@@ -176,9 +179,15 @@ namespace Horde.Build.Logs.Builder
 				if (await writeTransaction.ExecuteAsync())
 				{
 					_logger.LogDebug("Completed sub-chunk for log {LogId} chunk offset {Offset} -> sub-chunk size {SubChunkSize}, chunk size {ChunkSize}", logId, offset, subChunkDataBytes.Length, await newLength);
-					break;
+					return;
 				}
+
+				// Cool down before retrying
+				await Task.Delay(100);
 			}
+
+			_logger.LogError("Unable to complete sub-chunk for {LogId} at {Offset} after {NumTries} tries", logId, offset, numTries);
+			throw new Exception($"Unable to complete sub-chunk for {logId} at {offset} after {numTries} tries");
 		}
 
 		/// <inheritdoc/>
@@ -230,7 +239,9 @@ namespace Horde.Build.Logs.Builder
 		{
 			IDatabase redisDb = _redisConnectionPool.GetDatabase();
 			ChunkKeys keys = new ChunkKeys(logId, offset);
-			for (; ; )
+			const int MaxRetries = 10;
+			int numTries;
+			for (numTries = 0; numTries < MaxRetries; numTries++)
 			{
 				using IScope scope = GlobalTracer.Instance.BuildSpan("Redis.GetChunk").StartActive();
 				scope.Span.SetTag("LogId", logId.ToString());
@@ -272,12 +283,17 @@ namespace Horde.Build.Logs.Builder
 
 					if (subChunks.Count == 0)
 					{
-						break;
+						return null;
 					}
 
 					return new LogChunkData(offset, lineIndex, subChunks);
 				}
+
+				// Cool down before retrying
+				await Task.Delay(100);
 			}
+
+			_logger.LogWarning("Unable to get chunk for {LogId} at {Offset}:{LineIndex} after {NumTries} tries", logId, offset, lineIndex, numTries);
 			return null;
 		}
 
