Horde: Keep hash consistent if only have non-unique events
#rnx

diff --git a/Engine/Source/Programs/Horde/Horde.Build/Issues/Handlers/HashedIssueHandler.cs b/Engine/Source/Programs/Horde/Horde.Build/Issues/Handlers/HashedIssueHandler.cs
--- a/Engine/Source/Programs/Horde/Horde.Build/Issues/Handlers/HashedIssueHandler.cs
+++ b/Engine/Source/Programs/Horde/Horde.Build/Issues/Handlers/HashedIssueHandler.cs
@@ -47,16 +47,20 @@ namespace Horde.Build.Issues.Handlers
 			NewIssueFingerprint? genericFingerprint = null;
 			HashSet<Md5Hash> hashes = new HashSet<Md5Hash>();
 
+			// keep hash consistent when only have general, non-unique events
+			bool allGeneral = stepEvents.FirstOrDefault(stepEvent => stepEvent.EventId == null || !IsGeneralEventId(stepEvent.EventId.Value)) == null;
+
 			foreach (IssueEvent stepEvent in stepEvents)
 			{
 				string hashSource = stepEvent.Message;
 				
-				if (stepEvent.EventId != null)
+				if (!allGeneral && stepEvent.EventId != null)
 				{
-					// If the event is general, salt the hash with the stream id, template, and node name, otherwise it will be aggressively matched
+					// If the event is general, salt the hash with the stream id, template, otherwise it will be aggressively matched.
+					// Consider salting with node name, though template id should be enough and have better grouping
 					if (IsGeneralEventId(stepEvent.EventId.Value))
 					{
-						hashSource += $"step:{job.StreamId}:{job.TemplateId}:{node.Name}";
+						hashSource += $"step:{job.StreamId}:{job.TemplateId}";
 					}					
 				}
 
