Horde: Check that a preflight change is still valid before starting conforms on agents.

diff --git a/Engine/Source/Programs/Horde/Horde.Build/Jobs/JobService.cs b/Engine/Source/Programs/Horde/Horde.Build/Jobs/JobService.cs
--- a/Engine/Source/Programs/Horde/Horde.Build/Jobs/JobService.cs
+++ b/Engine/Source/Programs/Horde/Horde.Build/Jobs/JobService.cs
@@ -829,7 +829,7 @@ namespace Horde.Build.Jobs
 
 			using IDisposable scope = _logger.BeginScope("UpdateBatchAsync({JobId})", job.Id);
 
-			bool bCheckForBadAgent = true;
+			bool checkForBadAgent = true;
 			for (; ; )
 			{
 				// Find the index of the appropriate batch
@@ -857,22 +857,25 @@ namespace Horde.Build.Jobs
 						newError = JobStepBatchError.Incomplete;
 
 						// Find the agent and set the conform flag
-						if (bCheckForBadAgent)
+						if (checkForBadAgent)
 						{
-							for (; ; )
+							if (await IsValidJobAsync(job))
 							{
-								IAgent? agent = await _agents.GetAsync(batch.AgentId.Value);
-								if (agent == null || agent.RequestConform)
+								for (; ; )
 								{
-									break;
-								}
-								if (await _agents.TryUpdateSettingsAsync(agent, bRequestConform: true) != null)
-								{
-									_logger.LogError("Agent {AgentId} did not complete lease; marking for conform", agent.Id);
-									break;
+									IAgent? agent = await _agents.GetAsync(batch.AgentId.Value);
+									if (agent == null || agent.RequestConform)
+									{
+										break;
+									}
+									if (await _agents.TryUpdateSettingsAsync(agent, bRequestConform: true) != null)
+									{
+										_logger.LogError("Agent {AgentId} did not complete lease; marking for conform", agent.Id);
+										break;
+									}
 								}
 							}
-							bCheckForBadAgent = false;
+							checkForBadAgent = false;
 						}
 					}
 				}
@@ -895,6 +898,35 @@ namespace Horde.Build.Jobs
 			}
 		}
 
+		async Task<bool> IsValidJobAsync(IJob job)
+		{
+			try
+			{
+				if (job.PreflightChange != 0)
+				{
+					IStream? stream = await _streamService.GetStreamAsync(job.StreamId);
+					if (stream == null)
+					{
+						_logger.LogWarning("Job {JobId} is no longer valid - stream {StreamId} does not exist.", job.Id, job.StreamId);
+						return false;
+					}
+
+					(CheckShelfResult result, _) = await _perforceService.CheckShelfAsync(stream, job.PreflightChange);
+					if (result != CheckShelfResult.Ok)
+					{
+						_logger.LogWarning("Job {JobId} is no longer valid - check shelf returned {Result}", result);
+						return false;
+					}
+				}
+				return true;
+			}
+			catch(Exception ex)
+			{
+				_logger.LogWarning(ex, "Job {JobId} is no longer valid.", job.Id);
+				return false;
+			}
+		}
+
 		/// <summary>
 		/// Attempts to update the state of a batch
 		/// </summary>
