Horde: Improve errors with invalid storage backends.

diff --git a/Engine/Source/Programs/Horde/Horde.Build/Storage/StorageService.cs b/Engine/Source/Programs/Horde/Horde.Build/Storage/StorageService.cs
--- a/Engine/Source/Programs/Horde/Horde.Build/Storage/StorageService.cs
+++ b/Engine/Source/Programs/Horde/Horde.Build/Storage/StorageService.cs
@@ -275,6 +275,11 @@ namespace Horde.Build.Storage
 					Dictionary<BackendId, BackendConfig> mergedIdToConfig = new Dictionary<BackendId, BackendConfig>();
 					foreach (NamespaceConfig namespaceConfig in storageConfig.Namespaces)
 					{
+						if (namespaceConfig.Backend.IsEmpty)
+						{
+							throw new StorageException($"No backend configured for namespace {namespaceConfig.Id}");
+						}
+
 						BackendConfig backendConfig = GetBackendConfig(namespaceConfig.Backend, backendIdToConfig, mergedIdToConfig);
 						IStorageBackend backend = CreateStorageBackend(backendConfig);
 
@@ -339,7 +344,11 @@ namespace Horde.Build.Storage
 			{
 				mergedIdToConfig.Add(backendId, null!);
 
-				config = baseIdToConfig[backendId];
+				if (!baseIdToConfig.TryGetValue(backendId, out config))
+				{
+					throw new StorageException($"Unable to find storage backend '{backendId}'"); 
+				}
+
 				if (config.Base != BackendId.Empty)
 				{
 					BackendConfig baseConfig = GetBackendConfig(config.Base, baseIdToConfig, mergedIdToConfig);
