Horde: Additional logging for ref updates.

diff --git a/Engine/Source/Programs/Horde/Horde.Build/Storage/RefCollection.cs b/Engine/Source/Programs/Horde/Horde.Build/Storage/RefCollection.cs
--- a/Engine/Source/Programs/Horde/Horde.Build/Storage/RefCollection.cs
+++ b/Engine/Source/Programs/Horde/Horde.Build/Storage/RefCollection.cs
@@ -6,6 +6,7 @@ using System.Threading.Tasks;
 using EpicGames.Horde.Storage;
 using Horde.Build.Server;
 using Microsoft.Extensions.Caching.Memory;
+using Microsoft.Extensions.Logging;
 using MongoDB.Bson;
 using MongoDB.Bson.Serialization.Attributes;
 using MongoDB.Driver;
@@ -83,16 +84,19 @@ namespace Horde.Build.Storage
 
 		readonly IMongoCollection<RefDocument> _refs;
 		readonly IMemoryCache _cache;
+		readonly ILogger _logger;
 
 		/// <summary>
 		/// Constructor
 		/// </summary>
 		/// <param name="mongoService">The mongo service implementation</param>
 		/// <param name="cache">Cache for ref values</param>
-		public RefCollection(MongoService mongoService, IMemoryCache cache)
+		/// <param name="logger">Logger for output</param>
+		public RefCollection(MongoService mongoService, IMemoryCache cache, ILogger<RefCollection> logger)
 		{
 			_refs = mongoService.GetCollection<RefDocument>("Refs", keys => keys.Ascending(x => x.NamespaceId).Ascending(x => x.Name), unique: true);
 			_cache = cache;
+			_logger = logger;
 		}
 
 		CachedRefValue AddRefToCache(NamespaceId namespaceId, RefName name, NodeLocator target)
@@ -141,7 +145,25 @@ namespace Horde.Build.Storage
 		public async Task WriteRefTargetAsync(NamespaceId namespaceId, RefName name, NodeLocator target, CancellationToken cancellationToken = default)
 		{
 			RefDocument refDocument = new RefDocument(namespaceId, name, target);
-			await _refs.ReplaceOneAsync(x => x.NamespaceId == namespaceId && x.Name == name, refDocument, new ReplaceOptions { IsUpsert = true }, cancellationToken);
+			
+			ReplaceOneResult result = await _refs.ReplaceOneAsync(x => x.NamespaceId == namespaceId && x.Name == name, refDocument, new ReplaceOptions { IsUpsert = true }, cancellationToken);
+			if (result.ModifiedCount > 0)
+			{
+				_logger.LogInformation("Updated ref {NamespaceId}:{RefName} to {Target}", namespaceId, name, target);
+			}
+			else if (result.MatchedCount > 0)
+			{
+				_logger.LogInformation("Updated ref {NamespaceId}:{RefName} to {Target} (no-op)", namespaceId, name, target);
+			}
+			else if (result.UpsertedId.IsObjectId)
+			{
+				_logger.LogInformation("Inserted ref {NamespaceId}:{RefName} to {Target} (id = {NewRefId})", namespaceId, name, target, (ObjectId)result.UpsertedId);
+			}
+			else
+			{
+				_logger.LogError("Unable to update ref {NamespaceId}:{RefName} (modified: {Modified}, matched: {Matched}, upsert: {Upsert})", namespaceId, name, result.ModifiedCount, result.MatchedCount, result.UpsertedId);
+			}
+			
 			AddRefToCache(namespaceId, name, target);
 		}
 	}
diff --git a/Engine/Source/Programs/Shared/EpicGames.Horde/Storage/IStorageClient.cs b/Engine/Source/Programs/Shared/EpicGames.Horde/Storage/IStorageClient.cs
--- a/Engine/Source/Programs/Shared/EpicGames.Horde/Storage/IStorageClient.cs
+++ b/Engine/Source/Programs/Shared/EpicGames.Horde/Storage/IStorageClient.cs
@@ -87,6 +87,9 @@ namespace EpicGames.Horde.Storage
 		/// Determines if this locator points to a valid entry
 		/// </summary>
 		public bool IsValid() => Blob.IsValid();
+
+		/// <inheritdoc/>
+		public override string ToString() => $"{Blob}#{ExportIdx}";
 	}
 
 	/// <summary>
