Horde: Support specifying a list of device models as a constraint
#rnx

diff --git a/Engine/Source/Programs/Horde/Horde.Build.Tests/DeviceServiceTests.cs b/Engine/Source/Programs/Horde/Horde.Build.Tests/DeviceServiceTests.cs
--- a/Engine/Source/Programs/Horde/Horde.Build.Tests/DeviceServiceTests.cs
+++ b/Engine/Source/Programs/Horde/Horde.Build.Tests/DeviceServiceTests.cs
@@ -12,6 +12,8 @@ using Horde.Build.Users;
 using System.Security.Claims;
 using Microsoft.AspNetCore.Http;
 using Horde.Build.Server;
+using System.Linq;
+using System;
 
 namespace Horde.Build.Tests
 {
@@ -214,6 +216,43 @@ namespace Horde.Build.Tests
 			Assert.AreEqual(telemetry[0].Telemetry[0].JobName, "hello2");
 		}
 
+		[TestMethod]
+		public async Task TestReservationPerfModel()
+		{
+
+
+			List<string> deviceModels = new List<string>() { "TestDevicePlatform1:TestDevicePlatform1_Model2", "TestDevicePlatform1_Model3" };
+			LegacyCreateReservationRequest request = await SetupReservationTest("TestDevicePool1", String.Join(';', deviceModels));
+
+			// update device model
+			UpdateDeviceRequest updateRequest = new UpdateDeviceRequest() { ModelId = "TestDevicePlatform1_Model2" };
+			await DeviceController!.UpdateDeviceAsync("testdevice1_platform1_testdevicepool1", updateRequest);
+
+			// create a reservation
+			GetLegacyReservationResponse reservation = ResultToValue(await DeviceController!.CreateDeviceReservationV1Async(request));
+			Assert.AreEqual(1, reservation.DeviceNames.Length);
+			Assert.AreEqual("hello2", reservation.JobName);
+			Assert.AreEqual("abcd", reservation.StepId);			
+
+			// get the device in the reservation, and make sure it is the right platform and an acceptable model
+			GetLegacyDeviceResponse device = ResultToValue(await DeviceController!.GetDeviceV1Async(reservation.DeviceNames[0]));
+			Assert.AreEqual("TestDevicePlatform1", device.Type);
+
+			string? firstModel = device.Model;
+			Assert.IsTrue(firstModel == "TestDevicePlatform1_Model2" || firstModel == "TestDevicePlatform1_Model3");
+
+			// get a 2nd reservation
+			reservation = ResultToValue(await DeviceController!.CreateDeviceReservationV1Async(request));
+			Assert.AreEqual(1, reservation.DeviceNames.Length);
+
+			device = ResultToValue(await DeviceController!.GetDeviceV1Async(reservation.DeviceNames[0]));
+			Assert.AreEqual("TestDevicePlatform1", device.Type);
+
+			string? secondModel = device.Model;
+			Assert.IsTrue(secondModel != firstModel && (secondModel == "TestDevicePlatform1_Model2" || secondModel == "TestDevicePlatform1_Model3"));
+		}
+
+
 		[TestMethod]
 		public async Task TestProblemDevice()
 		{
diff --git a/Engine/Source/Programs/Horde/Horde.Build/Devices/DevicesController.cs b/Engine/Source/Programs/Horde/Horde.Build/Devices/DevicesController.cs
--- a/Engine/Source/Programs/Horde/Horde.Build/Devices/DevicesController.cs
+++ b/Engine/Source/Programs/Horde/Horde.Build/Devices/DevicesController.cs
@@ -915,26 +915,19 @@ namespace Horde.Build.Devices
 				}
 				else
 				{
-					string? model = platform.Models?.FirstOrDefault(x => x == constraint);
-					if (model == null)
-					{
-						return NotFound($"Platform {platform.Id} has no model {model}");
-					}
+					string[] requestedModels = constraint.Split(';');
 
-					string modelPerfSpec = "Minimum";
-					string? specModel = null;
-
-					if (mapV1.PerfSpecHighMap.TryGetValue(platform.Id, out specModel))
+					if (requestedModels.Length > 0)
 					{
-						if (model == specModel)
+						List<string>? models = platform.Models?.Where(x => requestedModels.Contains(x, StringComparer.OrdinalIgnoreCase)).ToList();
+
+						if (models == null || models.Count == 0)
 						{
-							modelPerfSpec = "High";
+							return NotFound($"Invalid model constraint for platform {platform.Id}: {constraint}");
 						}
-					}
-
-					perfSpecs.Add(modelPerfSpec);
 
-					includeModels.Add(constraint);
+						includeModels = models;
+					}
 				}
 
 				requestedDevices.Add(new DeviceRequestData(platformId, platformName, includeModels, excludeModels));
