Horde Storage - Changed so cassandra (non scylla) version of content id and blob index avoids using UDTs for their primary key as this is not performant for basic cassandra (and thus not support in cosmos).
When using scylla there are no changes as UDTs are great and adds very little (a few bytes extra data is all)

diff --git a/Engine/Source/Programs/Horde/HordeStorage/Horde.Storage/Horde.Storage.Startup.cs b/Engine/Source/Programs/Horde/HordeStorage/Horde.Storage/Horde.Storage.Startup.cs
--- a/Engine/Source/Programs/Horde/HordeStorage/Horde.Storage/Horde.Storage.Startup.cs
+++ b/Engine/Source/Programs/Horde/HordeStorage/Horde.Storage/Horde.Storage.Startup.cs
@@ -275,8 +275,8 @@ namespace Horde.Storage
             {
                 store = ActivatorUtilities.CreateInstance<MemoryCachedContentIdStore>(provider, store);
             }
-			return store;
-		}
+            return store;
+        }
 
         private IScyllaSessionManager ScyllaFactory(IServiceProvider provider)
         {
@@ -309,41 +309,44 @@ namespace Horde.Storage
                 string[] contactPoints = connectionStringBuilder.ContactPoints;
 
                 // Connect to cassandra cluster using TLSv1.2.
-                SSLOptions sslOptions = new SSLOptions(SslProtocols.None, false,
-                    (sender, certificate, chain, sslPolicyErrors) =>
-                    {
-                        if (sslPolicyErrors == SslPolicyErrors.None)
+                if (settings.UseSSL)
+                {
+                    SSLOptions sslOptions = new SSLOptions(SslProtocols.None, false,
+                        (sender, certificate, chain, sslPolicyErrors) =>
                         {
-                            return true;
-                        }
-
-                        Logger.Error("Certificate error: {0}", sslPolicyErrors);
-                        // Do not allow this client to communicate with unauthenticated servers.
-                        return false;
-                    });
+                            if (sslPolicyErrors == SslPolicyErrors.None)
+                            {
+                                return true;
+                            }
+
+                            Logger.Error("Certificate error: {0}", sslPolicyErrors);
+                            // Do not allow this client to communicate with unauthenticated servers.
+                            return false;
+                        });
 
-                // Prepare a map to resolve the host name from the IP address.
-                Dictionary<IPAddress, string> hostNameByIp = new Dictionary<IPAddress, string>();
-                foreach (string contactPoint in contactPoints)
-                {
-                    IPAddress[] resolvedIps = Dns.GetHostAddresses(contactPoint);
-                    foreach (IPAddress resolvedIp in resolvedIps)
+                    // Prepare a map to resolve the host name from the IP address.
+                    Dictionary<IPAddress, string> hostNameByIp = new Dictionary<IPAddress, string>();
+                    foreach (string contactPoint in contactPoints)
                     {
-                        hostNameByIp[resolvedIp] = contactPoint;
+                        IPAddress[] resolvedIps = Dns.GetHostAddresses(contactPoint);
+                        foreach (IPAddress resolvedIp in resolvedIps)
+                        {
+                            hostNameByIp[resolvedIp] = contactPoint;
+                        }
                     }
-                }
 
-                sslOptions.SetHostNameResolver((ipAddress) =>
-                {
-                    if (hostNameByIp.TryGetValue(ipAddress, out string? resolvedName))
+                    sslOptions.SetHostNameResolver((ipAddress) =>
                     {
-                        return resolvedName;
-                    }
-                    IPHostEntry hostEntry = Dns.GetHostEntry(ipAddress.ToString());
-                    return hostEntry.HostName;
-                });
+                        if (hostNameByIp.TryGetValue(ipAddress, out string? resolvedName))
+                        {
+                            return resolvedName;
+                        }
+                        IPHostEntry hostEntry = Dns.GetHostEntry(ipAddress.ToString());
+                        return hostEntry.HostName;
+                    });
 
-                clusterBuilder = clusterBuilder.WithSSL(sslOptions);
+                    clusterBuilder = clusterBuilder.WithSSL(sslOptions);
+                }
             }
             Cluster cluster = clusterBuilder.Build();
 
diff --git a/Engine/Source/Programs/Horde/HordeStorage/Horde.Storage/Implementation/Blob/ScyllaBlobIndex.cs b/Engine/Source/Programs/Horde/HordeStorage/Horde.Storage/Implementation/Blob/ScyllaBlobIndex.cs
--- a/Engine/Source/Programs/Horde/HordeStorage/Horde.Storage/Implementation/Blob/ScyllaBlobIndex.cs
+++ b/Engine/Source/Programs/Horde/HordeStorage/Horde.Storage/Implementation/Blob/ScyllaBlobIndex.cs
@@ -27,9 +27,10 @@ public class ScyllaBlobIndex : IBlobIndex
         _session = scyllaSessionManager.GetSessionForReplicatedKeyspace();
         _mapper = new Mapper(_session);
 
-        _session.Execute(new SimpleStatement(@"CREATE TABLE IF NOT EXISTS blob_index (
+        string blobType = scyllaSessionManager.IsCassandra ? "blob" : "frozen<blob_identifier>";
+        _session.Execute(new SimpleStatement(@$"CREATE TABLE IF NOT EXISTS blob_index (
             namespace text,
-            blob_id frozen<blob_identifier>,
+            blob_id {blobType},
             regions set<text>,
             references set<frozen<object_reference>>,
             PRIMARY KEY ((namespace, blob_id))
@@ -43,8 +44,16 @@ public class ScyllaBlobIndex : IBlobIndex
         using IScope scope = Tracer.Instance.StartActive("scylla.insert_blob_index");
         scope.Span.ResourceName = $"{ns}.{id}";
 
-        await _mapper.UpdateAsync<ScyllaBlobIndexTable>("SET regions = regions + ? WHERE namespace = ? AND blob_id = ?",
-            new string[] { region }, ns.ToString(), new ScyllaBlobIdentifier(id));
+        if (_scyllaSessionManager.IsScylla)
+        {
+            await _mapper.UpdateAsync<ScyllaBlobIndexTable>("SET regions = regions + ? WHERE namespace = ? AND blob_id = ?",
+                new string[] { region }, ns.ToString(), new ScyllaBlobIdentifier(id));
+        }
+        else
+        {
+            await _mapper.UpdateAsync<CassandraBlobIndexTable>("SET regions = regions + ? WHERE namespace = ? AND blob_id = ?",
+                new string[] { region }, ns.ToString(), id.HashData);
+        }
     }
 
     public async Task<BlobInfo?> GetBlobInfo(NamespaceId ns, BlobIdentifier id, BlobIndexFlags flags = BlobIndexFlags.None)
@@ -54,17 +63,32 @@ public class ScyllaBlobIndex : IBlobIndex
 
         bool includeReferences = (flags & BlobIndexFlags.IncludeReferences) != 0;
 
-        ScyllaBlobIndexTable o;
-        if (includeReferences)
+        if (_scyllaSessionManager.IsScylla)
         {
-            o = await _mapper.SingleOrDefaultAsync<ScyllaBlobIndexTable>("SELECT namespace, blob_id, regions, references FROM blob_index WHERE namespace = ? AND blob_id = ?", ns.ToString(), new ScyllaBlobIdentifier(id));
+            ScyllaBlobIndexTable o;
+            if (includeReferences)
+            {
+                o = await _mapper.SingleOrDefaultAsync<ScyllaBlobIndexTable>("SELECT namespace, blob_id, regions, references FROM blob_index WHERE namespace = ? AND blob_id = ?", ns.ToString(), new ScyllaBlobIdentifier(id));
+            }
+            else
+            {
+                o = await _mapper.SingleOrDefaultAsync<ScyllaBlobIndexTable>("SELECT namespace, blob_id, regions FROM blob_index WHERE namespace = ? AND blob_id = ?", ns.ToString(), new ScyllaBlobIdentifier(id));
+            }
+            return o?.ToBlobInfo();
         }
         else
         {
-            o = await _mapper.SingleOrDefaultAsync<ScyllaBlobIndexTable>("SELECT namespace, blob_id, regions FROM blob_index WHERE namespace = ? AND blob_id = ?", ns.ToString(), new ScyllaBlobIdentifier(id));
+            CassandraBlobIndexTable o;
+            if (includeReferences)
+            {
+                o = await _mapper.SingleOrDefaultAsync<CassandraBlobIndexTable>("SELECT namespace, blob_id, regions, references FROM blob_index WHERE namespace = ? AND blob_id = ?", ns.ToString(), id.HashData);
+            }
+            else
+            {
+                o = await _mapper.SingleOrDefaultAsync<CassandraBlobIndexTable>("SELECT namespace, blob_id, regions FROM blob_index WHERE namespace = ? AND blob_id = ?", ns.ToString(), id.HashData);
+            }
+            return o?.ToBlobInfo();
         }
-
-        return o?.ToBlobInfo();
     }
 
     public async Task<bool> RemoveBlobFromIndex(NamespaceId ns, BlobIdentifier id)
@@ -72,8 +96,16 @@ public class ScyllaBlobIndex : IBlobIndex
         using IScope scope = Tracer.Instance.StartActive("scylla.remove_from_blob_index");
         scope.Span.ResourceName = $"{ns}.{id}";
 
-        await _mapper.DeleteAsync<ScyllaBlobIndexTable>("WHERE namespace = ? AND blob_id = ?", ns.ToString(),
-            new ScyllaBlobIdentifier(id));
+        if (_scyllaSessionManager.IsScylla)
+        {
+            await _mapper.DeleteAsync<ScyllaBlobIndexTable>("WHERE namespace = ? AND blob_id = ?", ns.ToString(),
+                new ScyllaBlobIdentifier(id));
+        }
+        else
+        {
+            await _mapper.DeleteAsync<CassandraBlobIndexTable>("WHERE namespace = ? AND blob_id = ?", ns.ToString(),
+                id.HashData);
+        }
 
         return true;
     }
@@ -84,7 +116,14 @@ public class ScyllaBlobIndex : IBlobIndex
         using IScope scope = Tracer.Instance.StartActive("scylla.remove_blob_index_region");
         scope.Span.ResourceName = $"{ns}.{id}";
 
-        await _mapper.UpdateAsync<ScyllaBlobIndexTable>("SET regions = regions - ? WHERE namespace = ? AND blob_id = ?", new string[] { region }, ns.ToString(), new ScyllaBlobIdentifier(id));
+        if (_scyllaSessionManager.IsScylla)
+        {
+            await _mapper.UpdateAsync<ScyllaBlobIndexTable>("SET regions = regions - ? WHERE namespace = ? AND blob_id = ?", new string[] { region }, ns.ToString(), new ScyllaBlobIdentifier(id));
+        }
+        else
+        {
+            await _mapper.UpdateAsync<CassandraBlobIndexTable>("SET regions = regions - ? WHERE namespace = ? AND blob_id = ?", new string[] { region }, ns.ToString(), id.HashData);
+        }
 
         BlobInfo? blobInfo = await GetBlobInfo(ns, id);
         if (blobInfo != null && blobInfo.Regions.Count == 0)
@@ -109,11 +148,25 @@ public class ScyllaBlobIndex : IBlobIndex
         for (int i = 0; i < blobs.Length; i++)
         {
             BlobIdentifier id = blobs[i];
-            refUpdateTasks[i] = _mapper.UpdateAsync<ScyllaBlobIndexTable>(
-                "SET references = references + ? WHERE namespace = ? AND blob_id = ?",
-                new[] { reference },
-                nsAsString,
-                new ScyllaBlobIdentifier(id));
+            Task updateTask;
+            if (_scyllaSessionManager.IsScylla)
+            {
+                updateTask = _mapper.UpdateAsync<ScyllaBlobIndexTable>(
+                    "SET references = references + ? WHERE namespace = ? AND blob_id = ?",
+                    new[] { reference },
+                    nsAsString,
+                    new ScyllaBlobIdentifier(id));
+            }
+            else
+            {
+                updateTask = _mapper.UpdateAsync<CassandraBlobIndexTable>(
+                    "SET references = references + ? WHERE namespace = ? AND blob_id = ?",
+                    new[] { reference },
+                    nsAsString,
+                    id.HashData);
+            }
+
+            refUpdateTasks[i] = updateTask;
         }
 
         await Task.WhenAll(refUpdateTasks);
@@ -124,9 +177,20 @@ public class ScyllaBlobIndex : IBlobIndex
         using IScope _ = Tracer.Instance.StartActive("scylla.get_all_blobs");
 
         string cqlOptions = _scyllaSessionManager.IsScylla ? "BYPASS CACHE" : "";
-        foreach (ScyllaBlobIndexTable blobIndex in await _mapper.FetchAsync<ScyllaBlobIndexTable>(cqlOptions))
+
+        if (_scyllaSessionManager.IsScylla)
+        {
+            foreach (ScyllaBlobIndexTable blobIndex in await _mapper.FetchAsync<ScyllaBlobIndexTable>(cqlOptions))
+            {
+                yield return blobIndex.ToBlobInfo();
+            }
+        }
+        else
         {
-            yield return blobIndex.ToBlobInfo();
+            foreach (CassandraBlobIndexTable blobIndex in await _mapper.FetchAsync<CassandraBlobIndexTable>(cqlOptions))
+            {
+                yield return blobIndex.ToBlobInfo();
+            }
         }
     }
 }
@@ -155,7 +219,7 @@ class ScyllaBlobIndexTable
 
     [Cassandra.Mapping.Attributes.PartitionKey]
     [Cassandra.Mapping.Attributes.Column("blob_id")]
-    public ScyllaBlobIdentifier BlobId { get;set; }
+    public ScyllaBlobIdentifier BlobId { get; set; }
 
     [Cassandra.Mapping.Attributes.Column("regions")]
     public HashSet<string>? Regions { get; set; }
@@ -174,3 +238,47 @@ class ScyllaBlobIndexTable
         };
     }
 }
+
+[Cassandra.Mapping.Attributes.Table("blob_index")]
+class CassandraBlobIndexTable
+{
+    public CassandraBlobIndexTable()
+    {
+        Namespace = null!;
+        BlobId = null!;
+        Regions = null;
+        References = null;
+    }
+
+    public CassandraBlobIndexTable(string @namespace, BlobIdentifier blobId, HashSet<string> regions, List<ScyllaObjectReference> references)
+    {
+        Namespace = @namespace;
+        BlobId =  blobId.HashData;
+        Regions = regions;
+        References = references;
+    }
+
+    [Cassandra.Mapping.Attributes.PartitionKey]
+    public string Namespace { get; set; }
+
+    [Cassandra.Mapping.Attributes.PartitionKey]
+    [Cassandra.Mapping.Attributes.Column("blob_id")]
+    public byte[] BlobId { get; set; }
+
+    [Cassandra.Mapping.Attributes.Column("regions")]
+    public HashSet<string>? Regions { get; set; }
+
+    [Cassandra.Mapping.Attributes.Column("references")]
+    public List<ScyllaObjectReference>? References { get; set; }
+
+    public BlobInfo ToBlobInfo()
+    {
+        return new BlobInfo
+        {
+            Regions = Regions ?? new HashSet<string>(),
+            BlobIdentifier = new BlobIdentifier(BlobId),
+            Namespace = new NamespaceId(Namespace),
+            References = References != null ? References.Select(reference => reference.AsTuple()).ToList() : new List<(BucketId, IoHashKey)>()
+        };
+    }
+}
\ No newline at end of file
diff --git a/Engine/Source/Programs/Horde/HordeStorage/Horde.Storage/Implementation/ContentId/ScyllaContentIdStore.cs b/Engine/Source/Programs/Horde/HordeStorage/Horde.Storage/Implementation/ContentId/ScyllaContentIdStore.cs
--- a/Engine/Source/Programs/Horde/HordeStorage/Horde.Storage/Implementation/ContentId/ScyllaContentIdStore.cs
+++ b/Engine/Source/Programs/Horde/HordeStorage/Horde.Storage/Implementation/ContentId/ScyllaContentIdStore.cs
@@ -17,18 +17,21 @@ namespace Horde.Storage.Implementation
         private readonly ISession _session;
         private readonly IBlobService _blobStore;
         private readonly Mapper _mapper;
+        private readonly IScyllaSessionManager _scyllaSessionManager;
 
         public ScyllaContentIdStore(IScyllaSessionManager scyllaSessionManager, IBlobService blobStore)
         {
+            _scyllaSessionManager = scyllaSessionManager;
             _session = scyllaSessionManager.GetSessionForReplicatedKeyspace();
             _blobStore = blobStore;
 
             _mapper = new Mapper(_session);
 
-            _session.Execute(new SimpleStatement(@"CREATE TABLE IF NOT EXISTS content_id (
-                content_id frozen<blob_identifier>,
+            string blobType = scyllaSessionManager.IsCassandra ? "blob" : "frozen<blob_identifier>";
+            _session.Execute(new SimpleStatement(@$"CREATE TABLE IF NOT EXISTS content_id (
+                content_id {blobType},
                 content_weight int, 
-                chunks set<frozen<blob_identifier>>, 
+                chunks set<{blobType}>, 
                 PRIMARY KEY ((content_id), content_weight)
             );"
             ));
@@ -50,26 +53,53 @@ namespace Horde.Storage.Implementation
                 using IScope contentIdFetchScope = Tracer.Instance.StartActive("ScyllaContentIdStore.FetchContentId");
                 scope.Span.ResourceName = contentId.ToString();
 
-                // lower content_weight means its a better candidate to resolve to
-                foreach (ScyllaContentId? resolvedContentId in await _mapper.FetchAsync<ScyllaContentId>("WHERE content_id = ? ORDER BY content_weight DESC", new ScyllaBlobIdentifier(contentId)))
+                if (_scyllaSessionManager.IsScylla)
                 {
-                    if (resolvedContentId == null)
+                    // lower content_weight means its a better candidate to resolve to
+                    foreach (ScyllaContentId? resolvedContentId in await _mapper.FetchAsync<ScyllaContentId>("WHERE content_id = ? ORDER BY content_weight DESC", new ScyllaBlobIdentifier(contentId)))
                     {
-                        throw new InvalidContentIdException(contentId);
-                    }
+                        if (resolvedContentId == null)
+                        {
+                            throw new InvalidContentIdException(contentId);
+                        }
 
-                    BlobIdentifier[] blobs = resolvedContentId.Chunks.Select(b => b.AsBlobIdentifier()).ToArray();
+                        BlobIdentifier[] blobs = resolvedContentId.Chunks.Select(b => b.AsBlobIdentifier()).ToArray();
+
+                        {
+                            using IScope _ = Tracer.Instance.StartActive("ScyllaContentIdStore.FindMissingBlobs");
 
+                            BlobIdentifier[] missingBlobs = await _blobStore.FilterOutKnownBlobs(ns, blobs);
+                            if (missingBlobs.Length == 0)
+                            {
+                                return blobs;
+                            }
+                        }
+                        // blobs are missing continue testing with the next content id in the weighted list as that might exist
+                    }
+                }
+                else
+                {
+                    // lower content_weight means its a better candidate to resolve to
+                    foreach (CassandraContentId? resolvedContentId in await _mapper.FetchAsync<CassandraContentId>("WHERE content_id = ? ORDER BY content_weight DESC", contentId.HashData))
                     {
-                        using IScope _ = Tracer.Instance.StartActive("ScyllaContentIdStore.FindMissingBlobs");
+                        if (resolvedContentId == null)
+                        {
+                            throw new InvalidContentIdException(contentId);
+                        }
+
+                        BlobIdentifier[] blobs = resolvedContentId.Chunks.Select(b => b.AsBlobIdentifier()).ToArray();
 
-                        BlobIdentifier[] missingBlobs = await _blobStore.FilterOutKnownBlobs(ns, blobs);
-                        if (missingBlobs.Length == 0)
                         {
-                            return blobs;
+                            using IScope _ = Tracer.Instance.StartActive("ScyllaContentIdStore.FindMissingBlobs");
+
+                            BlobIdentifier[] missingBlobs = await _blobStore.FilterOutKnownBlobs(ns, blobs);
+                            if (missingBlobs.Length == 0)
+                            {
+                                return blobs;
+                            }
                         }
+                        // blobs are missing continue testing with the next content id in the weighted list as that might exist
                     }
-                    // blobs are missing continue testing with the next content id in the weighted list as that might exist
                 }
             }
             
@@ -91,7 +121,15 @@ namespace Horde.Storage.Implementation
 
         public async Task Put(NamespaceId ns, ContentId contentId, BlobIdentifier blobIdentifier, int contentWeight)
         {
-            await _mapper.UpdateAsync<ScyllaContentId>("SET chunks = ? WHERE content_id = ? AND content_weight = ?", new [] {new ScyllaBlobIdentifier(blobIdentifier)}, new ScyllaBlobIdentifier(contentId), contentWeight);
+            if (_scyllaSessionManager.IsScylla)
+            {
+                await _mapper.UpdateAsync<ScyllaContentId>("SET chunks = ? WHERE content_id = ? AND content_weight = ?", new [] {new ScyllaBlobIdentifier(blobIdentifier)}, new ScyllaBlobIdentifier(contentId), contentWeight);
+
+            }
+            else
+            {
+                await _mapper.UpdateAsync<CassandraContentId>("SET chunks = ? WHERE content_id = ? AND content_weight = ?", new [] {new ScyllaBlobIdentifier(blobIdentifier)}, contentId.HashData, contentWeight);
+            }
         }
     }
 
@@ -120,4 +158,31 @@ namespace Horde.Storage.Implementation
 
         public ScyllaBlobIdentifier[] Chunks { get; set; } = Array.Empty<ScyllaBlobIdentifier>();
     }
+
+    
+    [Cassandra.Mapping.Attributes.Table("content_id")]
+    public class CassandraContentId
+    {
+        public CassandraContentId()
+        {
+
+        }
+
+        public CassandraContentId(BlobIdentifier contentId, BlobIdentifier[] chunks, int contentWeight)
+        {
+            ContentId = contentId.HashData;
+            ContentWeight = contentWeight;
+            Chunks = chunks.Select(b => new ScyllaBlobIdentifier(b)).ToArray();
+        }
+
+        [Cassandra.Mapping.Attributes.PartitionKey]
+        [Cassandra.Mapping.Attributes.Column("content_id")]
+        public byte[]? ContentId { get; set; }
+
+        [Cassandra.Mapping.Attributes.ClusteringKey]
+        [Cassandra.Mapping.Attributes.Column("content_weight")]
+        public int? ContentWeight { get; set; }
+
+        public ScyllaBlobIdentifier[] Chunks { get; set; } = Array.Empty<ScyllaBlobIdentifier>();
+    }
 }
diff --git a/Engine/Source/Programs/Horde/HordeStorage/Horde.Storage/Implementation/Objects/ScyllaSettings.cs b/Engine/Source/Programs/Horde/HordeStorage/Horde.Storage/Implementation/Objects/ScyllaSettings.cs
--- a/Engine/Source/Programs/Horde/HordeStorage/Horde.Storage/Implementation/Objects/ScyllaSettings.cs
+++ b/Engine/Source/Programs/Horde/HordeStorage/Horde.Storage/Implementation/Objects/ScyllaSettings.cs
@@ -52,4 +52,10 @@ public class ScyllaSettings
     /// Whether a Cosmos DB Cassandra layer from Azure should be used.
     /// </summary>
     public bool UseAzureCosmosDB { get; set; } = false;
+
+    /// <summary>
+    /// Set to false to disable use of SSL, this is only used when connecting to a CosmosDB instance
+    /// </summary>
+    public bool UseSSL { get; set; } = true;
+
 }
diff --git a/Engine/Source/Programs/Horde/HordeStorage/Horde.Storage/Tests/Functional/Blob/BlobIndexTests.cs b/Engine/Source/Programs/Horde/HordeStorage/Horde.Storage/Tests/Functional/Blob/BlobIndexTests.cs
--- a/Engine/Source/Programs/Horde/HordeStorage/Horde.Storage/Tests/Functional/Blob/BlobIndexTests.cs
+++ b/Engine/Source/Programs/Horde/HordeStorage/Horde.Storage/Tests/Functional/Blob/BlobIndexTests.cs
@@ -245,6 +245,36 @@ namespace Horde.Storage.FunctionalTests.Storage
         }
     }
 
+    [TestClass]
+    public class CassandraBlobIndexTests : BlobIndexTests
+    {
+        protected override IEnumerable<KeyValuePair<string, string>> GetSettings()
+        {
+            return new[]
+            {
+                new KeyValuePair<string, string>("Horde_Storage:BlobIndexImplementation", HordeStorageSettings.BlobIndexImplementations.Scylla.ToString()),
+                new KeyValuePair<string, string>("Scylla:UseAzureCosmosDB", "true"),
+                new KeyValuePair<string, string>("Scylla:UseSSL", "false"),
+            };
+        }
+
+        protected override Task Seed(IServiceProvider provider)
+        {
+            IScyllaSessionManager scyllaSessionManager = provider.GetService<IScyllaSessionManager>()!;
+
+            Assert.IsTrue(scyllaSessionManager.IsCassandra);
+            return Task.CompletedTask;
+        }
+
+        protected override async Task Teardown(IServiceProvider provider)
+        {
+            IScyllaSessionManager scyllaSessionManager = provider.GetService<IScyllaSessionManager>()!;
+
+            ISession replicatedKeyspace = scyllaSessionManager.GetSessionForReplicatedKeyspace();
+            await replicatedKeyspace.ExecuteAsync(new SimpleStatement("DROP TABLE IF EXISTS blob_index"));
+        }
+    }
+
     [TestClass]
     public class MongoBlobIndexTests : BlobIndexTests
     {
diff --git a/Engine/Source/Programs/Horde/HordeStorage/Horde.Storage/Tests/Functional/Objects/ReferencesTests.cs b/Engine/Source/Programs/Horde/HordeStorage/Horde.Storage/Tests/Functional/Objects/ReferencesTests.cs
--- a/Engine/Source/Programs/Horde/HordeStorage/Horde.Storage/Tests/Functional/Objects/ReferencesTests.cs
+++ b/Engine/Source/Programs/Horde/HordeStorage/Horde.Storage/Tests/Functional/Objects/ReferencesTests.cs
@@ -90,6 +90,65 @@ namespace Horde.Storage.FunctionalTests.References
 
         }
     }
+    
+    [TestClass]
+    public class CassandraReferencesTests : ReferencesTests
+    {
+        protected override IEnumerable<KeyValuePair<string, string>> GetSettings()
+        {
+            return new[]
+            {
+                new KeyValuePair<string, string>("Horde_Storage:ReferencesDbImplementation", HordeStorageSettings.ReferencesDbImplementations.Scylla.ToString()),
+                new KeyValuePair<string, string>("Horde_Storage:ContentIdStoreImplementation", HordeStorageSettings.ContentIdStoreImplementations.Scylla.ToString()),
+                new KeyValuePair<string, string>("Horde_Storage:ReplicationLogWriterImplementation", HordeStorageSettings.ReplicationLogWriterImplementations.Scylla.ToString()),
+                new KeyValuePair<string, string>("Scylla:UseAzureCosmosDB", "true"),
+                new KeyValuePair<string, string>("Scylla:UseSSL", "false"),
+            };
+        }
+
+        protected override async Task SeedDb(IServiceProvider provider)
+        {
+            IReferencesStore referencesStore = provider.GetService<IReferencesStore>()!;
+            //verify we are using the expected store
+            Assert.IsTrue(referencesStore.GetType() == typeof(ScyllaReferencesStore));
+
+            IContentIdStore contentIdStore = provider.GetService<IContentIdStore>()!;
+            //verify we are using the expected store
+            if (contentIdStore is MemoryCachedContentIdStore memoryStore)
+            {
+                memoryStore.Clear();
+                contentIdStore = memoryStore.GetUnderlyingContentIdStore();
+            }
+            Assert.IsTrue(contentIdStore.GetType() == typeof(ScyllaContentIdStore));
+
+            IReplicationLog replicationLog = provider.GetService<IReplicationLog>()!;
+            //verify we are using the replication log writer
+            Assert.IsTrue(replicationLog.GetType() == typeof(ScyllaReplicationLog));
+
+            await SeedTestData();
+        }
+
+        protected override async Task TeardownDb(IServiceProvider provider)
+        {
+            IScyllaSessionManager scyllaSessionManager = provider.GetService<IScyllaSessionManager>()!;
+
+            ISession replicatedKeyspace = scyllaSessionManager.GetSessionForReplicatedKeyspace();
+            await replicatedKeyspace.ExecuteAsync(new SimpleStatement("DROP TABLE IF EXISTS objects"));
+            await replicatedKeyspace.ExecuteAsync(new SimpleStatement("DROP TABLE IF EXISTS content_id"));
+
+            // we need to clear out the state we modified in the replication log, otherwise the replication log tests will fail
+            ISession localKeyspace = scyllaSessionManager.GetSessionForLocalKeyspace();
+            await Task.WhenAll(
+                // remove replication log table as we expect it to be empty when starting the tests
+                localKeyspace.ExecuteAsync(new SimpleStatement("DROP TABLE IF EXISTS replication_log;")),
+                // remove the snapshots
+                localKeyspace.ExecuteAsync(new SimpleStatement("DROP TABLE IF EXISTS replication_snapshot;")),
+                // remove the namespaces
+                localKeyspace.ExecuteAsync(new SimpleStatement("DROP TABLE IF EXISTS replication_namespace;"))
+            );
+
+        }
+    }
 
     [TestClass]
     public class MongoReferencesTests : ReferencesTests
