Horde: Remove shutdown dialog timeout for Windows agent shutdowns

While convenient for any logged in user, this adds delays and cost to machine management.

diff --git a/Engine/Source/Programs/Horde/Horde.Agent/Utility/Shutdown.cs b/Engine/Source/Programs/Horde/Horde.Agent/Utility/Shutdown.cs
--- a/Engine/Source/Programs/Horde/Horde.Agent/Utility/Shutdown.cs
+++ b/Engine/Source/Programs/Horde/Horde.Agent/Utility/Shutdown.cs
@@ -111,7 +111,8 @@ namespace Horde.Agent.Utility
 					return false;
 				}
 
-				if (!InitiateSystemShutdownEx(null, "HordeAgent has initiated shutdown", 10, true, restartAfterShutdown, SHTDN_REASON_MAJOR_APPLICATION | SHTDN_REASON_MINOR_MAINTENANCE))
+				uint dialogTimeout = 0; // The length of time that the shutdown dialog box should be displayed, in seconds
+				if (!InitiateSystemShutdownEx(null, "HordeAgent has initiated shutdown", dialogTimeout, true, restartAfterShutdown, SHTDN_REASON_MAJOR_APPLICATION | SHTDN_REASON_MINOR_MAINTENANCE))
 				{
 					logger.LogError("Shutdown failed (0x{Code:x8})", Marshal.GetLastWin32Error());
 					return false;
