Horde: Conifgurable shared device checkout duration and device collection upgrade method
#rnx

diff --git a/Engine/Source/Programs/Horde/Horde.Build/Devices/DeviceCollection.cs b/Engine/Source/Programs/Horde/Horde.Build/Devices/DeviceCollection.cs
--- a/Engine/Source/Programs/Horde/Horde.Build/Devices/DeviceCollection.cs
+++ b/Engine/Source/Programs/Horde/Horde.Build/Devices/DeviceCollection.cs
@@ -5,11 +5,13 @@ using System.Collections.Generic;
 using System.Linq;
 using System.Threading.Tasks;
 using Horde.Build.Acls;
+using Horde.Build.Jobs;
 using Horde.Build.Projects;
 using Horde.Build.Server;
 using Horde.Build.Streams;
 using Horde.Build.Users;
 using Horde.Build.Utilities;
+using Microsoft.Extensions.Logging;
 using MongoDB.Bson;
 using MongoDB.Bson.Serialization.Attributes;
 using MongoDB.Bson.Serialization.Options;
@@ -183,6 +185,8 @@ namespace Horde.Build.Devices
 		/// </summary>
 		class DeviceDocument : IDevice
 		{
+			public static int CurrentVersion = 1;
+
 			[BsonRequired, BsonId]
 			public DeviceId Id { get; set; }
 
@@ -236,6 +240,9 @@ namespace Horde.Build.Devices
 			[BsonIgnoreIfNull]
 			public Acl? Acl { get; set; }
 
+			[BsonIgnoreIfNull]
+			public int? Version { get; set; }
+
 			[BsonConstructor]
 			private DeviceDocument()
 			{
@@ -251,6 +258,7 @@ namespace Horde.Build.Devices
 				Address = address;
 				ModelId = modelId;
 				ModifiedByUser = userId?.ToString();
+				Version = CurrentVersion;
 			}
 		}
 
@@ -497,13 +505,15 @@ namespace Horde.Build.Devices
 
 		readonly IMongoCollection<DevicePoolTelemetryDocument> _poolTelemetry;
 
-		readonly int _checkoutDays = 7;
+		readonly ILogger<DeviceCollection> _logger;
 
 		/// <summary>
 		/// Constructor
 		/// </summary>
-		public DeviceCollection(MongoService mongoService)
+		public DeviceCollection(MongoService mongoService, ILogger<DeviceCollection> logger)
 		{
+			_logger = logger;
+
 			_devices = mongoService.GetCollection<DeviceDocument>("Devices", keys => keys.Ascending(x => x.Name), unique: true);
 			_platforms = mongoService.GetCollection<DevicePlatformDocument>("Devices.Platforms", keys => keys.Ascending(x => x.Name), unique: true);
 			_pools = mongoService.GetCollection<DevicePoolDocument>("Devices.Pools", keys => keys.Ascending(x => x.Name), unique: true);
@@ -698,7 +708,6 @@ namespace Horde.Build.Devices
 		public async Task UpdateDeviceAsync(DeviceId deviceId, DevicePoolId? newPoolId, string? newName, string? newAddress, string? newModelId, string? newNotes, bool? newEnabled, bool? newProblem, bool? newMaintenance, UserId? modifiedByUserId = null)
 		{
 			UpdateDefinitionBuilder<DeviceDocument> updateBuilder = Builders<DeviceDocument>.Update;
-
 			List<UpdateDefinition<DeviceDocument>> updates = new List<UpdateDefinition<DeviceDocument>>();
 
 			DateTime utcNow = DateTime.UtcNow;
@@ -987,7 +996,7 @@ namespace Horde.Build.Devices
 		/// <summary>
 		/// Deletes expired user checkouts
 		/// </summary>
-		public async Task<List<(UserId, IDevice)>?> ExpireCheckedOutAsync()
+		public async Task<List<(UserId, IDevice)>?> ExpireCheckedOutAsync(int checkoutDays)
 		{
 			FilterDefinition<DeviceDocument> filter = Builders<DeviceDocument>.Filter.Empty;
 			filter &= Builders<DeviceDocument>.Filter.Where(x => x.CheckedOutByUser != null && x.CheckOutTime != null);
@@ -995,7 +1004,7 @@ namespace Horde.Build.Devices
 			List<DeviceDocument> checkedOutDevices = await _devices.Find(filter).ToListAsync();
 
 			DateTime utcNow = DateTime.UtcNow;
-			List<DeviceDocument> expiredDevices = checkedOutDevices.FindAll(x => (utcNow - x.CheckOutTime!.Value).TotalDays >= _checkoutDays).ToList();
+			List<DeviceDocument> expiredDevices = checkedOutDevices.FindAll(x => (utcNow - x.CheckOutTime!.Value).TotalDays >= checkoutDays).ToList();
 
 			if (expiredDevices.Count > 0)
 			{
@@ -1024,7 +1033,7 @@ namespace Horde.Build.Devices
 		/// <summary>
 		/// Gets a list of users to notify  that their device is about to expire in the next 24 hours
 		/// </summary>
-		public async Task<List<(UserId, IDevice)>?> ExpireNotificatonsAsync()
+		public async Task<List<(UserId, IDevice)>?> ExpireNotificatonsAsync(int checkoutDays)
 		{
 			FilterDefinition<DeviceDocument> filter = Builders<DeviceDocument>.Filter.Empty;
 			filter &= Builders<DeviceDocument>.Filter.Where(x => x.CheckedOutByUser != null && x.CheckoutExpiringNotificationSent != true);
@@ -1032,7 +1041,7 @@ namespace Horde.Build.Devices
 			List<DeviceDocument> checkedOutDevices = await _devices.Find(filter).ToListAsync();
 
 			DateTime utcNow = DateTime.UtcNow;
-			List<DeviceDocument> expiredDevices = checkedOutDevices.FindAll(x => (utcNow - x.CheckOutTime!.Value).TotalDays >= (_checkoutDays - 1)).ToList();
+			List<DeviceDocument> expiredDevices = checkedOutDevices.FindAll(x => (utcNow - x.CheckOutTime!.Value).TotalDays >= (checkoutDays - 1)).ToList();
 
 			if (expiredDevices.Count > 0)
 			{
@@ -1226,5 +1235,44 @@ namespace Horde.Build.Devices
 			List<DevicePoolTelemetryDocument> results = await _poolTelemetry.Find(filter).Range(index, count).ToListAsync();
 			return results.ConvertAll<IDevicePoolTelemetry>(x => x);
 		}
+
+		/// <inheritdoc/>
+		public async Task UpgradeAsync()
+		{
+			FilterDefinition<DeviceDocument> filter = Builders<DeviceDocument>.Filter.Empty;
+			filter &= Builders<DeviceDocument>.Filter.Where(x => x.Version == null || x.Version < DeviceDocument.CurrentVersion);
+
+			List<DeviceDocument> results = await _devices.Find(filter).ToListAsync();
+
+			if (results.Count > 0)
+			{
+				_logger.LogInformation("Found {Count} device douments to upgrade", results.Count);
+			}			
+
+			foreach (DeviceDocument device in results)
+			{
+				// unversioned => 1
+				if (device.Version == null && DeviceDocument.CurrentVersion == 1)
+				{
+					UpdateDefinitionBuilder<DeviceDocument> updateBuilder = Builders<DeviceDocument>.Update;
+					List<UpdateDefinition<DeviceDocument>> updates = new List<UpdateDefinition<DeviceDocument>>();
+
+					updates.Add(updateBuilder.Set(x => x.Version, DeviceDocument.CurrentVersion));
+
+					if (device.CheckOutTime != null)
+					{
+						DateTime now = DateTime.UtcNow;
+						double Days = (now - device.CheckOutTime!.Value).TotalDays;
+						if (Days > 3)
+						{							
+							updates.Add(updateBuilder.Set(x => x.CheckOutTime, now));
+							updates.Add(updateBuilder.Set(x => x.CheckoutExpiringNotificationSent, null));
+						}
+					}
+
+					await _devices.FindOneAndUpdateAsync(x => x.Id == device.Id, updateBuilder.Combine(updates));
+				}
+			}
+		}
 	}
 }
diff --git a/Engine/Source/Programs/Horde/Horde.Build/Devices/DeviceService.cs b/Engine/Source/Programs/Horde/Horde.Build/Devices/DeviceService.cs
--- a/Engine/Source/Programs/Horde/Horde.Build/Devices/DeviceService.cs
+++ b/Engine/Source/Programs/Horde/Horde.Build/Devices/DeviceService.cs
@@ -18,6 +18,7 @@ using Horde.Build.Utilities;
 using HordeCommon;
 using Microsoft.Extensions.Hosting;
 using Microsoft.Extensions.Logging;
+using Microsoft.Extensions.Options;
 using MongoDB.Bson;
 using OpenTracing;
 using OpenTracing.Util;
@@ -94,22 +95,24 @@ namespace Horde.Build.Devices
 		readonly JobService _jobService;
 		readonly StreamService _streamService;
 		readonly ProjectService _projectService;
-		readonly MongoService _mongoService;
 		IUserCollection UserCollection { get; set; }
 		readonly ILogger<DeviceService> _logger;
 		readonly IDeviceCollection _devices;
 		readonly ITicker _ticker;
 		readonly ITicker _telemetryTicker;
+		readonly IOptionsMonitor<ServerSettings> _settings;
 
 		/// <summary>
 		/// Platform map V1 singleton
 		/// </summary>
 		readonly ISingletonDocument<DevicePlatformMapV1> _platformMapSingleton;
 
+		bool runUpgrade = true;
+
 		/// <summary>
 		/// Device service constructor
 		/// </summary>
-		public DeviceService(GlobalsService globalsService, IDeviceCollection devices, ISingletonDocument<DevicePlatformMapV1> platformMapSingleton, IUserCollection userCollection, MongoService mongoService, JobService jobService, ProjectService projectService, StreamService streamService, AclService aclService, INotificationService notificationService, IClock clock, ILogger<DeviceService> logger)
+		public DeviceService(GlobalsService globalsService, IDeviceCollection devices, ISingletonDocument<DevicePlatformMapV1> platformMapSingleton, IUserCollection userCollection, JobService jobService, ProjectService projectService, StreamService streamService, AclService aclService, IOptionsMonitor<ServerSettings> settings, INotificationService notificationService, IClock clock, ILogger<DeviceService> logger)
 		{
 			UserCollection = userCollection;
 			_globalsService = globalsService;
@@ -122,8 +125,7 @@ namespace Horde.Build.Devices
 			_ticker = clock.AddSharedTicker<DeviceService>(TimeSpan.FromMinutes(1.0), TickAsync, logger);
 			_telemetryTicker = clock.AddSharedTicker("DeviceService.Telemetry", TimeSpan.FromMinutes(10.0), TickTelemetryAsync, logger);
 			_logger = logger;
-			_mongoService = mongoService;
-
+			_settings = settings;
 			_platformMapSingleton = platformMapSingleton;
 
 		}
@@ -204,6 +206,20 @@ namespace Horde.Build.Devices
 			{
 				using IScope scope = GlobalTracer.Instance.BuildSpan("DeviceService.TickAsync").StartActive();
 
+				if (runUpgrade)
+				{
+					try
+					{
+						await _devices.UpgradeAsync();
+						runUpgrade = false;
+					}
+					catch (Exception ex)
+					{
+						_logger.LogError(ex, "Exception while upgrading device collection: {Message}", ex.Message);
+					}
+
+				}
+
 				try
 				{
 					_logger.LogInformation("Expiring reservations");
@@ -219,7 +235,7 @@ namespace Horde.Build.Devices
 				{
 
 					_logger.LogInformation("Sending shared device notifications");
-					List<(UserId, IDevice)>? expireNotifications = await _devices.ExpireNotificatonsAsync();
+					List<(UserId, IDevice)>? expireNotifications = await _devices.ExpireNotificatonsAsync(_settings.CurrentValue.SharedDeviceCheckoutDays);
 					if (expireNotifications != null && expireNotifications.Count > 0)
 					{
 						foreach ((UserId, IDevice) expiredDevice in expireNotifications)
@@ -229,7 +245,7 @@ namespace Horde.Build.Devices
 					}
 
 					_logger.LogInformation("Expiring shared device checkouts");
-					List<(UserId, IDevice)>? expireCheckouts = await _devices.ExpireCheckedOutAsync();
+					List<(UserId, IDevice)>? expireCheckouts = await _devices.ExpireCheckedOutAsync(_settings.CurrentValue.SharedDeviceCheckoutDays);
 					if (expireCheckouts != null && expireCheckouts.Count > 0)
 					{
 						foreach ((UserId, IDevice) expiredDevice in expireCheckouts)
diff --git a/Engine/Source/Programs/Horde/Horde.Build/Devices/IDeviceCollection.cs b/Engine/Source/Programs/Horde/Horde.Build/Devices/IDeviceCollection.cs
--- a/Engine/Source/Programs/Horde/Horde.Build/Devices/IDeviceCollection.cs
+++ b/Engine/Source/Programs/Horde/Horde.Build/Devices/IDeviceCollection.cs
@@ -227,12 +227,14 @@ namespace Horde.Build.Devices
 		/// <summary>
 		/// Deletes user device checkouts
 		/// </summary>
-		public Task<List<(UserId, IDevice)>?> ExpireCheckedOutAsync();
+		/// <param name="checkoutDays"></param>
+		public Task<List<(UserId, IDevice)>?> ExpireCheckedOutAsync(int checkoutDays);
 
 		/// <summary>
 		/// Gets a list of users to notify whose device checkout is about to expire
 		/// </summary>
-		public Task<List<(UserId, IDevice)>?> ExpireNotificatonsAsync();
+		/// <param name="checkoutDays"></param>
+		public Task<List<(UserId, IDevice)>?> ExpireNotificatonsAsync(int checkoutDays);
 
 		/// <summary>
 		/// Find device telemetry
@@ -260,5 +262,11 @@ namespace Horde.Build.Devices
 		/// <param name="count"></param>
 		/// <returns></returns>
 		public Task<List<IDevicePoolTelemetry>> FindPoolTelemetryAsync(DateTimeOffset? minCreateTime = null, DateTimeOffset? maxCreateTime = null, int? index = null, int? count = null);
+
+		/// <summary>
+		/// Upgrades the collection
+		/// </summary>
+		public Task UpgradeAsync();
+
 	}
 }
\ No newline at end of file
diff --git a/Engine/Source/Programs/Horde/Horde.Build/ServerSettings.cs b/Engine/Source/Programs/Horde/Horde.Build/ServerSettings.cs
--- a/Engine/Source/Programs/Horde/Horde.Build/ServerSettings.cs
+++ b/Engine/Source/Programs/Horde/Horde.Build/ServerSettings.cs
@@ -510,6 +510,11 @@ namespace Horde.Build
 		/// </summary>
 		public Uri? JiraUrl { get; set; }
 
+		/// <summary>
+		/// The number of days shared device checkouts are held
+		/// </summary>
+		public int SharedDeviceCheckoutDays { get; set; } = 3;
+
 		/// <summary>
 		/// Default agent pool sizing strategy for pools that doesn't have one explicitly configured
 		/// </summary>
