Horde: Dashboard side of configurable shared checkout times
#rnx

diff --git a/Engine/Source/Programs/Horde/Horde.Build/Devices/DeviceCollection.cs b/Engine/Source/Programs/Horde/Horde.Build/Devices/DeviceCollection.cs
--- a/Engine/Source/Programs/Horde/Horde.Build/Devices/DeviceCollection.cs
+++ b/Engine/Source/Programs/Horde/Horde.Build/Devices/DeviceCollection.cs
@@ -1246,7 +1246,7 @@ namespace Horde.Build.Devices
 
 			if (results.Count > 0)
 			{
-				_logger.LogInformation("Found {Count} device douments to upgrade", results.Count);
+				_logger.LogInformation("Found {Count} device documents to upgrade", results.Count);
 			}			
 
 			foreach (DeviceDocument device in results)
diff --git a/Engine/Source/Programs/Horde/Horde.Build/Devices/DeviceMessages.cs b/Engine/Source/Programs/Horde/Horde.Build/Devices/DeviceMessages.cs
--- a/Engine/Source/Programs/Horde/Horde.Build/Devices/DeviceMessages.cs
+++ b/Engine/Source/Programs/Horde/Horde.Build/Devices/DeviceMessages.cs
@@ -342,6 +342,11 @@ namespace Horde.Build.Devices
 		/// </summary>
 		public DateTime? CheckOutTime { get; set; }
 
+		/// <summary>
+		/// When the checkout will expire
+		/// </summary>
+		public DateTime? CheckOutExpirationTime { get; set; }
+
 		/// <summary>
 		/// The last user to modifiy the device
 		/// </summary>
@@ -355,7 +360,7 @@ namespace Horde.Build.Devices
 		/// <summary>
 		/// Device response constructor
 		/// </summary>
-		public GetDeviceResponse(string id, string platformId, string poolId, string name, bool enabled, string? address, string? modelId, string? modifiedByUser, string? notes, DateTime? problemTime, DateTime? maintenanceTime, List<DeviceUtilizationTelemetry>? utilization, string? checkedOutByUser = null, DateTime? checkOutTime = null)
+		public GetDeviceResponse(string id, string platformId, string poolId, string name, bool enabled, string? address, string? modelId, string? modifiedByUser, string? notes, DateTime? problemTime, DateTime? maintenanceTime, List<DeviceUtilizationTelemetry>? utilization, string? checkedOutByUser = null, DateTime? checkOutTime = null, DateTime? checkOutExpirationTime = null)
 		{
 			Id = id;
 			Name = name;
@@ -371,6 +376,7 @@ namespace Horde.Build.Devices
 			Utilization = utilization?.Select(u => new GetDeviceUtilizationResponse(u)).ToList();
 			CheckedOutByUserId = checkedOutByUser;
 			CheckOutTime = checkOutTime;
+			CheckOutExpirationTime = checkOutExpirationTime;
 		}
 	}
 
diff --git a/Engine/Source/Programs/Horde/Horde.Build/Devices/DevicesController.cs b/Engine/Source/Programs/Horde/Horde.Build/Devices/DevicesController.cs
--- a/Engine/Source/Programs/Horde/Horde.Build/Devices/DevicesController.cs
+++ b/Engine/Source/Programs/Horde/Horde.Build/Devices/DevicesController.cs
@@ -12,6 +12,7 @@ using Horde.Build.Utilities;
 using Microsoft.AspNetCore.Authorization;
 using Microsoft.AspNetCore.Mvc;
 using Microsoft.Extensions.Logging;
+using Microsoft.Extensions.Options;
 using MongoDB.Bson;
 using MongoDB.Driver;
 
@@ -48,15 +49,18 @@ namespace Horde.Build.Devices
 		/// </summary>
 		private readonly ILogger<DevicesController> _logger;
 
+		readonly IOptionsMonitor<ServerSettings> _settings;
+
 		/// <summary>
 		/// Constructor
 		/// </summary>
-		public DevicesController(DeviceService deviceService, AclService aclService, IUserCollection userCollection, ILogger<DevicesController> logger)
+		public DevicesController(DeviceService deviceService, AclService aclService, IUserCollection userCollection, ILogger<DevicesController> logger, IOptionsMonitor<ServerSettings> settings)
 		{
 			UserCollection = userCollection;
 			_deviceService = deviceService;
 			_logger = logger;
 			_aclService = aclService;
+			_settings = settings;
 		}
 
 		// DEVICES
@@ -147,7 +151,13 @@ namespace Horde.Build.Devices
 					continue;
 				}
 
-				responses.Add(new GetDeviceResponse(device.Id.ToString(), device.PlatformId.ToString(), device.PoolId.ToString(), device.Name, device.Enabled, device.Address, device.ModelId?.ToString(), device.ModifiedByUser, device.Notes, device.ProblemTimeUtc, device.MaintenanceTimeUtc, device.Utilization, device.CheckedOutByUser, device.CheckOutTime));
+				DateTime? checkoutExpiration = null;
+				if (device.CheckOutTime != null)
+				{
+					checkoutExpiration = device.CheckOutTime.Value.AddDays(_settings.CurrentValue.SharedDeviceCheckoutDays);
+				}
+
+				responses.Add(new GetDeviceResponse(device.Id.ToString(), device.PlatformId.ToString(), device.PoolId.ToString(), device.Name, device.Enabled, device.Address, device.ModelId?.ToString(), device.ModifiedByUser, device.Notes, device.ProblemTimeUtc, device.MaintenanceTimeUtc, device.Utilization, device.CheckedOutByUser, device.CheckOutTime, checkoutExpiration));
 			}
 
 			return responses;
@@ -179,7 +189,13 @@ namespace Horde.Build.Devices
 				return Forbid();
 			}
 
-			return new GetDeviceResponse(device.Id.ToString(), device.PlatformId.ToString(), device.PoolId.ToString(), device.Name, device.Enabled, device.Address, device.ModelId?.ToString(), device.ModifiedByUser?.ToString(), device.Notes, device.ProblemTimeUtc, device.MaintenanceTimeUtc, device.Utilization, device.CheckedOutByUser, device.CheckOutTime);
+			DateTime? checkoutExpiration = null;
+			if (device.CheckOutTime != null)
+			{
+				checkoutExpiration = device.CheckOutTime.Value.AddDays(_settings.CurrentValue.SharedDeviceCheckoutDays);
+			}
+
+			return new GetDeviceResponse(device.Id.ToString(), device.PlatformId.ToString(), device.PoolId.ToString(), device.Name, device.Enabled, device.Address, device.ModelId?.ToString(), device.ModifiedByUser?.ToString(), device.Notes, device.ProblemTimeUtc, device.MaintenanceTimeUtc, device.Utilization, device.CheckedOutByUser, device.CheckOutTime, checkoutExpiration);
 		}
 
 		/// <summary>
diff --git a/Engine/Source/Programs/Horde/HordeDashboard/src/backend/Api.ts b/Engine/Source/Programs/Horde/HordeDashboard/src/backend/Api.ts
--- a/Engine/Source/Programs/Horde/HordeDashboard/src/backend/Api.ts
+++ b/Engine/Source/Programs/Horde/HordeDashboard/src/backend/Api.ts
@@ -3230,6 +3230,9 @@ export type GetDeviceResponse = {
 	/// The last time the device was checked out
 	checkOutTime?: Date | string;
 
+	/// The time the device checkout will expire
+	checkOutExpirationTime?: Date | string;
+
 	/** Id of the user which last modified this device */
 	modifiedByUser?: string;
 
diff --git a/Engine/Source/Programs/Horde/HordeDashboard/src/components/DeviceView.tsx b/Engine/Source/Programs/Horde/HordeDashboard/src/components/DeviceView.tsx
--- a/Engine/Source/Programs/Horde/HordeDashboard/src/components/DeviceView.tsx
+++ b/Engine/Source/Programs/Horde/HordeDashboard/src/components/DeviceView.tsx
@@ -97,7 +97,6 @@ const customStyles = mergeStyleSets({
 
 const pivotKeyAutomation = "pivot-key-automation";
 const pivotKeyShared = "pivot-key-shared";
-const checkoutDays = 7;
 
 const DevicePanel: React.FC = observer(() => {
 
@@ -334,8 +333,8 @@ const DevicePanel: React.FC = observer(() => {
       const device = item.device;
       let notes = device.notes ?? "";
 
-      if (device.checkOutTime && device.checkedOutByUserId) {
-         notes = `Checked out until ${getNiceTime(moment(new Date(device.checkOutTime)).add(checkoutDays, 'd').toDate())}.  ` + notes;
+      if (device.checkOutExpirationTime && device.checkedOutByUserId) {
+         notes = `Checked out until ${getNiceTime(device.checkOutExpirationTime)}.  ` + notes;
       }
 
       return <Stack verticalFill={true} verticalAlign="center">
@@ -599,8 +598,8 @@ export const CheckoutConfirmModal: React.FC<{ check: "in" | "out" | "error", dev
          case 'Name':
             return <Text variant="medium" styles={{ root: { fontFamily: "Horde Open Sans SemiBold" } }}>{platform} : {device.name}</Text>
          case 'Date':
-            if (device.checkOutTime) {
-               return <Text variant="medium">{`Checked out until ${getNiceTime(moment(new Date(device.checkOutTime!)).add(checkoutDays, 'd').toDate())}`}</Text>
+            if (device.checkOutExpirationTime) {
+               return <Text variant="medium">{`Checked out until ${getNiceTime(device.checkOutExpirationTime)}`}</Text>
             }
             return null;
       }
