Horde: Fix server compilation
#rnx

diff --git a/Engine/Source/Programs/Horde/Horde.Build/Devices/DeviceService.cs b/Engine/Source/Programs/Horde/Horde.Build/Devices/DeviceService.cs
--- a/Engine/Source/Programs/Horde/Horde.Build/Devices/DeviceService.cs
+++ b/Engine/Source/Programs/Horde/Horde.Build/Devices/DeviceService.cs
@@ -109,6 +109,11 @@ namespace Horde.Build.Devices
 
 		bool runUpgrade = true;
 
+		/// <summary>
+		/// The number of days shared device sheckouts are held
+		/// </summary>
+		public int sharedDeviceCheckoutDays => _settings.CurrentValue.SharedDeviceCheckoutDays;
+
 		/// <summary>
 		/// Device service constructor
 		/// </summary>
diff --git a/Engine/Source/Programs/Horde/Horde.Build/Devices/DevicesController.cs b/Engine/Source/Programs/Horde/Horde.Build/Devices/DevicesController.cs
--- a/Engine/Source/Programs/Horde/Horde.Build/Devices/DevicesController.cs
+++ b/Engine/Source/Programs/Horde/Horde.Build/Devices/DevicesController.cs
@@ -49,18 +49,15 @@ namespace Horde.Build.Devices
 		/// </summary>
 		private readonly ILogger<DevicesController> _logger;
 
-		readonly IOptionsMonitor<ServerSettings> _settings;
-
 		/// <summary>
 		/// Constructor
 		/// </summary>
-		public DevicesController(DeviceService deviceService, AclService aclService, IUserCollection userCollection, ILogger<DevicesController> logger, IOptionsMonitor<ServerSettings> settings)
+		public DevicesController(DeviceService deviceService, AclService aclService, IUserCollection userCollection, ILogger<DevicesController> logger)
 		{
 			UserCollection = userCollection;
 			_deviceService = deviceService;
 			_logger = logger;
 			_aclService = aclService;
-			_settings = settings;
 		}
 
 		// DEVICES
@@ -154,7 +151,7 @@ namespace Horde.Build.Devices
 				DateTime? checkoutExpiration = null;
 				if (device.CheckOutTime != null)
 				{
-					checkoutExpiration = device.CheckOutTime.Value.AddDays(_settings.CurrentValue.SharedDeviceCheckoutDays);
+					checkoutExpiration = device.CheckOutTime.Value.AddDays(_deviceService.sharedDeviceCheckoutDays);
 				}
 
 				responses.Add(new GetDeviceResponse(device.Id.ToString(), device.PlatformId.ToString(), device.PoolId.ToString(), device.Name, device.Enabled, device.Address, device.ModelId?.ToString(), device.ModifiedByUser, device.Notes, device.ProblemTimeUtc, device.MaintenanceTimeUtc, device.Utilization, device.CheckedOutByUser, device.CheckOutTime, checkoutExpiration));
@@ -192,7 +189,7 @@ namespace Horde.Build.Devices
 			DateTime? checkoutExpiration = null;
 			if (device.CheckOutTime != null)
 			{
-				checkoutExpiration = device.CheckOutTime.Value.AddDays(_settings.CurrentValue.SharedDeviceCheckoutDays);
+				checkoutExpiration = device.CheckOutTime.Value.AddDays(_deviceService.sharedDeviceCheckoutDays);
 			}
 
 			return new GetDeviceResponse(device.Id.ToString(), device.PlatformId.ToString(), device.PoolId.ToString(), device.Name, device.Enabled, device.Address, device.ModelId?.ToString(), device.ModifiedByUser?.ToString(), device.Notes, device.ProblemTimeUtc, device.MaintenanceTimeUtc, device.Utilization, device.CheckedOutByUser, device.CheckOutTime, checkoutExpiration);
