Horde: Use a shared ticker for background ticks in LogFileService

Prevents multiple worker processes from trying to flush log chunks at the same time. Which hopefully reduces peak network I/O.

diff --git a/Engine/Source/Programs/Horde/Horde.Build/Logs/LogFileService.cs b/Engine/Source/Programs/Horde/Horde.Build/Logs/LogFileService.cs
--- a/Engine/Source/Programs/Horde/Horde.Build/Logs/LogFileService.cs
+++ b/Engine/Source/Programs/Horde/Horde.Build/Logs/LogFileService.cs
@@ -407,7 +407,7 @@ namespace Horde.Build.Logs
 			_logFileCache = new MemoryCache(new MemoryCacheOptions());
 			_builder = builder;
 			_storage = storage;
-			_ticker = clock.AddTicker<LogFileService>(TimeSpan.FromSeconds(30.0), TickAsync, logger);
+			_ticker = clock.AddSharedTicker<LogFileService>(TimeSpan.FromSeconds(30.0), TickAsync, logger);
 			_settings = settings;
 			_logger = logger;
 		}
