Horde: Drop /api/v1 and additional gRPC logs to Debug level, /ugs/api to Verbose
#rnx

diff --git a/Engine/Source/Programs/Horde/Horde.Build/Startup.cs b/Engine/Source/Programs/Horde/Horde.Build/Startup.cs
--- a/Engine/Source/Programs/Horde/Horde.Build/Startup.cs
+++ b/Engine/Source/Programs/Horde/Horde.Build/Startup.cs
@@ -985,6 +985,7 @@ namespace Horde.Build
 			if (context.Request != null && context.Request.Path.HasValue)
 			{
 				string requestPath = context.Request.Path;
+
 				if (requestPath.Equals("/Horde.HordeRpc/QueryServerStateV2", StringComparison.OrdinalIgnoreCase))
 				{
 					return LogEventLevel.Verbose;
@@ -1001,6 +1002,20 @@ namespace Horde.Build
 				{
 					return LogEventLevel.Verbose;
 				}
+				if (requestPath.StartsWith("/Horde.HordeRpc", StringComparison.OrdinalIgnoreCase))
+				{
+					return LogEventLevel.Debug;
+				}
+
+				if (requestPath.StartsWith("/api/v1", StringComparison.OrdinalIgnoreCase))
+				{
+					return LogEventLevel.Debug;
+				}
+
+				if (requestPath.StartsWith("/ugs/api", StringComparison.OrdinalIgnoreCase))
+				{
+					return LogEventLevel.Verbose;
+				}
 			}
 			return LogEventLevel.Information;
 		}
