Horde: Drop /health and /grpc.health loggnig to Debug
#rnx

diff --git a/Engine/Source/Programs/Horde/Horde.Build/Startup.cs b/Engine/Source/Programs/Horde/Horde.Build/Startup.cs
--- a/Engine/Source/Programs/Horde/Horde.Build/Startup.cs
+++ b/Engine/Source/Programs/Horde/Horde.Build/Startup.cs
@@ -1007,6 +1007,16 @@ namespace Horde.Build
 					return LogEventLevel.Debug;
 				}
 
+				if (requestPath.StartsWith("/health", StringComparison.OrdinalIgnoreCase))
+				{
+					return LogEventLevel.Debug;
+				}
+
+				if (requestPath.StartsWith("/grpc.health", StringComparison.OrdinalIgnoreCase))
+				{
+					return LogEventLevel.Debug;
+				}
+
 				if (requestPath.StartsWith("/api/v1", StringComparison.OrdinalIgnoreCase))
 				{
 					return LogEventLevel.Debug;
