Horde: Use dependency injection to manage lifetime of session objects, and add separate classes for different types of lease handlers. Also allow using the same service collection for the long-running service process as other commands.

diff --git a/Engine/Source/Programs/Horde/Horde.Agent.Tests/AwsLambdaListenerCommandTest.cs b/Engine/Source/Programs/Horde/Horde.Agent.Tests/AwsLambdaListenerCommandTest.cs
--- a/Engine/Source/Programs/Horde/Horde.Agent.Tests/AwsLambdaListenerCommandTest.cs
+++ b/Engine/Source/Programs/Horde/Horde.Agent.Tests/AwsLambdaListenerCommandTest.cs
@@ -5,10 +5,12 @@ using System.Threading;
 using System.Threading.Tasks;
 using Google.Protobuf;
 using Horde.Agent.Commands.Compute;
+using Horde.Agent.Leases;
 using Horde.Agent.Services;
 using HordeCommon;
 using HordeCommon.Rpc.Messages;
-using Microsoft.Extensions.Logging;
+using HordeCommon.Rpc.Tasks;
+using Microsoft.Extensions.DependencyInjection;
 using Microsoft.VisualStudio.TestTools.UnitTesting;
 
 namespace Horde.Agent.Tests
@@ -16,44 +18,50 @@ namespace Horde.Agent.Tests
 	[TestClass]
 	public class AwsLambdaListenerCommandTest
 	{
-		private readonly ILogger<WorkerService> _workerLogger;
-		private readonly ILogger<GrpcService> _grpcServiceLogger;
-		private readonly ILogger<FakeHordeRpcServer> _hordeRpcServerLogger;
-
-		public AwsLambdaListenerCommandTest()
+		class TestTaskHandler : LeaseHandler<TestTask>
 		{
-			using ILoggerFactory loggerFactory = LoggerFactory.Create(builder =>
+			public override Task<LeaseResult> ExecuteAsync(ISession session, string leaseId, TestTask message, CancellationToken cancellationToken)
 			{
-				builder.AddSimpleConsole(consoleLoggerOptions => consoleLoggerOptions.TimestampFormat = "[HH:mm:ss] ");
-			});
+				return Task.FromResult(LeaseResult.Success);
+			}
+		}
 
-			_workerLogger = loggerFactory.CreateLogger<WorkerService>();
-			_grpcServiceLogger = loggerFactory.CreateLogger<GrpcService>();
-			_hordeRpcServerLogger = loggerFactory.CreateLogger<FakeHordeRpcServer>();
+		private readonly FakeHordeRpcServer _fakeServer = new();
+		private readonly ServiceCollection _serviceCollection;
+
+		public AwsLambdaListenerCommandTest()
+		{
+			_serviceCollection = new ServiceCollection();
+			_serviceCollection.AddLogging();
+			_serviceCollection.AddSingleton<CapabilitiesService>();
+			_serviceCollection.AddSingleton<LeaseHandler, TestTaskHandler>();
+			_serviceCollection.AddSingleton<ISessionFactoryService>(sp => new FakeServerSessionFactory(_fakeServer));
 		}
-		
+
 		[TestMethod]
 		public async Task ShutdownIfNoLeasesArriveWithinMaxWaitTime()
 		{
-			using CancellationTokenSource cts = new (10000);
-			FakeHordeRpcServer fakeServer = new("testServerName", _hordeRpcServerLogger, cts.Token);
-			WorkerServiceLambda lambdaFunc = new((_) => WorkerServiceTest.GetWorkerService(_grpcServiceLogger, _workerLogger, (a, b, c) => WorkerServiceTest.NullExecutor,(c) => fakeServer.GetClient()));
+			using CancellationTokenSource cts = new(10000);
+			using ServiceProvider serviceProvider = _serviceCollection.BuildServiceProvider();
+
+			WorkerServiceLambda lambdaFunc = new(serviceProvider);
 
-			AwsLambdaListenResponse response = await InvokeProtoAsync(lambdaFunc, new() { MaxWaitTimeForLeaseMs = 1000 }, cts.Token);
+			AwsLambdaListenResponse response = await InvokeProtoAsync(lambdaFunc, new AwsLambdaListenRequest { MaxWaitTimeForLeaseMs = 1000 }, cts.Token);
 			Assert.IsFalse(response.DidAcceptLease);
 		}
-		
+
 		[TestMethod]
 		public async Task HandleOneLeaseAndShutdown()
 		{
+			using ServiceProvider serviceProvider = _serviceCollection.BuildServiceProvider();
+
 			using CancellationTokenSource cts = new (10000);
-			FakeHordeRpcServer fakeServer = new("testServerName", _hordeRpcServerLogger, cts.Token);
-			WorkerServiceLambda lambdaFunc = new((_) => WorkerServiceTest.GetWorkerService(_grpcServiceLogger, _workerLogger, (a, b, c) => WorkerServiceTest.NullExecutor,(c) => fakeServer.GetClient()));
+			WorkerServiceLambda lambdaFunc = new(serviceProvider);
 
-			fakeServer.AddTestLease("testLeaseId");
-			AwsLambdaListenResponse response = await InvokeProtoAsync(lambdaFunc, new() { MaxWaitTimeForLeaseMs = 5000 }, cts.Token);
-			Assert.AreEqual(LeaseState.Completed, fakeServer.GetLease("testLeaseId").State);
-			Assert.AreEqual(LeaseOutcome.Success, fakeServer.GetLease("testLeaseId").Outcome);
+			_fakeServer.AddTestLease("testLeaseId");
+			AwsLambdaListenResponse response = await InvokeProtoAsync(lambdaFunc, new AwsLambdaListenRequest { MaxWaitTimeForLeaseMs = 5000 }, cts.Token);
+			Assert.AreEqual(LeaseState.Completed, _fakeServer.GetLease("testLeaseId").State);
+			Assert.AreEqual(LeaseOutcome.Success, _fakeServer.GetLease("testLeaseId").Outcome);
 			Assert.IsTrue(response.DidAcceptLease);
 		}
 
diff --git a/Engine/Source/Programs/Horde/Horde.Agent.Tests/Helper.cs b/Engine/Source/Programs/Horde/Horde.Agent.Tests/Helper.cs
--- a/Engine/Source/Programs/Horde/Horde.Agent.Tests/Helper.cs
+++ b/Engine/Source/Programs/Horde/Horde.Agent.Tests/Helper.cs
@@ -8,9 +8,11 @@ using Google.Protobuf.WellKnownTypes;
 using Grpc.Core;
 using Grpc.Net.Client;
 using Horde.Agent.Execution.Interfaces;
+using Horde.Agent.Services;
 using Horde.Agent.Utility;
 using HordeCommon;
 using HordeCommon.Rpc;
+using HordeCommon.Rpc.Tasks;
 using Microsoft.Extensions.Logging;
 using Microsoft.Extensions.Options;
 
@@ -228,4 +230,16 @@ namespace Horde.Agent.Tests
 			return Task.CompletedTask;
 		}
 	}
+
+	class SimpleTestExecutorFactory : IExecutorFactory
+	{
+		readonly IExecutor _executor;
+
+		public SimpleTestExecutorFactory(IExecutor executor)
+		{
+			_executor = executor;
+		}
+
+		public IExecutor CreateExecutor(ISession session, ExecuteJobTask executeJobTask, BeginBatchResponse beginBatchResponse) => _executor;
+	}
 }
\ No newline at end of file
diff --git a/Engine/Source/Programs/Horde/Horde.Agent.Tests/WorkerServiceTest.cs b/Engine/Source/Programs/Horde/Horde.Agent.Tests/WorkerServiceTest.cs
--- a/Engine/Source/Programs/Horde/Horde.Agent.Tests/WorkerServiceTest.cs
+++ b/Engine/Source/Programs/Horde/Horde.Agent.Tests/WorkerServiceTest.cs
@@ -4,21 +4,29 @@ using System;
 using System.Collections.Generic;
 using System.IO;
 using System.Linq;
+using System.Net.Sockets;
 using System.Threading;
 using System.Threading.Channels;
 using System.Threading.Tasks;
+using Amazon.EC2.Model;
+using EpicGames.Core;
 using Google.Protobuf.WellKnownTypes;
 using Grpc.Core;
 using Grpc.Net.Client;
 using Horde.Agent.Execution.Interfaces;
+using Horde.Agent.Leases;
+using Horde.Agent.Leases.Handlers;
 using Horde.Agent.Services;
 using Horde.Agent.Utility;
 using HordeCommon;
 using HordeCommon.Rpc;
 using HordeCommon.Rpc.Messages;
 using HordeCommon.Rpc.Tasks;
+using Microsoft.Extensions.DependencyInjection;
 using Microsoft.Extensions.Logging;
+using Microsoft.Extensions.Logging.Abstractions;
 using Microsoft.Extensions.Options;
+using Microsoft.VisualStudio.TestPlatform.ObjectModel.DataCollection;
 using Microsoft.VisualStudio.TestTools.UnitTesting;
 using Moq;
 
@@ -27,10 +35,8 @@ namespace Horde.Agent.Tests
 	[TestClass]
 	public class WorkerServiceTest
 	{
-		private readonly ILogger<WorkerService> _workerLogger;
-		private readonly ILogger<GrpcService> _grpcServiceLogger;
-		private readonly ILogger<FakeHordeRpcServer> _hordeRpcServerLogger;
-		
+		private ServiceCollection _serviceCollection;
+
 		internal static IExecutor NullExecutor = new SimpleTestExecutor(async (step, logger, cancellationToken) =>
 		{
 			await Task.Delay(1, cancellationToken);
@@ -39,52 +45,32 @@ namespace Horde.Agent.Tests
 
 		public WorkerServiceTest()
 		{
-			using ILoggerFactory loggerFactory = LoggerFactory.Create(builder =>
+			_serviceCollection = new ServiceCollection();
+			_serviceCollection.AddLogging();
+
+			_serviceCollection.Configure<AgentSettings>(settings =>
 			{
-				builder.AddSimpleConsole(consoleLoggerOptions => consoleLoggerOptions.TimestampFormat = "[HH:mm:ss] ");
+				ServerProfile profile = new();
+				profile.Name = "test";
+				profile.Environment = "test-env";
+				profile.Token = "bogus-token";
+				profile.Url = new Uri("http://localhost");
+
+				settings.ServerProfiles.Add(profile);
+				settings.Server = "test";
+				settings.WorkingDir = Path.GetTempPath();
+				settings.Executor = ExecutorType.Test; // Not really used since the executor is overridden in the tests
 			});
 
-			_workerLogger = loggerFactory.CreateLogger<WorkerService>();
-			_grpcServiceLogger = loggerFactory.CreateLogger<GrpcService>();
-			_hordeRpcServerLogger = loggerFactory.CreateLogger<FakeHordeRpcServer>();
-		}
-
-		internal static WorkerService GetWorkerService(
-			ILogger<GrpcService> grpcServiceLogger,
-			ILogger<WorkerService> workerServiceLogger,
-			Func<IRpcConnection, ExecuteJobTask, BeginBatchResponse, IExecutor>? createExecutor = null,
-			Func<GrpcChannel, HordeRpc.HordeRpcClient>? createHordeRpcClient = null)
-		{
-			AgentSettings settings = new ();
-			ServerProfile profile = new ();
-			profile.Name = "test";
-			profile.Environment = "test-env";
-			profile.Token = "bogus-token";
-			profile.Url = new Uri("http://localhost");
-			settings.ServerProfiles.Add(profile);
-			settings.Server = "test";
-			settings.WorkingDir = Path.GetTempPath();
-			settings.Executor = ExecutorType.Test; // Not really used since the executor is overridden in the tests
-			IOptions<AgentSettings> settingsMonitor = new TestOptionsMonitor<AgentSettings>(settings);
-
-			GrpcService grpcService = new (settingsMonitor, grpcServiceLogger);
-			WorkerService ws = new (workerServiceLogger, settingsMonitor, grpcService, null!, createExecutor, createHordeRpcClient);
-			ws._rpcConnectionRetryDelay = TimeSpan.FromSeconds(0.1);
-			return ws;
-		}
+			_serviceCollection.AddSingleton<WorkerService>();
 
-		private WorkerService GetWorkerService(
-			Func<IRpcConnection, ExecuteJobTask, BeginBatchResponse, IExecutor>? createExecutor = null,
-			Func<GrpcChannel, HordeRpc.HordeRpcClient>? createHordeRpcClient = null)
-		{
-			return GetWorkerService(_grpcServiceLogger, _workerLogger, createExecutor, createHordeRpcClient);
+			_serviceCollection.AddSingleton<JobHandler>();
+			_serviceCollection.AddSingleton<LeaseHandler>(sp => sp.GetRequiredService<JobHandler>());
 		}
 
 		[TestMethod]
 		public async Task AbortExecuteStepTest()
 		{
-			using WorkerService ws = GetWorkerService();
-
 			{
 				using CancellationTokenSource cancelSource = new CancellationTokenSource();
 				using CancellationTokenSource stepCancelSource = new CancellationTokenSource();
@@ -95,8 +81,9 @@ namespace Horde.Agent.Tests
 					await Task.Delay(5000, cancelToken);
 					return JobStepOutcome.Success;
 				});
-				await Assert.ThrowsExceptionAsync<TaskCanceledException>(() => WorkerService.ExecuteStepAsync(executor,
-					new BeginStepResponse(), _workerLogger, cancelSource.Token, stepCancelSource.Token));
+
+				await Assert.ThrowsExceptionAsync<TaskCanceledException>(() => JobHandler.ExecuteStepAsync(executor,
+					new BeginStepResponse(), NullLogger.Instance, cancelSource.Token, stepCancelSource.Token));
 			}
 
 			{
@@ -109,7 +96,7 @@ namespace Horde.Agent.Tests
 					await Task.Delay(5000, cancelToken);
 					return JobStepOutcome.Success;
 				});
-				(JobStepOutcome stepOutcome, JobStepState stepState) = await WorkerService.ExecuteStepAsync(executor, new BeginStepResponse(), _workerLogger,
+				(JobStepOutcome stepOutcome, JobStepState stepState) = await JobHandler.ExecuteStepAsync(executor, new BeginStepResponse(), NullLogger.Instance,
 					cancelSource.Token, stepCancelSource.Token);
 				Assert.AreEqual(JobStepOutcome.Failure, stepOutcome);
 				Assert.AreEqual(JobStepState.Aborted, stepState);
@@ -130,9 +117,11 @@ namespace Horde.Agent.Tests
 			executeJobTask.AutoSdkWorkspace = new AgentWorkspace();
 			executeJobTask.Workspace = new AgentWorkspace();
 
-			HordeRpcClientStub client = new HordeRpcClientStub(_workerLogger);
+			HordeRpcClientStub client = new HordeRpcClientStub(NullLogger.Instance);
 			await using RpcConnectionStub rpcConnection = new RpcConnectionStub(null!, client);
 
+			await using ISession session = FakeServerSessionFactory.CreateSession(rpcConnection);
+
 			client.BeginStepResponses.Enqueue(new BeginStepResponse {Name = "stepName1", StepId = "stepId1"});
 			client.BeginStepResponses.Enqueue(new BeginStepResponse {Name = "stepName2", StepId = "stepId2"});
 			client.BeginStepResponses.Enqueue(new BeginStepResponse {Name = "stepName3", StepId = "stepId3"});
@@ -147,10 +136,14 @@ namespace Horde.Agent.Tests
 				return JobStepOutcome.Success;
 			});
 
-			using WorkerService ws = GetWorkerService((a, b, c) => executor);
-			ws._stepAbortPollInterval = TimeSpan.FromMilliseconds(1);
-			LeaseOutcome outcome = (await ws.ExecuteJobAsync(rpcConnection, "leaseId1", executeJobTask,
-				_workerLogger, token)).Outcome;
+			_serviceCollection.AddSingleton<IExecutorFactory>(x => new SimpleTestExecutorFactory(executor));
+			using ServiceProvider serviceProvider = _serviceCollection.BuildServiceProvider();
+
+			JobHandler jobHandler = serviceProvider.GetRequiredService<JobHandler>();
+			jobHandler._stepAbortPollInterval = TimeSpan.FromMilliseconds(1);
+
+			LeaseOutcome outcome = (await jobHandler.ExecuteAsync(session, "leaseId1", executeJobTask,
+				token)).Outcome;
 
 			Assert.AreEqual(LeaseOutcome.Success, outcome);
 			Assert.AreEqual(4, client.UpdateStepRequests.Count);
@@ -175,10 +168,13 @@ namespace Horde.Agent.Tests
 				return JobStepOutcome.Success;
 			});
 
-			using WorkerService ws = GetWorkerService((a, b, c) => executor);
-			ws._stepAbortPollInterval = TimeSpan.FromMilliseconds(5);
-			
-			HordeRpcClientStub client = new HordeRpcClientStub(_workerLogger);
+			_serviceCollection.AddSingleton<IExecutorFactory>(x => new SimpleTestExecutorFactory(executor));
+			using ServiceProvider serviceProvider = _serviceCollection.BuildServiceProvider();
+
+			JobHandler jobHandler = serviceProvider.GetRequiredService<JobHandler>();
+			jobHandler._stepAbortPollInterval = TimeSpan.FromMilliseconds(5);
+
+			HordeRpcClientStub client = new HordeRpcClientStub(NullLogger.Instance);
 			await using RpcConnectionStub rpcConnection = new RpcConnectionStub(null!, client);
 
 			int c = 0;
@@ -197,7 +193,7 @@ namespace Horde.Agent.Tests
 			using CancellationTokenSource stepCancelSource = new CancellationTokenSource();
 			TaskCompletionSource<bool> stepFinishedSource = new TaskCompletionSource<bool>();
 
-			await ws.PollForStepAbort(rpcConnection, "jobId1", "batchId1", "logId1", stepCancelSource, stepFinishedSource.Task, stepPollCancelSource.Token);
+			await jobHandler.PollForStepAbort(rpcConnection, "jobId1", "batchId1", "logId1", stepCancelSource, stepFinishedSource.Task, stepPollCancelSource.Token);
 			Assert.IsTrue(stepCancelSource.IsCancellationRequested);
 		}
 
@@ -209,21 +205,50 @@ namespace Horde.Agent.Tests
 				await Task.Delay(50, cancellationToken);
 				return JobStepOutcome.Success;
 			});
-			
+
+			_serviceCollection.AddSingleton<IExecutorFactory>(x => new SimpleTestExecutorFactory(executor));
+			using ServiceProvider serviceProvider = _serviceCollection.BuildServiceProvider();
+
 			using CancellationTokenSource cts = new ();
 			cts.CancelAfter(20000);
 
-			FakeHordeRpcServer fakeServer = new("bogusServerName", _hordeRpcServerLogger, cts.Token);
-			using WorkerService ws = GetWorkerService((a, b, c) => executor,(c) => fakeServer.GetClient());
+			FakeHordeRpcServer fakeServer = new();// ("bogusServerName", cts.Token);
+
+			await using ISession session = FakeServerSessionFactory.CreateSession(fakeServer.GetConnection());
 
-			Task handleSessionTask = ws.HandleSessionAsync(false, false, cts.Token);
-			await fakeServer.CreateSessionReceived.Task.WaitAsync(cts.Token);
+			LeaseManager manager = new LeaseManager(session, null!, serviceProvider.GetRequiredService<IEnumerable<LeaseHandler>>(), serviceProvider.GetRequiredService<IOptions<AgentSettings>>(), NullLogger.Instance);
+
+			Task handleSessionTask = Task.Run(() => manager.RunAsync(false, cts.Token), cts.Token);
 			await fakeServer.UpdateSessionReceived.Task.WaitAsync(cts.Token);
-			await ws.ShutdownAsync(_workerLogger);
+			cts.Cancel();
 			await handleSessionTask; // Ensure it runs to completion and no exceptions are raised
 		}
 	}
 
+	internal class FakeServerSessionFactory : ISessionFactoryService
+	{
+		readonly FakeHordeRpcServer _fakeServer;
+
+		public FakeServerSessionFactory(FakeHordeRpcServer fakeServer) => _fakeServer = fakeServer;
+
+		public Task<ISession> CreateAsync(CancellationToken cancellationToken)
+		{
+			return Task.FromResult(CreateSession(_fakeServer.GetConnection()));
+		}
+
+		public static ISession CreateSession(IRpcConnection rpcConnection)
+		{
+			Mock<ISession> fakeSession = new Mock<ISession>();
+			fakeSession.Setup(x => x.AgentId).Returns("LocalAgent");
+			fakeSession.Setup(x => x.SessionId).Returns("Session");
+			fakeSession.Setup(x => x.RpcConnection).Returns(rpcConnection);
+			fakeSession.Setup(x => x.TerminateProcessesAsync(It.IsAny<ILogger>(), It.IsAny<CancellationToken>())).Returns(Task.CompletedTask);
+			fakeSession.Setup(x => x.DisposeAsync()).Returns(new ValueTask());
+			fakeSession.Setup(x => x.WorkingDir).Returns(DirectoryReference.GetCurrentDirectory());
+			return fakeSession.Object;
+		}
+	}
+
 	/// <summary>
 	/// Fake implementation of a HordeRpc gRPC server.
 	/// Provides a corresponding gRPC client class that can be used with the WorkerService
@@ -235,15 +260,17 @@ namespace Horde.Agent.Tests
 		private bool _isStopping = false;
 		private Dictionary<string, Lease> _leases = new();
 		private readonly Mock<HordeRpc.HordeRpcClient> _mockClient;
+		private readonly Mock<IRpcClientRef> _mockClientRef;
+		private readonly Mock<IRpcConnection> _mockConnection;
 		private readonly ILogger<FakeHordeRpcServer> _logger;
 		public readonly TaskCompletionSource<bool> CreateSessionReceived = new();
 		public readonly TaskCompletionSource<bool> UpdateSessionReceived = new();
 
-		public FakeHordeRpcServer(string serverName,ILogger<FakeHordeRpcServer> logger, CancellationToken cancellationToken)
+		public FakeHordeRpcServer()
 		{
-			_serverName = serverName;
+			_serverName = "FakeServer";
 			_mockClient = new (MockBehavior.Strict);
-			_logger = logger;
+			_logger = NullLogger<FakeHordeRpcServer>.Instance;
 
 			_mockClient
 				.Setup(m => m.CreateSessionAsync(It.IsAny<CreateSessionRequest>(), null, null, It.IsAny<CancellationToken>()))
@@ -251,11 +278,27 @@ namespace Horde.Agent.Tests
 			
 			_mockClient
 				.Setup(m => m.QueryServerStateV2(null, null, It.IsAny<CancellationToken>()))
-				.Returns(() => GetQueryServerStateCall(cancellationToken));
+				.Returns(() => GetQueryServerStateCall(CancellationToken.None));
 			
 			_mockClient
 				.Setup(m => m.UpdateSession(null, It.IsAny<DateTime>(), It.IsAny<CancellationToken>()))
-				.Returns(() => GetUpdateSessionCall(cancellationToken));
+				.Returns(() => GetUpdateSessionCall(CancellationToken.None));
+
+			_mockClientRef = new Mock<IRpcClientRef>();
+			_mockClientRef
+				.Setup(m => m.Client)
+				.Returns(() => _mockClient.Object);
+
+			_mockConnection = new(MockBehavior.Strict);
+			_mockConnection
+				.Setup(m => m.TryGetClientRef(It.IsAny<RpcContext>()))
+				.Returns<RpcContext>(x => _mockClientRef.Object);
+			_mockConnection
+				.Setup(m => m.GetClientRef(It.IsAny<RpcContext>(), It.IsAny<CancellationToken>()))
+				.Returns(() => Task.FromResult(_mockClientRef.Object));
+			_mockConnection
+				.Setup(m => m.DisposeAsync())
+				.Returns(() => new ValueTask());
 		}
 
 		public void AddTestLease(string leaseId)
@@ -284,6 +327,11 @@ namespace Horde.Agent.Tests
 			return _mockClient.Object;
 		}
 
+		public IRpcConnection GetConnection()
+		{
+			return _mockConnection.Object;
+		}
+
 		public CreateSessionResponse OnCreateSessionRequest(CreateSessionRequest request)
 		{
 			CreateSessionReceived.TrySetResult(true);
@@ -451,8 +499,14 @@ namespace Horde.Agent.Tests
 		/// <inheritdoc/>
 		public async Task WriteAsync(T message)
 		{
-			if (_isCompleted) throw new InvalidOperationException("Stream is marked as complete");
-			if (_onWrite != null) await _onWrite(message);
+			if (_isCompleted)
+			{
+				throw new InvalidOperationException("Stream is marked as complete");
+			}
+			if (_onWrite != null)
+			{
+				await _onWrite(message);
+			}
 		}
 
 		/// <inheritdoc/>
@@ -462,7 +516,10 @@ namespace Horde.Agent.Tests
 		public async Task CompleteAsync()
 		{
 			_isCompleted = true;
-			if (_onComplete != null) await _onComplete();
+			if (_onComplete != null)
+			{
+				await _onComplete();
+			}
 		}
 	}
 }
diff --git a/Engine/Source/Programs/Horde/Horde.Agent/AgentSettings.cs b/Engine/Source/Programs/Horde/Horde.Agent/AgentSettings.cs
--- a/Engine/Source/Programs/Horde/Horde.Agent/AgentSettings.cs
+++ b/Engine/Source/Programs/Horde/Horde.Agent/AgentSettings.cs
@@ -202,6 +202,11 @@ namespace Horde.Agent
 		/// </summary>
 		public bool WriteStepOutputToLogger { get; set; }
 
+		/// <summary>
+		/// Queries information about the current agent through the AWS EC2 interface
+		/// </summary>
+		public bool EnableAwsEc2Support { get; set; } = true;
+
 		/// <summary>
 		/// Key/value properties in addition to those set internally by the agent
 		/// </summary>
diff --git a/Engine/Source/Programs/Horde/Horde.Agent/Commands/Compute/AwsLambdaListenerCommand.cs b/Engine/Source/Programs/Horde/Horde.Agent/Commands/Compute/AwsLambdaListenerCommand.cs
--- a/Engine/Source/Programs/Horde/Horde.Agent/Commands/Compute/AwsLambdaListenerCommand.cs
+++ b/Engine/Source/Programs/Horde/Horde.Agent/Commands/Compute/AwsLambdaListenerCommand.cs
@@ -15,8 +15,10 @@ using EpicGames.Horde.Storage;
 using EpicGames.Horde.Storage.Impl;
 using Google.Protobuf;
 using Grpc.Core;
+using Horde.Agent.Leases;
 using Horde.Agent.Services;
 using HordeCommon.Rpc.Messages;
+using Microsoft.Extensions.DependencyInjection;
 using Microsoft.Extensions.Logging;
 using Microsoft.Extensions.Options;
 using Serilog;
@@ -34,11 +36,11 @@ namespace Horde.Agent.Commands.Compute
 	/// </summary>
 	class WorkerServiceLambda : IAwsLambdaFunction
 	{
-		private readonly Func<string, WorkerService> _createWorkerService;
+		readonly IServiceProvider _serviceProvider;
 		
-		public WorkerServiceLambda(Func<string, WorkerService> createWorkerService)
+		public WorkerServiceLambda(IServiceProvider serviceProvider)
 		{
-			_createWorkerService = createWorkerService;
+			_serviceProvider = serviceProvider;
 		}
 
 		/// <inheritdoc />
@@ -50,9 +52,12 @@ namespace Horde.Agent.Commands.Compute
 				AwsLambdaListenRequest.Parser.ParseFrom(requestData.Span);
 
 			TimeSpan maxWaitTime = TimeSpan.FromMilliseconds(request.MaxWaitTimeForLeaseMs);
-			using WorkerService service = _createWorkerService(requestId);
+
+			await using ISession session = await _serviceProvider.GetRequiredService<ISessionFactoryService>().CreateAsync(cancellationToken);
+
+			LeaseManager manager = new LeaseManager(session, _serviceProvider);
 			bool isLeaseActive = false;
-			service.OnLeaseActive += lease => isLeaseActive = true;
+			manager.OnLeaseActive += lease => isLeaseActive = true;
 			using CancellationTokenSource maxWaitTimeCts = new();
 
 			Task _ = Task.Run(async () =>
@@ -68,7 +73,7 @@ namespace Horde.Agent.Commands.Compute
 			using CancellationTokenSource linkedCts = CancellationTokenSource.CreateLinkedTokenSource(cancellationToken, maxWaitTimeCts.Token);
 			try
 			{
-				await service.HandleSessionAsync(true, false, linkedCts.Token);
+				await manager.RunAsync(true, linkedCts.Token);
 			}
 			catch (RpcException e)
 			{
@@ -79,7 +84,7 @@ namespace Horde.Agent.Commands.Compute
 				}
 			}
 			
-			AwsLambdaListenResponse response = new () { DidAcceptLease = service.NumLeasesCompleted > 0 };
+			AwsLambdaListenResponse response = new () { DidAcceptLease = manager.NumLeasesCompleted > 0 };
 			byte[] responseData = isProtobufCompatibleJson ?
 				Encoding.UTF8.GetBytes(JsonFormatter.Default.Format(response)) :
 				response.ToByteArray();
@@ -97,6 +102,9 @@ namespace Horde.Agent.Commands.Compute
 		/// <inheritdoc/>
 		public override async Task<int> ExecuteAsync(ILogger logger)
 		{
+			ServiceCollection serviceCollection = new ServiceCollection();
+			using ServiceProvider serviceProvider = serviceCollection.BuildServiceProvider();
+
 			LogLevel logLevel;
 			string logLevelStr = Environment.GetEnvironmentVariable("UE_HORDE_LOG_LEVEL") ?? "debug";
 
@@ -150,11 +158,7 @@ namespace Horde.Agent.Commands.Compute
 					cts.Cancel();
 				};
 
-				WorkerServiceLambda lambdaFunction = new ((requestId) =>
-				{
-					agentSettingsOptions.Value.Name = "lambda-" + requestId;
-					return new WorkerService(workerServiceLogger, agentSettingsOptions, grpcService, storageClient);
-				});
+				WorkerServiceLambda lambdaFunction = new(serviceProvider);
 				bool success = await lambdaClient.ListenForInvocationsAsync(lambdaFunction, cts.Token);
 				return success ? 1 : 0;
 			}
diff --git a/Engine/Source/Programs/Horde/Horde.Agent/Commands/Service/RunCommand.cs b/Engine/Source/Programs/Horde/Horde.Agent/Commands/Service/RunCommand.cs
--- a/Engine/Source/Programs/Horde/Horde.Agent/Commands/Service/RunCommand.cs
+++ b/Engine/Source/Programs/Horde/Horde.Agent/Commands/Service/RunCommand.cs
@@ -9,12 +9,18 @@ using Datadog.Trace.Configuration;
 using Datadog.Trace.OpenTracing;
 using EpicGames.Core;
 using EpicGames.Horde.Storage;
+using Horde.Agent.Execution;
+using Horde.Agent.Execution.Interfaces;
+using Horde.Agent.Leases;
+using Horde.Agent.Leases.Handlers;
 using Horde.Agent.Services;
 using Horde.Agent.Utility;
 using Microsoft.Extensions.Configuration;
 using Microsoft.Extensions.DependencyInjection;
+using Microsoft.Extensions.DependencyInjection.Extensions;
 using Microsoft.Extensions.Hosting;
 using Microsoft.Extensions.Logging;
+using Microsoft.Extensions.Options;
 using OpenTracing;
 using OpenTracing.Util;
 using Polly;
@@ -47,7 +53,17 @@ namespace Horde.Agent.Modes.Service
 		/// </summary>
 		[CommandLine("-LogLevel")]
 		public string LogLevelStr { get; set; } = "information";
-		
+
+		readonly DefaultServices _defaultServices;
+
+		/// <summary>
+		/// Constructor
+		/// </summary>
+		public RunCommand(DefaultServices defaultServices)
+		{
+			_defaultServices = defaultServices;
+		}
+
 		/// <summary>
 		/// Runs the service indefinitely
 		/// </summary>
@@ -79,6 +95,8 @@ namespace Horde.Agent.Modes.Service
 			hostBuilder = hostBuilder
 				.ConfigureAppConfiguration(builder =>
 				{
+					builder.AddConfiguration(_defaultServices.Configuration);
+
 					Dictionary<string, string> overrides = new Dictionary<string, string>();
 					if (Server != null)
 					{
@@ -88,13 +106,13 @@ namespace Horde.Agent.Modes.Service
 					{
 						overrides.Add($"{AgentSettings.SectionName}:{nameof(AgentSettings.WorkingDir)}", WorkingDir);
 					}
+
 					builder.AddInMemoryCollection(overrides);
 				})
 				.ConfigureLogging(builder =>
 				{
+					// We add our logger through ConfigureServices, inherited from _defaultServices
 					builder.ClearProviders();
-					builder.AddProvider(new Logging.HordeLoggerProvider());
-					builder.AddFilter<Logging.HordeLoggerProvider>(null, LogLevel.Trace);
 				})
 				.ConfigureServices((hostContext, services) =>
 				{
@@ -104,38 +122,10 @@ namespace Horde.Agent.Modes.Service
 						options.ShutdownTimeout = TimeSpan.FromSeconds(30);
 					});
 
-					IConfigurationSection configSection = hostContext.Configuration.GetSection(AgentSettings.SectionName);
-					services.AddOptions<AgentSettings>().Configure(options => configSection.Bind(options)).ValidateDataAnnotations();
-
-					AgentSettings settings = new AgentSettings();
-					configSection.Bind(settings);
-
-					ServerProfile serverProfile = settings.GetCurrentServerProfile();
-					ConfigureTracing(serverProfile.Environment, Program.Version);
-
-					Logging.SetEnv(serverProfile.Environment);
-
-					services.AddHttpClient(Program.HordeServerClientName, config =>
-					{
-						config.BaseAddress = serverProfile.Url;
-						config.DefaultRequestHeaders.Add("Accept", "application/json");
-						config.Timeout = TimeSpan.FromSeconds(300); // Need to make sure this doesn't cancel any long running gRPC streaming calls (eg. session update)
-					})
-					.ConfigurePrimaryHttpMessageHandler(() =>
-					{
-						HttpClientHandler handler = new HttpClientHandler();
-						handler.ServerCertificateCustomValidationCallback += (sender, cert, chain, errors) => CertificateHelper.CertificateValidationCallBack(logger, sender, cert, chain, errors, serverProfile);
-						return handler;
-					})
-					.AddTransientHttpErrorPolicy(builder =>
+					foreach (ServiceDescriptor descriptor in _defaultServices.Descriptors)
 					{
-						return builder.WaitAndRetryAsync(new[] { TimeSpan.FromSeconds(1), TimeSpan.FromSeconds(5), TimeSpan.FromSeconds(10) });
-					});
-
-					services.AddHordeStorage(settings => configSection.GetCurrentServerProfile().GetSection(nameof(serverProfile.Storage)).Bind(settings));
-
-					services.AddSingleton<GrpcService>();
-					services.AddHostedService<WorkerService>();
+						services.Add(descriptor);
+					}
 				});
 
 			try
@@ -149,19 +139,5 @@ namespace Horde.Agent.Modes.Service
 
 			return 0;
 		}
-
-		static void ConfigureTracing(string environment, string version)
-		{
-			TracerSettings settings = TracerSettings.FromDefaultSources();
-			settings.Environment = environment;
-			settings.ServiceName = "hordeagent";
-			settings.ServiceVersion = version;
-			settings.LogsInjectionEnabled = true;
-
-			Tracer.Configure(settings);
-
-			ITracer openTracer = OpenTracingTracerFactory.WrapTracer(Tracer.Instance);
-			GlobalTracer.Register(openTracer);
-		}
 	}
 }
diff --git a/Engine/Source/Programs/Horde/Horde.Agent/Commands/Utilities/CapsCommand.cs b/Engine/Source/Programs/Horde/Horde.Agent/Commands/Utilities/CapsCommand.cs
--- a/Engine/Source/Programs/Horde/Horde.Agent/Commands/Utilities/CapsCommand.cs
+++ b/Engine/Source/Programs/Horde/Horde.Agent/Commands/Utilities/CapsCommand.cs
@@ -14,10 +14,20 @@ namespace Horde.Agent.Commands.Utilities
 	[Command("Caps", "Lists detected capabilities of this agent")]
 	class CapsCommand : Command
 	{
+		readonly CapabilitiesService _capabilitiesService;
+
+		/// <summary>
+		/// Constructor
+		/// </summary>
+		public CapsCommand(CapabilitiesService capabilitiesService)
+		{
+			_capabilitiesService = capabilitiesService;
+		}
+
 		/// <inheritdoc/>
 		public override async Task<int> ExecuteAsync(ILogger logger)
 		{
-			AgentCapabilities capabilities = await WorkerService.GetAgentCapabilities(DirectoryReference.GetCurrentDirectory(), logger);
+			AgentCapabilities capabilities = await _capabilitiesService.GetCapabilitiesAsync(null);
 			foreach (DeviceCapabilities device in capabilities.Devices)
 			{
 				logger.LogInformation("Device: {Name}", device.Handle);
diff --git a/Engine/Source/Programs/Horde/Horde.Agent/Execution/Interfaces/IExecutor.cs b/Engine/Source/Programs/Horde/Horde.Agent/Execution/Interfaces/IExecutor.cs
--- a/Engine/Source/Programs/Horde/Horde.Agent/Execution/Interfaces/IExecutor.cs
+++ b/Engine/Source/Programs/Horde/Horde.Agent/Execution/Interfaces/IExecutor.cs
@@ -2,8 +2,10 @@
 
 using System.Threading;
 using System.Threading.Tasks;
+using Horde.Agent.Services;
 using HordeCommon;
 using HordeCommon.Rpc;
+using HordeCommon.Rpc.Tasks;
 using Microsoft.Extensions.Logging;
 
 namespace Horde.Agent.Execution.Interfaces
@@ -14,4 +16,9 @@ namespace Horde.Agent.Execution.Interfaces
 		Task<JobStepOutcome> RunAsync(BeginStepResponse step, ILogger logger, CancellationToken cancellationToken);
 		Task FinalizeAsync(ILogger logger, CancellationToken cancellationToken);
 	}
+
+	interface IExecutorFactory
+	{
+		IExecutor CreateExecutor(ISession session, ExecuteJobTask executeJobTask, BeginBatchResponse beginBatchResponse);
+	}
 }
diff --git a/Engine/Source/Programs/Horde/Horde.Agent/Execution/LocalExecutor.cs b/Engine/Source/Programs/Horde/Horde.Agent/Execution/LocalExecutor.cs
--- a/Engine/Source/Programs/Horde/Horde.Agent/Execution/LocalExecutor.cs
+++ b/Engine/Source/Programs/Horde/Horde.Agent/Execution/LocalExecutor.cs
@@ -1,13 +1,16 @@
 // Copyright Epic Games, Inc. All Rights Reserved.
 
 using System;
-using System.Collections.Generic;
 using System.Threading;
 using System.Threading.Tasks;
 using EpicGames.Core;
+using Horde.Agent.Execution.Interfaces;
+using Horde.Agent.Services;
 using Horde.Agent.Utility;
 using HordeCommon.Rpc;
+using HordeCommon.Rpc.Tasks;
 using Microsoft.Extensions.Logging;
+using Microsoft.Extensions.Options;
 
 namespace Horde.Agent.Execution
 {
@@ -45,4 +48,19 @@ namespace Horde.Agent.Execution
 			}
 		}
 	}
+
+	class LocalExecutorFactory : IExecutorFactory
+	{
+		readonly LocalExecutorSettings _settings;
+
+		public LocalExecutorFactory(IOptions<LocalExecutorSettings> settings)
+		{
+			_settings = settings.Value;
+		}
+
+		public IExecutor CreateExecutor(ISession session, ExecuteJobTask executeJobTask, BeginBatchResponse beginBatchResponse)
+		{
+			return new LocalExecutor(session.RpcConnection, executeJobTask.JobId, executeJobTask.BatchId, beginBatchResponse.AgentType, _settings);
+		}
+	}
 }
diff --git a/Engine/Source/Programs/Horde/Horde.Agent/Execution/PerforceExecutor.cs b/Engine/Source/Programs/Horde/Horde.Agent/Execution/PerforceExecutor.cs
--- a/Engine/Source/Programs/Horde/Horde.Agent/Execution/PerforceExecutor.cs
+++ b/Engine/Source/Programs/Horde/Horde.Agent/Execution/PerforceExecutor.cs
@@ -11,10 +11,14 @@ using System.Threading.Tasks;
 using EpicGames.Core;
 using EpicGames.Perforce;
 using EpicGames.Perforce.Managed;
+using Horde.Agent.Execution.Interfaces;
+using Horde.Agent.Services;
 using Horde.Agent.Utility;
 using HordeCommon.Rpc;
 using HordeCommon.Rpc.Messages;
+using HordeCommon.Rpc.Tasks;
 using Microsoft.Extensions.Logging;
+using Microsoft.Extensions.Options;
 using OpenTracing;
 using OpenTracing.Util;
 
@@ -385,4 +389,12 @@ namespace Horde.Agent.Execution
 			}
 		}
 	}
+
+	class PerforceExecutorFactory : IExecutorFactory
+	{
+		public IExecutor CreateExecutor(ISession session, ExecuteJobTask executeJobTask, BeginBatchResponse beginBatchResponse)
+		{
+			return new PerforceExecutor(session.RpcConnection, executeJobTask.JobId, executeJobTask.BatchId, beginBatchResponse.AgentType, executeJobTask.AutoSdkWorkspace, executeJobTask.Workspace, session.WorkingDir);
+		}
+	}
 }
diff --git a/Engine/Source/Programs/Horde/Horde.Agent/Execution/TestExecutor.cs b/Engine/Source/Programs/Horde/Horde.Agent/Execution/TestExecutor.cs
--- a/Engine/Source/Programs/Horde/Horde.Agent/Execution/TestExecutor.cs
+++ b/Engine/Source/Programs/Horde/Horde.Agent/Execution/TestExecutor.cs
@@ -6,12 +6,16 @@ using System.Linq;
 using System.Reflection;
 using System.Threading;
 using System.Threading.Tasks;
+using Amazon.SimpleSystemsManagement.Model;
 using EpicGames.Core;
 using EpicGames.Perforce;
+using Horde.Agent.Execution.Interfaces;
 using Horde.Agent.Parser;
+using Horde.Agent.Services;
 using Horde.Agent.Utility;
 using HordeCommon;
 using HordeCommon.Rpc;
+using HordeCommon.Rpc.Tasks;
 using Microsoft.Extensions.Logging;
 
 namespace Horde.Agent.Execution
@@ -180,4 +184,12 @@ namespace Horde.Agent.Execution
 			return Task.CompletedTask;
 		}
 	}
+
+	class TestExecutorFactory : IExecutorFactory
+	{
+		public IExecutor CreateExecutor(ISession session, ExecuteJobTask executeJobTask, BeginBatchResponse beginBatchResponse)
+		{
+			return new TestExecutor(session.RpcConnection, executeJobTask.JobId, executeJobTask.BatchId, beginBatchResponse.AgentType);
+		}
+	}
 }
diff --git a/Engine/Source/Programs/Horde/Horde.Agent/Horde.Agent.csproj b/Engine/Source/Programs/Horde/Horde.Agent/Horde.Agent.csproj
--- a/Engine/Source/Programs/Horde/Horde.Agent/Horde.Agent.csproj
+++ b/Engine/Source/Programs/Horde/Horde.Agent/Horde.Agent.csproj
@@ -62,12 +62,24 @@
       <CopyToPublishDirectory>PreserveNewest</CopyToPublishDirectory>
       <DependentUpon>appsettings.json</DependentUpon>
     </Content>
+    <Content Include="appsettings.Local.json" Condition="Exists('appsettings.Local.json')">
+      <Generator></Generator>
+      <CopyToOutputDirectory>PreserveNewest</CopyToOutputDirectory>
+      <CopyToPublishDirectory>PreserveNewest</CopyToPublishDirectory>
+      <DependentUpon>appsettings.json</DependentUpon>
+    </Content>
     <Content Include="appsettings.Production.json" Condition="Exists('appsettings.Production.json')">
       <Generator></Generator>
       <CopyToOutputDirectory>PreserveNewest</CopyToOutputDirectory>
       <CopyToPublishDirectory>PreserveNewest</CopyToPublishDirectory>
       <DependentUpon>appsettings.json</DependentUpon>
     </Content>
+    <Content Include="appsettings.User.json" Condition="Exists('appsettings.User.json')">
+      <Generator></Generator>
+      <CopyToOutputDirectory>PreserveNewest</CopyToOutputDirectory>
+      <CopyToPublishDirectory>PreserveNewest</CopyToPublishDirectory>
+      <DependentUpon>appsettings.json</DependentUpon>
+    </Content>
   </ItemGroup>
 
   <ItemGroup>
diff --git a/Engine/Source/Programs/Horde/Horde.Agent/Leases/Handlers/ComputeHandler.cs b/Engine/Source/Programs/Horde/Horde.Agent/Leases/Handlers/ComputeHandler.cs
new file mode 100644
--- /dev/null
+++ b/Engine/Source/Programs/Horde/Horde.Agent/Leases/Handlers/ComputeHandler.cs
@@ -0,0 +1,78 @@
+// Copyright Epic Games, Inc. All Rights Reserved.
+
+using System;
+using System.Threading;
+using System.Threading.Tasks;
+using EpicGames.Core;
+using EpicGames.Horde.Compute;
+using EpicGames.Horde.Storage;
+using Google.Protobuf;
+using Horde.Agent.Execution;
+using Horde.Agent.Services;
+using Horde.Agent.Utility;
+using HordeCommon.Rpc.Tasks;
+using Microsoft.Extensions.Logging;
+
+namespace Horde.Agent.Leases.Handlers
+{
+	class ComputeHandler : LeaseHandler<ComputeTaskMessage>
+	{
+		readonly ILegacyStorageClient _legacyStorageClient;
+		readonly ILogger _logger;
+
+		/// <summary>
+		/// Constructor
+		/// </summary>
+		public ComputeHandler(ILegacyStorageClient legacyStorageClient, ILogger<ComputeHandler> logger)
+		{
+			_legacyStorageClient = legacyStorageClient;
+			_logger = logger;
+		}
+
+		/// <inheritdoc/>
+		public override async Task<LeaseResult> ExecuteAsync(ISession session, string leaseId, ComputeTaskMessage computeTask, CancellationToken cancellationToken)
+		{
+			_logger.LogInformation("Starting compute task (lease {LeaseId})", leaseId);
+
+			ComputeTaskResultMessage result;
+			try
+			{
+				DateTimeOffset actionTaskStartTime = DateTimeOffset.UtcNow;
+				using (IRpcClientRef client = await session.RpcConnection.GetClientRef(new RpcContext(), cancellationToken))
+				{
+					DirectoryReference leaseDir = DirectoryReference.Combine(session.WorkingDir, "Compute", leaseId);
+					DirectoryReference.CreateDirectory(leaseDir);
+
+					ComputeTaskExecutor executor = new ComputeTaskExecutor(_legacyStorageClient, _logger);
+					try
+					{
+						result = await executor.ExecuteAsync(leaseId, computeTask, leaseDir, cancellationToken);
+					}
+					finally
+					{
+						try
+						{
+							DirectoryReference.Delete(leaseDir, true);
+						}
+						catch
+						{
+						}
+					}
+				}
+			}
+			catch (LegacyBlobNotFoundException ex)
+			{
+				_logger.LogError(ex, "Blob not found: {Hash}", ex.Hash);
+				result = new ComputeTaskResultMessage(ComputeTaskOutcome.BlobNotFound, ex.Hash.ToString());
+			}
+			catch (Exception ex)
+			{
+				_logger.LogError(ex, "Exception while executing compute task");
+				result = new ComputeTaskResultMessage(ComputeTaskOutcome.Exception, ex.ToString());
+			}
+
+			return new LeaseResult(result.ToByteArray());
+		}
+	}
+}
+
diff --git a/Engine/Source/Programs/Horde/Horde.Agent/Leases/Handlers/ConformHandler.cs b/Engine/Source/Programs/Horde/Horde.Agent/Leases/Handlers/ConformHandler.cs
new file mode 100644
--- /dev/null
+++ b/Engine/Source/Programs/Horde/Horde.Agent/Leases/Handlers/ConformHandler.cs
@@ -0,0 +1,76 @@
+// Copyright Epic Games, Inc. All Rights Reserved.
+
+using System.Collections.Generic;
+using System.Threading;
+using System.Threading.Tasks;
+using Horde.Agent.Execution;
+using Horde.Agent.Parser;
+using Horde.Agent.Services;
+using Horde.Agent.Utility;
+using HordeCommon.Rpc;
+using HordeCommon.Rpc.Messages;
+using HordeCommon.Rpc.Tasks;
+using Microsoft.Extensions.Logging;
+using Microsoft.Extensions.Options;
+
+namespace Horde.Agent.Leases.Handlers
+{
+	class ConformHandler : LeaseHandler<ConformTask>
+	{
+		readonly AgentSettings _settings;
+		readonly ILogger _defaultLogger;
+
+		/// <summary>
+		/// Constructor
+		/// </summary>
+		public ConformHandler(IOptions<AgentSettings> settings, ILogger<ConformHandler> defaultLogger)
+		{
+			_settings = settings.Value;
+			_defaultLogger = defaultLogger;
+		}
+
+		/// <inheritdoc/>
+		public override async Task<LeaseResult> ExecuteAsync(ISession session, string leaseId, ConformTask conformTask, CancellationToken cancellationToken)
+		{
+			await using JsonRpcLogger conformLogger = new JsonRpcLogger(session.RpcConnection, conformTask.LogId, _defaultLogger);
+
+			conformLogger.LogInformation("Conforming, lease {LeaseId}", leaseId);
+			await session.TerminateProcessesAsync(conformLogger, cancellationToken);
+
+			bool removeUntrackedFiles = conformTask.RemoveUntrackedFiles;
+			IList<AgentWorkspace> pendingWorkspaces = conformTask.Workspaces;
+			for (; ; )
+			{
+				// Run the conform task
+				if (_settings.Executor == ExecutorType.Perforce && _settings.PerforceExecutor.RunConform)
+				{
+					await PerforceExecutor.ConformAsync(session.WorkingDir, pendingWorkspaces, removeUntrackedFiles, conformLogger, cancellationToken);
+				}
+				else
+				{
+					conformLogger.LogInformation("Skipping due to Settings.RunConform flag");
+				}
+
+				// Update the new set of workspaces
+				UpdateAgentWorkspacesRequest request = new UpdateAgentWorkspacesRequest();
+				request.AgentId = session.AgentId;
+				request.Workspaces.AddRange(pendingWorkspaces);
+				request.RemoveUntrackedFiles = removeUntrackedFiles;
+
+				UpdateAgentWorkspacesResponse response = await session.RpcConnection.InvokeAsync(x => x.UpdateAgentWorkspacesAsync(request, null, null, cancellationToken), new RpcContext(), cancellationToken);
+				if (!response.Retry)
+				{
+					conformLogger.LogInformation("Conform finished");
+					break;
+				}
+
+				conformLogger.LogInformation("Pending workspaces have changed - running conform again...");
+				pendingWorkspaces = response.PendingWorkspaces;
+				removeUntrackedFiles = response.RemoveUntrackedFiles;
+			}
+
+			return LeaseResult.Success;
+		}
+	}
+}
+
diff --git a/Engine/Source/Programs/Horde/Horde.Agent/Leases/Handlers/JobHandler.cs b/Engine/Source/Programs/Horde/Horde.Agent/Leases/Handlers/JobHandler.cs
new file mode 100644
--- /dev/null
+++ b/Engine/Source/Programs/Horde/Horde.Agent/Leases/Handlers/JobHandler.cs
@@ -0,0 +1,312 @@
+// Copyright Epic Games, Inc. All Rights Reserved.
+
+using System;
+using System.Diagnostics;
+using System.IO;
+using System.Runtime.InteropServices;
+using System.Threading;
+using System.Threading.Tasks;
+using EpicGames.Core;
+using Grpc.Core;
+using Horde.Agent.Execution.Interfaces;
+using Horde.Agent.Parser;
+using Horde.Agent.Services;
+using Horde.Agent.Utility;
+using HordeCommon;
+using HordeCommon.Rpc;
+using HordeCommon.Rpc.Tasks;
+using Microsoft.Extensions.Logging;
+using Microsoft.Extensions.Options;
+using OpenTracing;
+using OpenTracing.Util;
+
+namespace Horde.Agent.Leases.Handlers
+{
+	class JobHandler : LeaseHandler<ExecuteJobTask>
+	{
+		/// <summary>
+		/// How often to poll the server checking if a step has been aborted
+		/// Exposed as internal to ease testing.
+		/// </summary>
+		internal TimeSpan _stepAbortPollInterval = TimeSpan.FromSeconds(5);
+
+		readonly IExecutorFactory _executorFactory;
+		readonly AgentSettings _settings;
+		readonly ILogger _defaultLogger;
+
+		/// <summary>
+		/// Constructor
+		/// </summary>
+		public JobHandler(IExecutorFactory executorFactory, IOptions<AgentSettings> settings, ILogger<JobHandler> defaultLogger)
+		{
+			_executorFactory = executorFactory;
+			_settings = settings.Value;
+			_defaultLogger = defaultLogger;
+		}
+
+		/// <inheritdoc/>
+		public override async Task<LeaseResult> ExecuteAsync(ISession session, string leaseId, ExecuteJobTask executeTask, CancellationToken cancellationToken)
+		{
+			await using JsonRpcLogger logger = new JsonRpcLogger(session.RpcConnection, executeTask.LogId, executeTask.JobId, executeTask.BatchId, null, null, _defaultLogger);
+
+			logger.LogInformation("Executing job \"{JobName}\", jobId {JobId}, batchId {BatchId}, leaseId {LeaseId}", executeTask.JobName, executeTask.JobId, executeTask.BatchId, leaseId);
+			GlobalTracer.Instance.ActiveSpan?.SetTag("jobId", executeTask.JobId.ToString());
+			GlobalTracer.Instance.ActiveSpan?.SetTag("jobName", executeTask.JobName.ToString());
+			GlobalTracer.Instance.ActiveSpan?.SetTag("batchId", executeTask.BatchId.ToString());
+
+			// Start executing the current batch
+			BeginBatchResponse batch = await session.RpcConnection.InvokeAsync(x => x.BeginBatchAsync(new BeginBatchRequest(executeTask.JobId, executeTask.BatchId, leaseId), null, null, cancellationToken), new RpcContext(), cancellationToken);
+
+			// Execute the batch
+			try
+			{
+				await ExecuteBatchAsync(session, leaseId, executeTask, batch, logger, cancellationToken);
+			}
+			catch (Exception ex)
+			{
+				if (cancellationToken.IsCancellationRequested && ex.IsCancellationException())
+				{
+					logger.LogError("Step was aborted");
+					throw;
+				}
+				else
+				{
+					logger.LogError(ex, "Exception while executing batch: {Ex}", ex);
+				}
+			}
+
+			// If this lease was cancelled, don't bother updating the job state.
+			if (cancellationToken.IsCancellationRequested)
+			{
+				return LeaseResult.Cancelled;
+			}
+
+			// Mark the batch as complete
+			await session.RpcConnection.InvokeAsync(x => x.FinishBatchAsync(new FinishBatchRequest(executeTask.JobId, executeTask.BatchId, leaseId), null, null, cancellationToken), new RpcContext(), cancellationToken);
+			logger.LogInformation("Done.");
+
+			return LeaseResult.Success;
+		}
+
+		/// <summary>
+		/// Executes a batch
+		/// </summary>
+		async Task ExecuteBatchAsync(ISession session, string leaseId, ExecuteJobTask executeTask, BeginBatchResponse batch, ILogger batchLogger, CancellationToken cancellationToken)
+		{
+			IRpcConnection rpcClient = session.RpcConnection;
+
+			batchLogger.LogInformation("Executing batch {BatchId} using {Executor} executor", executeTask.BatchId, _settings.Executor.ToString());
+			await session.TerminateProcessesAsync(batchLogger, cancellationToken);
+
+			// Create an executor for this job
+			IExecutor executor = _executorFactory.CreateExecutor(session, executeTask, batch);
+
+			// Try to initialize the executor
+			batchLogger.LogInformation("Initializing...");
+			using (batchLogger.BeginIndentScope("  "))
+			{
+				using IScope scope = GlobalTracer.Instance.BuildSpan("Initialize").StartActive();
+				await executor.InitializeAsync(batchLogger, cancellationToken);
+			}
+
+			try
+			{
+				// Execute the steps
+				for (; ; )
+				{
+					// Get the next step to execute
+					BeginStepResponse step = await rpcClient.InvokeAsync(x => x.BeginStepAsync(new BeginStepRequest(executeTask.JobId, executeTask.BatchId, leaseId), null, null, cancellationToken), new RpcContext(), cancellationToken);
+					if (step.State == BeginStepResponse.Types.Result.Waiting)
+					{
+						batchLogger.LogInformation("Waiting for dependency to be ready");
+						await Task.Delay(TimeSpan.FromSeconds(20.0), cancellationToken);
+						continue;
+					}
+					else if (step.State == BeginStepResponse.Types.Result.Complete)
+					{
+						break;
+					}
+					else if (step.State != BeginStepResponse.Types.Result.Ready)
+					{
+						batchLogger.LogError("Unexpected step state: {StepState}", step.State);
+						break;
+					}
+
+					// Get current disk space available. This will allow us to more easily spot steps that eat up a lot of disk space.
+					string? driveName;
+					if (RuntimeInformation.IsOSPlatform(OSPlatform.Windows))
+					{
+						driveName = Path.GetPathRoot(session.WorkingDir.FullName);
+					}
+					else
+					{
+						driveName = session.WorkingDir.FullName;
+					}
+
+					float availableFreeSpace = 0;
+					if (driveName != null)
+					{
+						try
+						{
+							DriveInfo info = new DriveInfo(driveName);
+							availableFreeSpace = (1.0f * info.AvailableFreeSpace) / 1024 / 1024 / 1024;
+						}
+						catch (Exception ex)
+						{
+							batchLogger.LogWarning(ex, "Unable to query disk info for path '{DriveName}'", driveName);
+						}
+					}
+
+					// Print the new state
+					Stopwatch stepTimer = Stopwatch.StartNew();
+
+					batchLogger.LogInformation("Starting job {JobId}, batch {BatchId}, step {StepId} (Drive Space Left: {DriveSpaceRemaining} GB)", executeTask.JobId, executeTask.BatchId, step.StepId, availableFreeSpace.ToString("F1"));
+
+					// Create a trace span
+					using IScope scope = GlobalTracer.Instance.BuildSpan("Execute").WithResourceName(step.Name).StartActive();
+					scope.Span.SetTag("stepId", step.StepId);
+					scope.Span.SetTag("logId", step.LogId);
+					//				using IDisposable TraceProperty = LogContext.PushProperty("dd.trace_id", CorrelationIdentifier.TraceId.ToString());
+					//				using IDisposable SpanProperty = LogContext.PushProperty("dd.span_id", CorrelationIdentifier.SpanId.ToString());
+
+					// Update the context to include information about this step
+					JobStepOutcome stepOutcome;
+					JobStepState stepState;
+					using (batchLogger.BeginIndentScope("  "))
+					{
+						// Start writing to the log file
+#pragma warning disable CA2000 // Dispose objects before losing scope
+						await using (JsonRpcLogger stepLogger = new JsonRpcLogger(rpcClient, step.LogId, executeTask.JobId, executeTask.BatchId, step.StepId, step.Warnings, _defaultLogger))
+						{
+							// Execute the task
+							ILogger forwardingLogger = new DefaultLoggerIndentHandler(stepLogger);
+							if (_settings.WriteStepOutputToLogger)
+							{
+								forwardingLogger = new ForwardingLogger(_defaultLogger, forwardingLogger);
+							}
+
+							using CancellationTokenSource stepPollCancelSource = new CancellationTokenSource();
+							using CancellationTokenSource stepAbortSource = new CancellationTokenSource();
+							TaskCompletionSource<bool> stepFinishedSource = new TaskCompletionSource<bool>();
+							Task stepPollTask = Task.Run(() => PollForStepAbort(rpcClient, executeTask.JobId, executeTask.BatchId, step.StepId, stepAbortSource, stepFinishedSource.Task, stepPollCancelSource.Token), cancellationToken);
+
+							try
+							{
+								(stepOutcome, stepState) = await ExecuteStepAsync(executor, step, forwardingLogger, cancellationToken, stepAbortSource.Token);
+							}
+							finally
+							{
+								// Will get called even when cancellation token for the lease/batch fires
+								stepFinishedSource.SetResult(true); // Tell background poll task to stop
+								await stepPollTask;
+							}
+
+							// Kill any processes spawned by the step
+							await session.TerminateProcessesAsync(batchLogger, cancellationToken);
+
+							// Wait for the logger to finish
+							await stepLogger.StopAsync();
+
+							// Reflect the warnings/errors in the step outcome
+							if (stepOutcome > stepLogger.Outcome)
+							{
+								stepOutcome = stepLogger.Outcome;
+							}
+						}
+#pragma warning restore CA2000 // Dispose objects before losing scope
+
+						// Update the server with the outcome from the step
+						batchLogger.LogInformation("Marking step as complete (Outcome={Outcome}, State={StepState})", stepOutcome, stepState);
+						await rpcClient.InvokeAsync(x => x.UpdateStepAsync(new UpdateStepRequest(executeTask.JobId, executeTask.BatchId, step.StepId, stepState, stepOutcome), null, null, cancellationToken), new RpcContext(), cancellationToken);
+					}
+
+					// Print the finishing state
+					stepTimer.Stop();
+					batchLogger.LogInformation("Completed in {Time}", stepTimer.Elapsed);
+				}
+			}
+			catch (Exception ex)
+			{
+				if (cancellationToken.IsCancellationRequested && ex.IsCancellationException())
+				{
+					batchLogger.LogError("Step was aborted");
+				}
+				else
+				{
+					batchLogger.LogError(ex, "Exception while executing batch: {Ex}", ex);
+				}
+			}
+
+			// Clean the environment
+			batchLogger.LogInformation("Finalizing...");
+			using (batchLogger.BeginIndentScope("  "))
+			{
+				using IScope scope = GlobalTracer.Instance.BuildSpan("Finalize").StartActive();
+				await executor.FinalizeAsync(batchLogger, CancellationToken.None);
+			}
+		}
+
+		/// <summary>
+		/// Executes a step
+		/// </summary>
+		/// <param name="executor">The executor to run this step</param>
+		/// <param name="step">Step to execute</param>
+		/// <param name="stepLogger">Logger for the step</param>
+		/// <param name="cancellationToken">Cancellation token to abort the batch</param>
+		/// <param name="stepCancellationToken">Cancellation token to abort only this individual step</param>
+		/// <returns>Async task</returns>
+		internal static async Task<(JobStepOutcome, JobStepState)> ExecuteStepAsync(IExecutor executor, BeginStepResponse step, ILogger stepLogger, CancellationToken cancellationToken, CancellationToken stepCancellationToken)
+		{
+			using CancellationTokenSource combined = CancellationTokenSource.CreateLinkedTokenSource(cancellationToken, stepCancellationToken);
+			try
+			{
+				JobStepOutcome stepOutcome = await executor.RunAsync(step, stepLogger, combined.Token);
+				return (stepOutcome, JobStepState.Completed);
+			}
+			catch (Exception ex)
+			{
+				if (cancellationToken.IsCancellationRequested && ex.IsCancellationException())
+				{
+					stepLogger.LogError("The step was cancelled by batch/lease");
+					throw;
+				}
+
+				if (stepCancellationToken.IsCancellationRequested && ex.IsCancellationException())
+				{
+					stepLogger.LogError("The step was intentionally cancelled");
+					return (JobStepOutcome.Failure, JobStepState.Aborted);
+				}
+
+				stepLogger.LogError(ex, "Exception while executing step: {Ex}", ex);
+				return (JobStepOutcome.Failure, JobStepState.Completed);
+			}
+		}
+
+		internal async Task PollForStepAbort(IRpcConnection rpcClient, string jobId, string batchId, string stepId, CancellationTokenSource stepCancelSource, Task finishedTask, CancellationToken cancellationToken)
+		{
+			Stopwatch timer = Stopwatch.StartNew();
+			while (!finishedTask.IsCompleted)
+			{
+				try
+				{
+					GetStepResponse res = await rpcClient.InvokeAsync(x => x.GetStepAsync(new GetStepRequest(jobId, batchId, stepId), null, null, cancellationToken), new RpcContext(), cancellationToken);
+					if (res.AbortRequested)
+					{
+						_defaultLogger.LogDebug("Step was aborted by server (JobId={JobId} BatchId={BatchId} StepId={StepId})", jobId, batchId, stepId);
+						stepCancelSource.Cancel();
+						break;
+					}
+				}
+				catch (RpcException ex)
+				{
+					_defaultLogger.LogError(ex, "Poll for step abort has failed. Aborting (JobId={JobId} BatchId={BatchId} StepId={StepId})", jobId, batchId, stepId);
+					stepCancelSource.Cancel();
+					break;
+				}
+
+				await Task.WhenAny(Task.Delay(_stepAbortPollInterval, cancellationToken), finishedTask);
+			}
+		}
+	}
+}
+
diff --git a/Engine/Source/Programs/Horde/Horde.Agent/Leases/Handlers/RestartHandler.cs b/Engine/Source/Programs/Horde/Horde.Agent/Leases/Handlers/RestartHandler.cs
new file mode 100644
--- /dev/null
+++ b/Engine/Source/Programs/Horde/Horde.Agent/Leases/Handlers/RestartHandler.cs
@@ -0,0 +1,27 @@
+// Copyright Epic Games, Inc. All Rights Reserved.
+
+using System.Threading;
+using System.Threading.Tasks;
+using Horde.Agent.Services;
+using HordeCommon.Rpc.Tasks;
+using Microsoft.Extensions.Logging;
+
+namespace Horde.Agent.Leases.Handlers
+{
+	class RestartHandler : LeaseHandler<RestartTask>
+	{
+		readonly ILogger _logger;
+
+		public RestartHandler(ILogger<RestartHandler> logger)
+		{
+			_logger = logger;
+		}
+
+		/// <inheritdoc/>
+		public override Task<LeaseResult> ExecuteAsync(ISession session, string leaseId, RestartTask task, CancellationToken cancellationToken)
+		{
+			_logger.LogInformation("Setting shutdown flag");
+			return Task.FromResult(new LeaseResult(SessionResult.Shutdown));
+		}
+	}
+}
diff --git a/Engine/Source/Programs/Horde/Horde.Agent/Leases/Handlers/ShutdownHandler.cs b/Engine/Source/Programs/Horde/Horde.Agent/Leases/Handlers/ShutdownHandler.cs
new file mode 100644
--- /dev/null
+++ b/Engine/Source/Programs/Horde/Horde.Agent/Leases/Handlers/ShutdownHandler.cs
@@ -0,0 +1,28 @@
+// Copyright Epic Games, Inc. All Rights Reserved.
+
+using System.Threading;
+using System.Threading.Tasks;
+using Horde.Agent.Services;
+using HordeCommon.Rpc.Tasks;
+using Microsoft.Extensions.Logging;
+
+namespace Horde.Agent.Leases.Handlers
+{
+	class ShutdownHandler : LeaseHandler<ShutdownTask>
+	{
+		readonly ILogger _logger;
+
+		public ShutdownHandler(ILogger<ShutdownHandler> logger)
+		{
+			_logger = logger;
+		}
+
+		/// <inheritdoc/>
+		public override Task<LeaseResult> ExecuteAsync(ISession session, string leaseId, ShutdownTask task, CancellationToken cancellationToken)
+		{
+			_logger.LogInformation("Setting shutdown flag");
+			return Task.FromResult(new LeaseResult(SessionResult.Shutdown));
+		}
+	}
+}
+
diff --git a/Engine/Source/Programs/Horde/Horde.Agent/Leases/Handlers/UpgradeHandler.cs b/Engine/Source/Programs/Horde/Horde.Agent/Leases/Handlers/UpgradeHandler.cs
new file mode 100644
--- /dev/null
+++ b/Engine/Source/Programs/Horde/Horde.Agent/Leases/Handlers/UpgradeHandler.cs
@@ -0,0 +1,167 @@
+// Copyright Epic Games, Inc. All Rights Reserved.
+
+using System;
+using System.Collections.Generic;
+using System.Diagnostics;
+using System.IO;
+using System.IO.Compression;
+using System.Reflection;
+using System.Text;
+using System.Threading;
+using System.Threading.Tasks;
+using EpicGames.Core;
+using Grpc.Core;
+using Horde.Agent.Parser;
+using Horde.Agent.Services;
+using Horde.Agent.Utility;
+using HordeCommon.Rpc;
+using HordeCommon.Rpc.Tasks;
+using Microsoft.Extensions.Logging;
+
+namespace Horde.Agent.Leases.Handlers
+{
+	class UpgradeHandler : LeaseHandler<UpgradeTask>
+	{
+		readonly ILogger _defaultLogger;
+
+		public UpgradeHandler(ILogger<UpgradeHandler> defaultLogger)
+		{
+			_defaultLogger = defaultLogger;
+		}
+
+		/// <inheritdoc/>
+		public override async Task<LeaseResult> ExecuteAsync(ISession session, string leaseId, UpgradeTask task, CancellationToken cancellationToken)
+		{
+			await using JsonRpcLogger logger = new JsonRpcLogger(session.RpcConnection, task.LogId, _defaultLogger);
+
+			string requiredVersion = task.SoftwareId;
+
+			// Check if we're running the right version
+			if (requiredVersion != null && requiredVersion != Program.Version)
+			{
+				logger.LogInformation("Upgrading from {CurrentVersion} to {TargetVersion}", Program.Version, requiredVersion);
+
+				// Clear out the working directory
+				DirectoryReference upgradeDir = DirectoryReference.Combine(session.WorkingDir, "Upgrade");
+				DirectoryReference.CreateDirectory(upgradeDir);
+				await DeleteDirectoryContentsAsync(new DirectoryInfo(upgradeDir.FullName));
+
+				// Download the new software
+				FileInfo outputFile = new FileInfo(Path.Combine(upgradeDir.FullName, "Agent.zip"));
+				using (IRpcClientRef rpcClientRef = await session.RpcConnection.GetClientRef(new RpcContext(), cancellationToken))
+				using (AsyncServerStreamingCall<DownloadSoftwareResponse> cursor = rpcClientRef.Client.DownloadSoftware(new DownloadSoftwareRequest(requiredVersion), null, null, cancellationToken))
+				{
+					using (Stream outputStream = outputFile.Open(FileMode.Create))
+					{
+						while (await cursor.ResponseStream.MoveNext(cancellationToken))
+						{
+							outputStream.Write(cursor.ResponseStream.Current.Data.Span);
+						}
+					}
+				}
+
+				// Extract it to a temporary directory
+				DirectoryReference extractedDir = DirectoryReference.Combine(upgradeDir, "Extracted");
+				DirectoryReference.CreateDirectory(extractedDir);
+				ZipFile.ExtractToDirectory(outputFile.FullName, extractedDir.FullName);
+
+				//				// Debug code for updating an agent with the local version
+				//				foreach (FileInfo SourceFile in new FileInfo(Assembly.GetExecutingAssembly().Location).Directory.EnumerateFiles())
+				//				{
+				//					SourceFile.CopyTo(Path.Combine(ExtractedDir.FullName, SourceFile.Name), true);
+				//				}
+
+				// Get the current process and assembly. This may be different if running through dotnet.exe rather than a native PE image.
+				FileReference assemblyFileName = new FileReference(Assembly.GetExecutingAssembly().Location);
+
+				// Spawn the other process
+				using (Process process = new Process())
+				{
+					StringBuilder arguments = new StringBuilder();
+
+					DirectoryReference targetDir = assemblyFileName.Directory;
+
+					// We were launched via an external application (presumably dotnet.exe). Do the same thing again.
+					FileReference newAssemblyFileName = FileReference.Combine(extractedDir, assemblyFileName.MakeRelativeTo(targetDir));
+					if (!FileReference.Exists(newAssemblyFileName))
+					{
+						logger.LogError("Unable to find {AgentExe} in extracted archive", newAssemblyFileName);
+						return LeaseResult.Failed;
+					}
+
+					process.StartInfo.FileName = "dotnet";
+
+					StringBuilder currentArguments = new StringBuilder();
+					foreach (string arg in Program.Args)
+					{
+						currentArguments.AppendArgument(arg);
+					}
+
+					arguments.AppendArgument(newAssemblyFileName.FullName);
+					arguments.AppendArgument("Service");
+					arguments.AppendArgument("Upgrade");
+					arguments.AppendArgument("-ProcessId=", Environment.ProcessId.ToString());
+					arguments.AppendArgument("-TargetDir=", targetDir.FullName);
+					arguments.AppendArgument("-Arguments=", currentArguments.ToString());
+
+					process.StartInfo.Arguments = arguments.ToString();
+					process.StartInfo.UseShellExecute = false;
+					process.EnableRaisingEvents = true;
+
+					StringBuilder launchCommand = new StringBuilder();
+					launchCommand.AppendArgument(process.StartInfo.FileName);
+					launchCommand.Append(' ');
+					launchCommand.Append(arguments);
+					logger.LogInformation("Launching: {Launch}", launchCommand.ToString());
+
+					TaskCompletionSource<int> exitCodeSource = new TaskCompletionSource<int>();
+					process.Exited += (sender, args) => { exitCodeSource.SetResult(process.ExitCode); };
+
+					process.Start();
+
+					using (cancellationToken.Register(() => { exitCodeSource.SetResult(0); }))
+					{
+						await exitCodeSource.Task;
+					}
+				}
+			}
+
+			return LeaseResult.Success;
+		}
+
+		/// <summary>
+		/// Delete the contents of a directory without deleting it itself
+		/// </summary>
+		/// <param name="baseDir">Directory to clean</param>
+		/// <returns>Async task</returns>
+		static async Task DeleteDirectoryContentsAsync(DirectoryInfo baseDir)
+		{
+			List<Task> childTasks = new List<Task>();
+			foreach (DirectoryInfo subDir in baseDir.EnumerateDirectories())
+			{
+				childTasks.Add(Task.Run(() => DeleteDirectoryAsync(subDir)));
+			}
+			foreach (FileInfo file in baseDir.EnumerateFiles())
+			{
+				file.Attributes = FileAttributes.Normal;
+				file.Delete();
+			}
+			foreach (Task childTask in childTasks)
+			{
+				await childTask;
+			}
+		}
+
+		/// <summary>
+		/// Deletes a directory and its contents
+		/// </summary>
+		/// <param name="baseDir">Directory to delete</param>
+		/// <returns>Async task</returns>
+		static async Task DeleteDirectoryAsync(DirectoryInfo baseDir)
+		{
+			await DeleteDirectoryContentsAsync(baseDir);
+			baseDir.Delete();
+		}
+	}
+}
+
diff --git a/Engine/Source/Programs/Horde/Horde.Agent/Leases/LeaseHandler.cs b/Engine/Source/Programs/Horde/Horde.Agent/Leases/LeaseHandler.cs
new file mode 100644
--- /dev/null
+++ b/Engine/Source/Programs/Horde/Horde.Agent/Leases/LeaseHandler.cs
@@ -0,0 +1,52 @@
+// Copyright Epic Games, Inc. All Rights Reserved.
+
+using System.Threading;
+using System.Threading.Tasks;
+using Google.Protobuf;
+using Google.Protobuf.Reflection;
+using Google.Protobuf.WellKnownTypes;
+using Horde.Agent.Services;
+
+namespace Horde.Agent.Leases
+{
+	/// <summary>
+	/// Handles execution of a specific lease type
+	/// </summary>
+	abstract class LeaseHandler
+	{
+		/// <summary>
+		/// Returns protobuf type urls for the handled message types
+		/// </summary>
+		public abstract string LeaseType { get; }
+
+		/// <summary>
+		/// Executes a lease
+		/// </summary>
+		/// <returns>Result for the lease</returns>
+		public abstract Task<LeaseResult> ExecuteAsync(ISession session, string leaseId, Any message, CancellationToken cancellationToken);
+	}
+
+	/// <summary>
+	/// Implementation of <see cref="LeaseHandler"/> for a specific lease type
+	/// </summary>
+	/// <typeparam name="T">Type of the lease message</typeparam>
+	abstract class LeaseHandler<T> : LeaseHandler where T : IMessage<T>, new()
+	{
+		/// <summary>
+		/// Static for the message type descriptor
+		/// </summary>
+		public static MessageDescriptor Descriptor { get; } = new T().Descriptor;
+
+		/// <inheritdoc/>
+		public override string LeaseType { get; } = $"type.googleapis.com/{Descriptor.Name}";
+
+		/// <inheritdoc/>
+		public override Task<LeaseResult> ExecuteAsync(ISession session, string leaseId, Any message, CancellationToken cancellationToken)
+		{
+			return ExecuteAsync(session, leaseId, message.Unpack<T>(), cancellationToken); 
+		}
+
+		/// <inheritdoc/>
+		public abstract Task<LeaseResult> ExecuteAsync(ISession session, string leaseId, T message, CancellationToken cancellationToken);
+	}
+}
diff --git a/Engine/Source/Programs/Horde/Horde.Agent/Leases/LeaseManager.cs b/Engine/Source/Programs/Horde/Horde.Agent/Leases/LeaseManager.cs
new file mode 100644
--- /dev/null
+++ b/Engine/Source/Programs/Horde/Horde.Agent/Leases/LeaseManager.cs
@@ -0,0 +1,569 @@
+// Copyright Epic Games, Inc. All Rights Reserved.
+
+using System;
+using System.Collections.Generic;
+using System.Diagnostics;
+using System.Linq;
+using System.Threading;
+using System.Threading.Tasks;
+using EpicGames.Core;
+using Google.Protobuf;
+using Google.Protobuf.WellKnownTypes;
+using Grpc.Core;
+using Horde.Agent.Services;
+using Horde.Agent.Utility;
+using HordeCommon;
+using HordeCommon.Rpc;
+using HordeCommon.Rpc.Messages;
+using Microsoft.Extensions.DependencyInjection;
+using Microsoft.Extensions.Logging;
+using Microsoft.Extensions.Options;
+using OpenTracing;
+using OpenTracing.Util;
+
+namespace Horde.Agent.Leases
+{
+	/// <summary>
+	/// Outcome from a session
+	/// </summary>
+	enum SessionResult
+	{
+		/// <summary>
+		/// Create a new session (server migration, etc...)
+		/// </summary>
+		Continue,
+
+		/// <summary>
+		/// Session completed normally, terminate the application.
+		/// </summary>
+		Terminate,
+
+		/// <summary>
+		/// Machine should be scheduled to restart
+		/// </summary>
+		Restart,
+
+		/// <summary>
+		/// Machine should be scheduled to shut down
+		/// </summary>
+		Shutdown,
+	}
+
+	/// <summary>
+	/// Implements the message handling loop for an agent. Runs asynchronously until disposed.
+	/// </summary>
+	class LeaseManager
+	{
+		/// <summary>
+		/// Stores information about an active session
+		/// </summary>
+		internal class LeaseInfo
+		{
+			/// <summary>
+			/// The worker lease state
+			/// </summary>
+			public Lease Lease { get; set; }
+
+			/// <summary>
+			/// The task being executed for this lease
+			/// </summary>
+			public Task? Task { get; set; }
+
+			/// <summary>
+			/// Source for cancellation tokens for this session.
+			/// </summary>
+			public CancellationTokenSource CancellationTokenSource { get; }
+
+			/// <summary>
+			/// Constructor
+			/// </summary>
+			/// <param name="lease">The worker lease state</param>
+			public LeaseInfo(Lease lease)
+			{
+				Lease = lease;
+				CancellationTokenSource = new CancellationTokenSource();
+			}
+		}
+
+		/// <summary>
+		/// List of processes that should be terminated before running a job
+		/// </summary>
+		private readonly HashSet<string> _processNamesToTerminate = new HashSet<string>(StringComparer.OrdinalIgnoreCase);
+
+		/// <summary>
+		/// Object used for controlling access to the access tokens and active sessions list
+		/// </summary>
+		readonly object _lockObject = new object();
+
+		/// <summary>
+		/// The list of active leases.
+		/// </summary>
+		readonly List<LeaseInfo> _activeLeases = new List<LeaseInfo>();
+
+		/// <summary>
+		/// Number of leases completed
+		/// </summary>
+		public int NumLeasesCompleted { get; private set; }
+
+		/// <summary>
+		/// Whether the agent is currently in an unhealthy state
+		/// </summary>
+		readonly bool _unhealthy = false;
+
+		/// <summary>
+		/// Result from executing this session
+		/// </summary>
+		SessionResult _result = SessionResult.Continue;
+
+		/// <summary>
+		/// Task completion source used to trigger the background thread to update the leases. Must take a lock on LockObject before 
+		/// </summary>
+		readonly AsyncEvent _updateLeasesEvent = new AsyncEvent();
+
+		/// <summary>
+		/// Number of times UpdateSession has failed
+		/// </summary>
+		int _updateSessionFailures;
+
+		/// <summary>
+		/// Delegate for lease active events
+		/// </summary>
+		public delegate void LeaseActiveEvent(Lease lease);
+
+		/// <summary>
+		/// Event triggered when a lease is accepted and set to active
+		/// </summary>
+		public event LeaseActiveEvent? OnLeaseActive;
+
+		/// <summary>
+		/// How long to wait before trying to reacquire a new connection
+		/// Exposed as internal to ease testing. Using a lower delay can speed up tests. 
+		/// </summary>
+		internal TimeSpan _rpcConnectionRetryDelay = TimeSpan.FromSeconds(5);
+
+		readonly ISession _session;
+		readonly CapabilitiesService _capabilitiesService;
+		readonly Dictionary<string, LeaseHandler> _typeUrlToLeaseHandler;
+		readonly AgentSettings _settings;
+		readonly ILogger _logger;
+
+		public LeaseManager(ISession session, CapabilitiesService capabilitiesService, IEnumerable<LeaseHandler> leaseHandlers, IOptions<AgentSettings> settings, ILogger logger)
+		{
+			_session = session;
+			_capabilitiesService = capabilitiesService;
+			_typeUrlToLeaseHandler = leaseHandlers.ToDictionary(x => x.LeaseType, x => x);
+			_settings = settings.Value;
+			_logger = logger;
+
+			_processNamesToTerminate.UnionWith(_settings.ProcessNamesToTerminate);
+		}
+
+		public LeaseManager(ISession session, IServiceProvider serviceProvider)
+			: this(session, serviceProvider.GetRequiredService<CapabilitiesService>(), serviceProvider.GetRequiredService<IEnumerable<LeaseHandler>>(), serviceProvider.GetRequiredService<IOptions<AgentSettings>>(), serviceProvider.GetRequiredService<ILogger<LeaseManager>>())
+		{
+		}
+
+		public async Task<SessionResult> RunAsync(bool shutdownAfterFinishedLease, CancellationToken stoppingToken)
+		{
+			SessionResult result;
+			try
+			{
+				result = await HandleSessionAsync(shutdownAfterFinishedLease, stoppingToken);
+			}
+			catch (OperationCanceledException ex) when (stoppingToken.IsCancellationRequested)
+			{
+				_logger.LogInformation(ex, "Execution cancelled");
+				result = SessionResult.Terminate;
+			}
+			catch (Exception ex)
+			{
+				_logger.LogError(ex, "Exception while executing session. Restarting.");
+				result = SessionResult.Continue;
+			}
+
+			while (_activeLeases.Count > 0)
+			{
+				try
+				{
+					_logger.LogInformation("Draining leases... ({NumLeases} remaining)", _activeLeases.Count);
+					await DrainLeasesAsync();
+				}
+				catch (Exception ex)
+				{
+					_logger.LogError(ex, "Exception while draining leases. Agent may be in an inconsistent state.");
+					await Task.Delay(TimeSpan.FromSeconds(30.0), CancellationToken.None);
+				}
+			}
+			return result;
+		}
+
+		async Task DrainLeasesAsync()
+		{
+			for (int idx = 0; idx < _activeLeases.Count; idx++)
+			{
+				LeaseInfo activeLease = _activeLeases[idx];
+				if (activeLease.Task == null)
+				{
+					_activeLeases.RemoveAt(idx--);
+					_logger.LogInformation("Removed lease {LeaseId}", activeLease.Lease.Id);
+				}
+				else
+				{
+					_logger.LogInformation("Cancelling active lease {LeaseId}", activeLease.Lease.Id);
+					activeLease.CancellationTokenSource.Cancel();
+				}
+			}
+
+			while (_activeLeases.Count > 0)
+			{
+				List<Task> tasks = _activeLeases.Select(x => x.Task!).ToList();
+				tasks.Add(Task.Delay(TimeSpan.FromMinutes(1.0)));
+				await Task.WhenAny(tasks);
+
+				for (int idx = 0; idx < _activeLeases.Count; idx++)
+				{
+					LeaseInfo activeLease = _activeLeases[idx];
+					if (activeLease.Task!.IsCompleted)
+					{
+						_activeLeases.RemoveAt(idx--);
+						try
+						{
+							await activeLease.Task;
+						}
+						catch (OperationCanceledException)
+						{
+						}
+						catch (Exception ex)
+						{
+							_logger.LogError(ex, "Lease {LeaseId} threw an exception while terminating", activeLease.Lease.Id);
+						}
+						_logger.LogInformation("Lease {LeaseId} has completed", activeLease.Lease.Id);
+					}
+					else
+					{
+						_logger.LogInformation("Still waiting for lease {LeaseId} to terminate...", activeLease.Lease.Id);
+					}
+				}
+			}
+		}
+
+		async Task<SessionResult> HandleSessionAsync(bool shutdownAfterFinishedLease, CancellationToken stoppingToken)
+		{
+			IRpcConnection rpcCon = _session.RpcConnection;
+
+			// Terminate any remaining child processes from other instances
+			await _session.TerminateProcessesAsync(_logger, stoppingToken);
+
+			// Track how many updates we get in 10 seconds. We'll start rate limiting this if it looks like we've got a problem that's causing us to spam the server.
+			Stopwatch updateTimer = Stopwatch.StartNew();
+			Queue<TimeSpan> updateTimes = new Queue<TimeSpan>();
+
+			// Loop until we're ready to exit
+			Stopwatch updateCapabilitiesTimer = Stopwatch.StartNew();
+			for (; ; )
+			{
+				Task waitTask = _updateLeasesEvent.Task;
+
+				// Flag for whether the service is stopping
+				bool stopping = false;
+				if (stoppingToken.IsCancellationRequested)
+				{
+					_logger.LogInformation("Cancellation from token requested");
+					stopping = true;
+				}
+
+				if (_result != SessionResult.Continue)
+				{
+					_logger.LogInformation("Session termination requested (result: {Result})", _result);
+					stopping = true;
+				}
+
+				// Build the next update request
+				UpdateSessionRequest updateSessionRequest = new UpdateSessionRequest();
+				updateSessionRequest.AgentId = _session.AgentId;
+				updateSessionRequest.SessionId = _session.SessionId;
+
+				// Get the new the lease states. If a restart is requested and we have no active leases, signal to the server that we're stopping.
+				lock (_lockObject)
+				{
+					foreach (LeaseInfo leaseInfo in _activeLeases)
+					{
+						updateSessionRequest.Leases.Add(new Lease(leaseInfo.Lease));
+					}
+					if (_result != SessionResult.Continue && _activeLeases.Count == 0)
+					{
+						stopping = true;
+					}
+				}
+
+				// Get the new agent status
+				if (stopping)
+				{
+					updateSessionRequest.Status = AgentStatus.Stopping;
+				}
+				else if (_unhealthy)
+				{
+					updateSessionRequest.Status = AgentStatus.Unhealthy;
+				}
+				else
+				{
+					updateSessionRequest.Status = AgentStatus.Ok;
+				}
+
+				// Update the capabilities every 5m
+				if (updateCapabilitiesTimer.Elapsed > TimeSpan.FromMinutes(5.0))
+				{
+					try
+					{
+						updateSessionRequest.Capabilities = await _capabilitiesService.GetCapabilitiesAsync(_session.WorkingDir);
+					}
+					catch (Exception ex)
+					{
+						_logger.LogWarning(ex, "Unable to query agent capabilities. Ignoring.");
+					}
+					updateCapabilitiesTimer.Restart();
+				}
+
+				// Complete the wait task if we subsequently stop
+				using (stopping ? (CancellationTokenRegistration?)null : stoppingToken.Register(() => _updateLeasesEvent.Set()))
+				{
+					// Update the state with the server
+					UpdateSessionResponse? updateSessionResponse = null;
+					using (IRpcClientRef? rpcClientRef = rpcCon.TryGetClientRef(new RpcContext()))
+					{
+						if (rpcClientRef == null)
+						{
+							// An RpcConnection has not yet been established, wait a period of time and try again
+							await Task.WhenAny(Task.Delay(_rpcConnectionRetryDelay, stoppingToken), waitTask);
+						}
+						else
+						{
+							updateSessionResponse = await UpdateSessionAsync(rpcClientRef, updateSessionRequest, waitTask);
+						}
+					}
+
+					lock (_lockObject)
+					{
+						// Now reconcile the local state to match what the server reports
+						if (updateSessionResponse != null)
+						{
+							bool atLeastOneLeaseFinished = _activeLeases.Any(x => x.Lease.State is LeaseState.Completed or LeaseState.Cancelled);
+							if (atLeastOneLeaseFinished && shutdownAfterFinishedLease)
+							{
+								_logger.LogInformation("At least one lease executed. Requesting shutdown.");
+								_result = SessionResult.Terminate;
+							}
+
+							// Remove any leases which have completed
+							int numRemoved = _activeLeases.RemoveAll(x => (x.Lease.State == LeaseState.Completed || x.Lease.State == LeaseState.Cancelled) && !updateSessionResponse.Leases.Any(y => y.Id == x.Lease.Id && y.State != LeaseState.Cancelled));
+							NumLeasesCompleted += numRemoved;
+
+							// Create any new leases and cancel any running leases
+							foreach (Lease serverLease in updateSessionResponse.Leases)
+							{
+								if (serverLease.State == LeaseState.Cancelled)
+								{
+									LeaseInfo? info = _activeLeases.FirstOrDefault(x => x.Lease.Id == serverLease.Id);
+									if (info != null)
+									{
+										_logger.LogInformation("Cancelling lease {LeaseId}", serverLease.Id);
+										info.CancellationTokenSource.Cancel();
+									}
+								}
+								if (serverLease.State == LeaseState.Pending && !_activeLeases.Any(x => x.Lease.Id == serverLease.Id))
+								{
+									serverLease.State = LeaseState.Active;
+
+									_logger.LogInformation("Adding lease {LeaseId}", serverLease.Id);
+									LeaseInfo info = new LeaseInfo(serverLease);
+									info.Task = Task.Run(() => HandleLeaseAsync(_session, info), CancellationToken.None);
+									_activeLeases.Add(info);
+									OnLeaseActive?.Invoke(serverLease);
+								}
+							}
+						}
+
+						// If there's nothing still running and cancellation was requested, exit
+						if (_activeLeases.Count == 0 && updateSessionRequest.Status == AgentStatus.Stopping)
+						{
+							_logger.LogInformation("No leases are active. Agent is stopping.");
+							break;
+						}
+					}
+				}
+
+				// Update the historical update times
+				TimeSpan updateTime = updateTimer.Elapsed;
+				while (updateTimes.TryPeek(out TimeSpan firstTime) && firstTime < updateTime - TimeSpan.FromMinutes(1.0))
+				{
+					updateTimes.Dequeue();
+				}
+				updateTimes.Enqueue(updateTime);
+
+				// If we're updating too much, introduce an artificial delay
+				if (updateTimes.Count > 60)
+				{
+					_logger.LogWarning("Agent is issuing large number of UpdateSession() calls. Delaying for 10 seconds.");
+					await Task.Delay(TimeSpan.FromSeconds(10.0), stoppingToken);
+				}
+			}
+			return _result;
+		}
+
+		/// <summary>
+		/// Wrapper for <see cref="UpdateSessionInternalAsync"/> which filters/logs exceptions
+		/// </summary>
+		/// <param name="rpcClientRef">The RPC client connection</param>
+		/// <param name="updateSessionRequest">The session update request</param>
+		/// <param name="waitTask">Task which can be used to jump out of the update early</param>
+		/// <returns>Response from the call</returns>
+		async Task<UpdateSessionResponse?> UpdateSessionAsync(IRpcClientRef rpcClientRef, UpdateSessionRequest updateSessionRequest, Task waitTask)
+		{
+			UpdateSessionResponse? updateSessionResponse = null;
+			try
+			{
+				updateSessionResponse = await UpdateSessionInternalAsync(rpcClientRef, updateSessionRequest, waitTask);
+				_updateSessionFailures = 0;
+			}
+			catch (RpcException ex)
+			{
+				if (++_updateSessionFailures >= 3)
+				{
+					throw;
+				}
+				else if (ex.StatusCode == StatusCode.Unavailable)
+				{
+					_logger.LogInformation(ex, "Service unavailable while calling UpdateSessionAsync(), will retry");
+				}
+				else
+				{
+					_logger.LogError(ex, "Error while executing RPC. Will retry.");
+				}
+			}
+			return updateSessionResponse;
+		}
+
+		/// <summary>
+		/// Tries to update the session state on the server.
+		/// 
+		/// This operation is a little gnarly due to the fact that we want to long-poll for the result.
+		/// Since we're doing the update via a gRPC call, the way to do that without using cancellation tokens is to keep the request stream open
+		/// until we want to terminate the call (see [REDACTED_PROFILE_URL]/grpc/issues/8277). In order to do that, we need to make a 
+		/// bidirectional streaming call, even though we only expect one response/response.
+		/// </summary>
+		/// <param name="rpcClientRef">The RPC client</param>
+		/// <param name="request">The session update request</param>
+		/// <param name="waitTask">Task to use to terminate the wait</param>
+		/// <returns>The response object</returns>
+		async Task<UpdateSessionResponse?> UpdateSessionInternalAsync(IRpcClientRef rpcClientRef, UpdateSessionRequest request, Task waitTask)
+		{
+			DateTime deadline = DateTime.UtcNow + TimeSpan.FromMinutes(2.0);
+			using (AsyncDuplexStreamingCall<UpdateSessionRequest, UpdateSessionResponse> call = rpcClientRef.Client.UpdateSession(deadline: deadline))
+			{
+				_logger.LogDebug("Updating session {SessionId} (Status={Status})", request.SessionId, request.Status);
+
+				// Write the request to the server
+				await call.RequestStream.WriteAsync(request);
+
+				// Wait until the server responds or we need to trigger a new update
+				Task<bool> moveNextAsync = call.ResponseStream.MoveNext();
+
+				Task task = await Task.WhenAny(moveNextAsync, waitTask, rpcClientRef.DisposingTask);
+				if (task == waitTask)
+				{
+					_logger.LogDebug("Cancelling long poll from client side (new update)");
+				}
+				else if (task == rpcClientRef.DisposingTask)
+				{
+					_logger.LogDebug("Cancelling long poll from client side (server migration)");
+				}
+
+				// Close the request stream to indicate that we're finished
+				await call.RequestStream.CompleteAsync();
+
+				// Wait for a response or a new update to come in, then close the request stream
+				UpdateSessionResponse? response = null;
+				while (await moveNextAsync)
+				{
+					response = call.ResponseStream.Current;
+					moveNextAsync = call.ResponseStream.MoveNext();
+				}
+				return response;
+			}
+		}
+
+		/// <summary>
+		/// Handle a lease request
+		/// </summary>
+		/// <param name="session">The current session</param>
+		/// <param name="leaseInfo">Information about the lease</param>
+		/// <returns>Async task</returns>
+		async Task HandleLeaseAsync(ISession session, LeaseInfo leaseInfo)
+		{
+			using IScope scope = GlobalTracer.Instance.BuildSpan("HandleLease").WithResourceName(leaseInfo.Lease.Id).StartActive();
+			scope.Span.SetTag("LeaseId", leaseInfo.Lease.Id);
+			scope.Span.SetTag("AgentId", session.AgentId);
+			//			using IDisposable TraceProperty = LogContext.PushProperty("dd.trace_id", CorrelationIdentifier.TraceId.ToString());
+			//			using IDisposable SpanProperty = LogContext.PushProperty("dd.span_id", CorrelationIdentifier.SpanId.ToString());
+
+			_logger.LogInformation("Handling lease {LeaseId}", leaseInfo.Lease.Id);
+
+			// Get the lease outcome
+			LeaseResult result = LeaseResult.Failed;
+			try
+			{
+				result = await HandleLeasePayloadAsync(session, leaseInfo);
+			}
+			catch (Exception ex)
+			{
+				_logger.LogError(ex, "Unhandled exception while executing lease {LeaseId}", leaseInfo.Lease.Id);
+			}
+
+			// Update the state of the lease
+			lock (_lockObject)
+			{
+				if (leaseInfo.CancellationTokenSource.IsCancellationRequested)
+				{
+					leaseInfo.Lease.State = LeaseState.Cancelled;
+					leaseInfo.Lease.Outcome = LeaseOutcome.Failed;
+					leaseInfo.Lease.Output = ByteString.Empty;
+				}
+				else
+				{
+					leaseInfo.Lease.State = (result.Outcome == LeaseOutcome.Cancelled) ? LeaseState.Cancelled : LeaseState.Completed;
+					leaseInfo.Lease.Outcome = result.Outcome;
+					leaseInfo.Lease.Output = (result.Output != null) ? ByteString.CopyFrom(result.Output) : ByteString.Empty;
+				}
+				_logger.LogInformation("Transitioning lease {LeaseId} to {State}, outcome={Outcome}", leaseInfo.Lease.Id, leaseInfo.Lease.State, leaseInfo.Lease.Outcome);
+
+				if (result.SessionResult.HasValue && result.SessionResult.Value != SessionResult.Continue && _result == SessionResult.Continue)
+				{
+					_logger.LogInformation("Lease {LeaseId} is setting session result to {Result}", leaseInfo.Lease.Id, result.SessionResult.Value);
+					_result = result.SessionResult.Value;
+				}
+			}
+			_updateLeasesEvent.Set();
+		}
+
+		/// <summary>
+		/// Dispatch a lease payload to the appropriate handler
+		/// </summary>
+		/// <param name="session">The current session</param>
+		/// <param name="leaseInfo">Information about the lease</param>
+		/// <returns>Outcome from the lease</returns>
+		internal async Task<LeaseResult> HandleLeasePayloadAsync(ISession session, LeaseInfo leaseInfo)
+		{
+			Any payload = leaseInfo.Lease.Payload;
+			if (_typeUrlToLeaseHandler.TryGetValue(payload.TypeUrl, out LeaseHandler? leaseHandler))
+			{
+				GlobalTracer.Instance.ActiveSpan?.SetTag("task", payload.TypeUrl);
+				return await leaseHandler.ExecuteAsync(session, leaseInfo.Lease.Id, payload, leaseInfo.CancellationTokenSource.Token);
+			}
+			else
+			{
+				_logger.LogError("Invalid lease payload type ({PayloadType})", payload.TypeUrl);
+				return LeaseResult.Failed;
+			}
+		}
+	}
+}
diff --git a/Engine/Source/Programs/Horde/Horde.Agent/Leases/LeaseResult.cs b/Engine/Source/Programs/Horde/Horde.Agent/Leases/LeaseResult.cs
new file mode 100644
--- /dev/null
+++ b/Engine/Source/Programs/Horde/Horde.Agent/Leases/LeaseResult.cs
@@ -0,0 +1,71 @@
+// Copyright Epic Games, Inc. All Rights Reserved.
+
+using HordeCommon;
+
+namespace Horde.Agent.Leases
+{
+	/// <summary>
+	/// Result from executing a lease
+	/// </summary>
+	internal class LeaseResult
+	{
+		/// <summary>
+		/// Outcome of the lease (whether it completed/failed due to an internal error, etc...)
+		/// </summary>
+		public LeaseOutcome Outcome { get; }
+
+		/// <summary>
+		/// Output from executing the task
+		/// </summary>
+		public byte[]? Output { get; }
+
+		/// <summary>
+		/// If a lease wants to terminate the current session, this specifies the subsequent action to take.
+		/// </summary>
+		public SessionResult? SessionResult { get; }
+
+		/// <summary>
+		/// Static instance of a cancelled result
+		/// </summary>
+		public static LeaseResult Cancelled { get; } = new LeaseResult(LeaseOutcome.Cancelled);
+
+		/// <summary>
+		/// Static instance of a failed result
+		/// </summary>
+		public static LeaseResult Failed { get; } = new LeaseResult(LeaseOutcome.Failed);
+
+		/// <summary>
+		/// Static instance of a succesful result without a payload
+		/// </summary>
+		public static LeaseResult Success { get; } = new LeaseResult(LeaseOutcome.Success);
+
+		/// <summary>
+		/// Constructor for 
+		/// </summary>
+		/// <param name="outcome"></param>
+		private LeaseResult(LeaseOutcome outcome)
+		{
+			Outcome = outcome;
+		}
+
+		/// <summary>
+		/// Constructor for successful results
+		/// </summary>
+		/// <param name="output"></param>
+		public LeaseResult(byte[]? output)
+		{
+			Outcome = LeaseOutcome.Success;
+			Output = output;
+		}
+
+		/// <summary>
+		/// Constructor for successful results that cause a session state change
+		/// </summary>
+		/// <param name="sessionResult">Session result</param>
+		public LeaseResult(SessionResult sessionResult)
+		{
+			Outcome = LeaseOutcome.Success;
+			SessionResult = sessionResult;
+		}
+	}
+}
diff --git a/Engine/Source/Programs/Horde/Horde.Agent/Program.cs b/Engine/Source/Programs/Horde/Horde.Agent/Program.cs
--- a/Engine/Source/Programs/Horde/Horde.Agent/Program.cs
+++ b/Engine/Source/Programs/Horde/Horde.Agent/Program.cs
@@ -4,15 +4,57 @@ using System;
 using System.Collections.Generic;
 using System.Diagnostics;
 using System.IO;
+using System.Net.Http;
 using System.Reflection;
 using System.Runtime.InteropServices;
 using System.Threading.Tasks;
+using Datadog.Trace;
+using Datadog.Trace.Configuration;
+using Datadog.Trace.OpenTracing;
 using EpicGames.Core;
+using EpicGames.Horde.Storage;
+using Horde.Agent.Execution;
+using Horde.Agent.Execution.Interfaces;
+using Horde.Agent.Leases;
+using Horde.Agent.Leases.Handlers;
+using Horde.Agent.Services;
+using Horde.Agent.Utility;
+using Microsoft.Extensions.Configuration;
 using Microsoft.Extensions.DependencyInjection;
 using Microsoft.Extensions.Logging;
+using Microsoft.Extensions.Options;
+using OpenTracing.Util;
+using Polly;
 
 namespace Horde.Agent
 {
+	using ITracer = OpenTracing.ITracer;
+
+	/// <summary>
+	/// Injectable list of existing services
+	/// </summary>
+	class DefaultServices
+	{
+		/// <summary>
+		/// The base configuration object
+		/// </summary>
+		public IConfiguration Configuration { get; }
+
+		/// <summary>
+		/// List of service descriptors
+		/// </summary>
+		public IEnumerable<ServiceDescriptor> Descriptors { get; }
+
+		/// <summary>
+		/// Constructor
+		/// </summary>
+		public DefaultServices(IConfiguration configuration, IEnumerable<ServiceDescriptor> descriptors)
+		{
+			Configuration = configuration;
+			Descriptors = descriptors;
+		}
+	}
+
 	/// <summary>
 	/// Entry point
 	/// </summary>
@@ -52,9 +94,20 @@ namespace Horde.Agent
 		{
 			Program.Args = args;
 
+			using Logging.HordeLoggerProvider loggerProvider = new Logging.HordeLoggerProvider();
+
+			IConfiguration config = new ConfigurationBuilder()
+				.SetBasePath(AppDir.FullName)
+				.AddJsonFile("appsettings.json", optional: false)
+				.AddJsonFile("appsettings.Build.json", optional: true) // specific settings for builds (installer/dockerfile)
+				.AddJsonFile($"appsettings.{Environment.GetEnvironmentVariable("DOTNET_ENVIRONMENT")}.json", optional: true) // environment variable overrides, also used in k8s setups with Helm
+				.AddJsonFile("appsettings.User.json", optional: true)
+				.AddEnvironmentVariables()
+				.Build();
+
 			IServiceCollection services = new ServiceCollection();
 			services.AddCommandsFromAssembly(Assembly.GetExecutingAssembly());
-			services.AddLogging(builder => builder.AddProvider(new Logging.HordeLoggerProvider()));
+			services.AddLogging(builder => builder.AddProvider(loggerProvider));
 
 			// Enable unencrypted HTTP/2 for gRPC channel without TLS
 			AppContext.SetSwitch("System.Net.Http.SocketsHttpHandler.Http2UnencryptedSupport", true);
@@ -70,11 +123,87 @@ namespace Horde.Agent
 				}
 			}
 
+			// Add all the default 
+			IConfigurationSection configSection = config.GetSection(AgentSettings.SectionName);
+			services.AddOptions<AgentSettings>().Configure(options => configSection.Bind(options)).ValidateDataAnnotations();
+
+			AgentSettings settings = new AgentSettings();
+			configSection.Bind(settings);
+
+			ServerProfile serverProfile = settings.GetCurrentServerProfile();
+			ConfigureTracing(serverProfile.Environment, Program.Version);
+
+			Logging.SetEnv(serverProfile.Environment);
+
+			ILogger certificateLogger = loggerProvider.CreateLogger(typeof(CertificateHelper).FullName!);
+			services.AddHttpClient(Program.HordeServerClientName, config =>
+			{
+				config.BaseAddress = serverProfile.Url;
+				config.DefaultRequestHeaders.Add("Accept", "application/json");
+				config.Timeout = TimeSpan.FromSeconds(300); // Need to make sure this doesn't cancel any long running gRPC streaming calls (eg. session update)
+			})
+			.ConfigurePrimaryHttpMessageHandler(() =>
+			{
+				HttpClientHandler handler = new HttpClientHandler();
+				handler.ServerCertificateCustomValidationCallback += (sender, cert, chain, errors) => CertificateHelper.CertificateValidationCallBack(certificateLogger, sender, cert, chain, errors, serverProfile);
+				return handler;
+			})
+			.AddTransientHttpErrorPolicy(builder =>
+			{
+				return builder.WaitAndRetryAsync(new[] { TimeSpan.FromSeconds(1), TimeSpan.FromSeconds(5), TimeSpan.FromSeconds(10) });
+			});
+
+			services.AddHordeStorage(settings => configSection.GetCurrentServerProfile().GetSection(nameof(serverProfile.Storage)).Bind(settings));
+
+			services.AddSingleton<GrpcService>();
+
+			services.AddTransient<PerforceExecutorFactory>();
+			services.AddTransient<LocalExecutorFactory>();
+			services.AddTransient<TestExecutorFactory>();
+
+			services.AddTransient<IExecutorFactory>(sp =>
+				sp.GetRequiredService<IOptions<AgentSettings>>().Value.Executor switch
+				{
+					ExecutorType.Test => sp.GetRequiredService<TestExecutorFactory>(),
+					ExecutorType.Local => sp.GetRequiredService<LocalExecutorFactory>(),
+					ExecutorType.Perforce => sp.GetRequiredService<PerforceExecutorFactory>(),
+					_ => throw new NotImplementedException()
+				});
+
+			services.AddSingleton<LeaseHandler, ComputeHandler>();
+			services.AddSingleton<LeaseHandler, ConformHandler>();
+			services.AddSingleton<LeaseHandler, JobHandler>();
+			services.AddSingleton<LeaseHandler, RestartHandler>();
+			services.AddSingleton<LeaseHandler, ShutdownHandler>();
+			services.AddSingleton<LeaseHandler, UpgradeHandler>();
+
+			services.AddSingleton<CapabilitiesService>();
+			services.AddSingleton<SessionFactoryService>();
+			services.AddHostedService<WorkerService>();
+
+			// Allow commands to augment the service collection for their own DI service providers
+			services.AddSingleton<DefaultServices>(x => new DefaultServices(config, services));
+
 			// Execute all the commands
 			IServiceProvider serviceProvider = services.BuildServiceProvider();
 			return await CommandHost.RunAsync(new CommandLineArguments(args), serviceProvider, typeof(Horde.Agent.Modes.Service.RunCommand));
 		}
 
+		static void ConfigureTracing(string environment, string version)
+		{
+			TracerSettings settings = TracerSettings.FromDefaultSources();
+			settings.Environment = environment;
+			settings.ServiceName = "hordeagent";
+			settings.ServiceVersion = version;
+			settings.LogsInjectionEnabled = true;
+
+			Tracer.Configure(settings);
+
+			ITracer openTracer = OpenTracingTracerFactory.WrapTracer(Tracer.Instance);
+			GlobalTracer.Register(openTracer);
+		}
+
+
 		/// <summary>
 		/// Gets the version of the current assembly
 		/// </summary>
diff --git a/Engine/Source/Programs/Horde/Horde.Agent/Properties/AssemblyInfo.cs b/Engine/Source/Programs/Horde/Horde.Agent/Properties/AssemblyInfo.cs
--- a/Engine/Source/Programs/Horde/Horde.Agent/Properties/AssemblyInfo.cs
+++ b/Engine/Source/Programs/Horde/Horde.Agent/Properties/AssemblyInfo.cs
@@ -4,4 +4,5 @@ using System.Resources;
 using System.Runtime.CompilerServices;
 
 [assembly: InternalsVisibleTo("Horde.Agent.Tests")]
+[assembly: InternalsVisibleTo("DynamicProxyGenAssembly2")]
 [assembly: NeutralResourcesLanguage("en-US")]
diff --git a/Engine/Source/Programs/Horde/Horde.Agent/Properties/launchSettings.json b/Engine/Source/Programs/Horde/Horde.Agent/Properties/launchSettings.json
--- a/Engine/Source/Programs/Horde/Horde.Agent/Properties/launchSettings.json
+++ b/Engine/Source/Programs/Horde/Horde.Agent/Properties/launchSettings.json
@@ -4,7 +4,7 @@
       "commandName": "Project",
       "commandLineArgs": "service run",
       "environmentVariables": {
-        "DOTNET_ENVIRONMENT": "Development"
+        "DOTNET_ENVIRONMENT": "Local"
       }
     }
   }
diff --git a/Engine/Source/Programs/Horde/Horde.Agent/Services/WorkerService.cs b/Engine/Source/Programs/Horde/Horde.Agent/Services/WorkerService.cs
--- a/Engine/Source/Programs/Horde/Horde.Agent/Services/WorkerService.cs
+++ b/Engine/Source/Programs/Horde/Horde.Agent/Services/WorkerService.cs
@@ -3,43 +3,22 @@
 using System;
 using System.Collections.Generic;
 using System.Diagnostics;
-using System.Diagnostics.CodeAnalysis;
-using System.IO;
-using System.IO.Compression;
 using System.Linq;
-using System.Management;
-using System.Net;
 using System.Reflection;
 using System.Runtime.InteropServices;
-using System.Runtime.Versioning;
-using System.Text;
-using System.Text.RegularExpressions;
 using System.Threading;
 using System.Threading.Tasks;
-using System.Xml;
-using Amazon.EC2;
-using Amazon.EC2.Model;
-using Amazon.Util;
 using EpicGames.Core;
-using EpicGames.Horde.Compute;
-using EpicGames.Horde.Storage;
-using Google.Protobuf;
-using Google.Protobuf.WellKnownTypes;
-using Grpc.Core;
 using Grpc.Net.Client;
-using Horde.Agent.Execution;
-using Horde.Agent.Execution.Interfaces;
-using Horde.Agent.Parser;
+using Horde.Agent.Leases;
 using Horde.Agent.Utility;
 using HordeCommon;
 using HordeCommon.Rpc;
 using HordeCommon.Rpc.Messages;
-using HordeCommon.Rpc.Tasks;
+using Microsoft.Extensions.DependencyInjection;
 using Microsoft.Extensions.Hosting;
 using Microsoft.Extensions.Logging;
 using Microsoft.Extensions.Options;
-using OpenTracing;
-using OpenTracing.Util;
 
 namespace Horde.Agent.Services
 {
@@ -48,256 +27,24 @@ namespace Horde.Agent.Services
 	/// </summary>
 	class WorkerService : BackgroundService, IDisposable
 	{
-		/// <summary>
-		/// Stores information about an active session
-		/// </summary>
-		internal class LeaseInfo
-		{
-			/// <summary>
-			/// The worker lease state
-			/// </summary>
-			public Lease Lease { get; set; }
-
-			/// <summary>
-			/// The task being executed for this lease
-			/// </summary>
-			public Task? Task { get; set; }
-
-			/// <summary>
-			/// Source for cancellation tokens for this session.
-			/// </summary>
-			public CancellationTokenSource CancellationTokenSource { get; }
-
-			/// <summary>
-			/// Constructor
-			/// </summary>
-			/// <param name="lease">The worker lease state</param>
-			public LeaseInfo(Lease lease)
-			{
-				Lease = lease;
-				CancellationTokenSource = new CancellationTokenSource();
-			}
-		}
-
-		/// <summary>
-		/// Result from executing a lease
-		/// </summary>
-		internal class LeaseResult
-		{
-			/// <summary>
-			/// Outcome of the lease (whether it completed/failed due to an internal error, etc...)
-			/// </summary>
-			public LeaseOutcome Outcome { get; }
-
-			/// <summary>
-			/// Output from executing the task
-			/// </summary>
-			public byte[]? Output { get; }
-
-			/// <summary>
-			/// Static instance of a cancelled result
-			/// </summary>
-			public static LeaseResult Cancelled { get; } = new LeaseResult(LeaseOutcome.Cancelled);
-
-			/// <summary>
-			/// Static instance of a failed result
-			/// </summary>
-			public static LeaseResult Failed { get; } = new LeaseResult(LeaseOutcome.Failed);
-
-			/// <summary>
-			/// Static instance of a succesful result without a payload
-			/// </summary>
-			public static LeaseResult Success { get; } = new LeaseResult(LeaseOutcome.Success);
-
-			/// <summary>
-			/// Constructor for 
-			/// </summary>
-			/// <param name="outcome"></param>
-			private LeaseResult(LeaseOutcome outcome)
-			{
-				Outcome = outcome;
-			}
-
-			/// <summary>
-			/// Constructor for successful results
-			/// </summary>
-			/// <param name="output"></param>
-			public LeaseResult(byte[]? output)
-			{
-				Outcome = LeaseOutcome.Success;
-				Output = output;
-			}
-		}
-
-		/// <summary>
-		/// List of processes that should be terminated before running a job
-		/// </summary>
-		private readonly HashSet<string> _processNamesToTerminate = new HashSet<string>(StringComparer.OrdinalIgnoreCase);
-
-		/// <summary>
-		/// Object used for controlling access to the access tokens and active sessions list
-		/// </summary>
-		readonly object _lockObject = new object();
-
-		/// <summary>
-		/// Sink for log messages
-		/// </summary>
 		readonly ILogger _logger;
-
-		/// <summary>
-		/// The current settings instance
-		/// </summary>
-		readonly AgentSettings _settings;
-
-		/// <summary>
-		/// Settings for the current server
-		/// </summary>
-		readonly ServerProfile _serverProfile;
-
-		/// <summary>
-		/// The grpc service instance
-		/// </summary>
-		readonly GrpcService _grpcService;
-
-		/// <summary>
-		/// Client interface to the storage system
-		/// </summary>
-		readonly ILegacyStorageClient _storageClient;
-
-		/// <summary>
-		/// The working directory
-		/// </summary>
-		readonly DirectoryReference _workingDir;
-
-		/// <summary>
-		/// The list of active leases.
-		/// </summary>
-		readonly List<LeaseInfo> _activeLeases = new List<LeaseInfo>();
-
-		/// <summary>
-		/// Number of leases completed
-		/// </summary>
-		public int NumLeasesCompleted { get; private set; }
-
-		/// <summary>
-		/// Whether the agent is currently in an unhealthy state
-		/// </summary>
-		readonly bool _unhealthy = false;
-
-		/// <summary>
-		/// Whether a shutdown of the worker service has been requested
-		/// </summary>
-		bool _requestShutdown = false;
-
-		/// <summary>
-		/// Whether to restart after shutting down
-		/// </summary>
-		bool _restartAfterShutdown = false;
-
-		/// <summary>
-		/// Time the the service started
-		/// </summary>
-		static readonly DateTimeOffset s_startTime = DateTimeOffset.Now;
-
-		/// <summary>
-		/// Time at which the computer started
-		/// </summary>
-		static readonly DateTimeOffset s_bootTime = DateTimeOffset.Now - TimeSpan.FromTicks(Environment.TickCount64 * TimeSpan.TicksPerMillisecond);
-
-		/// <summary>
-		/// Task completion source used to trigger the background thread to update the leases. Must take a lock on LockObject before 
-		/// </summary>
-		readonly AsyncEvent _updateLeasesEvent = new AsyncEvent();
-
-		/// <summary>
-		/// Number of times UpdateSession has failed
-		/// </summary>
-		int _updateSessionFailures;
-		
-		/// <summary>
-		/// Delegate for lease active events
-		/// </summary>
-		public delegate void LeaseActiveEvent(Lease lease);
-
-		/// <summary>
-		/// Event triggered when a lease is accepted and set to active
-		/// </summary>
-		public event LeaseActiveEvent? OnLeaseActive;
-
-		/// <summary>
-		/// Function for creating a new executor. Primarily to aid testing. 
-		/// </summary>
-		private readonly Func<IRpcConnection, ExecuteJobTask, BeginBatchResponse, IExecutor> _createExecutor;
-
-		/// <summary>
-		/// Function for creating a new HordeRpcClient. Primarily to aid testing. 
-		/// </summary>
-		private readonly Func<GrpcChannel, HordeRpc.HordeRpcClient> _createHordeRpcClient;
-
-		/// <summary>
-		/// How often to poll the server checking if a step has been aborted
-		/// Exposed as internal to ease testing.
-		/// </summary>
-		internal TimeSpan _stepAbortPollInterval = TimeSpan.FromSeconds(5);
-		
-		
-		/// <summary>
-		/// How long to wait before trying to reacquire a new connection
-		/// Exposed as internal to ease testing. Using a lower delay can speed up tests. 
-		/// </summary>
-		internal TimeSpan _rpcConnectionRetryDelay = TimeSpan.FromSeconds(5);
+		readonly IOptions<AgentSettings> _settings;
+		readonly SessionFactoryService _sessionFactoryService;
+		readonly CapabilitiesService _capabilitiesService;
+		readonly LeaseHandler[] _leaseHandlers;
+		readonly IServiceProvider _serviceProvider;
 
 		/// <summary>
 		/// Constructor. Registers with the server and starts accepting connections.
 		/// </summary>
-		/// <param name="logger">Log sink</param>
-		/// <param name="options">The current settings</param>
-		/// <param name="grpcService">Instance of the Grpc service</param>
-		/// <param name="storageClient">Instance of the storage client</param>
-		/// <param name="createExecutor">Optional factory method for creating an executor</param>
-		/// <param name="createHordeRpcClient">Optional factory method for creating a HordeRpcClient</param>
-		public WorkerService(ILogger<WorkerService> logger, IOptions<AgentSettings> options, GrpcService grpcService, ILegacyStorageClient storageClient,
-			Func<IRpcConnection, ExecuteJobTask, BeginBatchResponse, IExecutor>? createExecutor = null,
-			Func<GrpcChannel, HordeRpc.HordeRpcClient>? createHordeRpcClient = null)
+		public WorkerService(IOptions<AgentSettings> settings, SessionFactoryService sessionFactoryService, CapabilitiesService capabilitiesService, IEnumerable<LeaseHandler> leaseHandlers, IServiceProvider serviceProvider, ILogger<WorkerService> logger)
 		{
+			_settings = settings;
+			_sessionFactoryService = sessionFactoryService;
+			_capabilitiesService = capabilitiesService;
 			_logger = logger;
-			_settings = options.Value;
-			_serverProfile = _settings.GetCurrentServerProfile();
-			_grpcService = grpcService;
-			_storageClient = storageClient;
-			_createHordeRpcClient = createHordeRpcClient ?? (channel => new HordeRpc.HordeRpcClient(channel));
-
-			if (_settings.WorkingDir == null)
-			{
-				throw new Exception("WorkingDir is not set. Unable to run service.");
-			}
-
-			DirectoryReference baseDir = new FileReference(Assembly.GetExecutingAssembly().Location).Directory;
-
-			_workingDir = DirectoryReference.Combine(baseDir, _settings.WorkingDir);
-			logger.LogInformation("Using working directory {WorkingDir}", _workingDir);
-			DirectoryReference.CreateDirectory(_workingDir);
-
-			_processNamesToTerminate.UnionWith(_settings.ProcessNamesToTerminate);
-
-			if (createExecutor == null)
-			{
-				_createExecutor = _settings.Executor switch
-				{
-					ExecutorType.Test => (rpcClient, executeTask, batch) =>
-						new TestExecutor(rpcClient, executeTask.JobId, executeTask.BatchId, batch.AgentType),
-					ExecutorType.Local => (rpcClient, executeTask, batch) =>
-						new LocalExecutor(rpcClient, executeTask.JobId, executeTask.BatchId, batch.AgentType, _settings.LocalExecutor),
-					ExecutorType.Perforce => (rpcClient, executeTask, batch) =>
-						new PerforceExecutor(rpcClient, executeTask.JobId, executeTask.BatchId, batch.AgentType, executeTask.AutoSdkWorkspace, executeTask.Workspace, _workingDir),
-
-					_ => throw new InvalidDataException($"Unknown executor type '{_settings.Executor}'")
-				};
-			}
-			else
-			{
-				_createExecutor = createExecutor;
-			}
+			_leaseHandlers = leaseHandlers.ToArray();
+			_serviceProvider = serviceProvider;
 		}
 
 		/// <summary>
@@ -322,46 +69,52 @@ namespace Horde.Agent.Services
 		/// <param name="stoppingToken">Indicates that the service is trying to stop</param>
 		internal async Task ExecuteInnerAsync(CancellationToken stoppingToken)
 		{
-			// Print the server info
-			_logger.LogInformation("Server: {Server}", _serverProfile.Url);
-			_logger.LogInformation("Arguments: {Arguments}", Environment.CommandLine);
-
 			// Show the current client id
 			string version = Program.Version;
 			_logger.LogInformation("Version: {Version}", version);
 
+			// Print the server info
+			_logger.LogInformation("Arguments: {Arguments}", Environment.CommandLine);
+
 			// Keep trying to start an agent session with the server
 			while (!stoppingToken.IsCancellationRequested)
 			{
+				SessionResult result = SessionResult.Continue;
+
 				Stopwatch sessionTime = Stopwatch.StartNew();
-				try
+				await using (AsyncServiceScope scope = _serviceProvider.CreateAsyncScope())
 				{
-					if (_activeLeases.Count == 0)
+					try
 					{
-						await HandleSessionAsync(false, true, stoppingToken);
+						using Mutex singleInstanceMutex = new(false, "Global\\HordeAgent-DB828ACB-0AA5-4D32-A62A-21D4429B1014");
+						await WaitForMutexAsync(singleInstanceMutex, stoppingToken);
+
+						await using (ISession session = await _sessionFactoryService.CreateAsync(stoppingToken))
+						{
+							LeaseManager leaseManager = new LeaseManager(session, _capabilitiesService, _leaseHandlers, _settings, _logger);
+							result = await leaseManager.RunAsync(false, stoppingToken);
+						}
 					}
-				}
-				catch (Exception ex)
-				{
-					if (!stoppingToken.IsCancellationRequested || !IsCancellationException(ex))
+					catch (OperationCanceledException) when (stoppingToken.IsCancellationRequested)
+					{
+						throw;
+					}
+					catch (Exception ex)
 					{
-						_logger.LogError(ex, "Exception while executing session. Restarting.");
+						_logger.LogError(ex, "Exception while executing session. Restarting. ({Message})", ex.Message);
 					}
 				}
 
-				try
+				if (result == SessionResult.Terminate)
 				{
-					await DrainLeasesAsync();
+					break;
 				}
-				catch (Exception ex)
+				else if (result == SessionResult.Shutdown || result == SessionResult.Restart)
 				{
-					_logger.LogError(ex, "Exception while draining leases. Agent may be in an inconsistent state.");
-				}
+					bool restart = (result == SessionResult.Restart);
+					_logger.LogInformation("Initiating shutdown (restart={Restart})", restart);
 
-				if (_requestShutdown)
-				{
-					_logger.LogInformation("Initiating shutdown (restart={Restart})", _restartAfterShutdown);
-					if (Shutdown.InitiateShutdown(_restartAfterShutdown, _logger))
+					if (Shutdown.InitiateShutdown(restart, _logger))
 					{
 						for (int idx = 10; idx > 0; idx--)
 						{
@@ -370,9 +123,9 @@ namespace Horde.Agent.Services
 						}
 						_logger.LogInformation("Shutdown aborted.");
 					}
-					_requestShutdown = _restartAfterShutdown = false;
 				}
-				else if (sessionTime.Elapsed < TimeSpan.FromSeconds(2.0))
+				
+				if (sessionTime.Elapsed < TimeSpan.FromSeconds(2.0))
 				{
 					_logger.LogInformation("Waiting 5 seconds before restarting session...");
 					await Task.Delay(TimeSpan.FromSeconds(5.0), stoppingToken);
@@ -398,1669 +151,5 @@ namespace Horde.Agent.Services
 			{
 			}
 		}
-
-		async Task DrainLeasesAsync()
-		{
-			for (int idx = 0; idx < _activeLeases.Count; idx++)
-			{
-				LeaseInfo activeLease = _activeLeases[idx];
-				if (activeLease.Task == null)
-				{
-					_activeLeases.RemoveAt(idx--);
-					_logger.LogInformation("Removed lease {LeaseId}", activeLease.Lease.Id);
-				}
-				else
-				{
-					_logger.LogInformation("Cancelling active lease {LeaseId}", activeLease.Lease.Id);
-					activeLease.CancellationTokenSource.Cancel();
-				}
-			}
-
-			while (_activeLeases.Count > 0)
-			{
-				List<Task> tasks = _activeLeases.Select(x => x.Task!).ToList();
-				tasks.Add(Task.Delay(TimeSpan.FromMinutes(1.0)));
-				await Task.WhenAny(tasks);
-
-				for (int idx = 0; idx < _activeLeases.Count; idx++)
-				{
-					LeaseInfo activeLease = _activeLeases[idx];
-					if (activeLease.Task!.IsCompleted)
-					{
-						_activeLeases.RemoveAt(idx--);
-						try
-						{
-							await activeLease.Task;
-						}
-						catch (OperationCanceledException)
-						{
-						}
-						catch (Exception ex)
-						{
-							_logger.LogError(ex, "Lease {LeaseId} threw an exception while terminating", activeLease.Lease.Id);
-						}
-						_logger.LogInformation("Lease {LeaseId} has completed", activeLease.Lease.Id);
-					}
-					else
-					{
-						_logger.LogInformation("Still waiting for lease {LeaseId} to terminate...", activeLease.Lease.Id);
-					}
-				}
-			}
-		}
-
-		/// <summary>
-		/// Handles the lifetime of an agent session.
-		///
-		/// <paramref name="shutdownAfterFinishedLease"/> is primarily for executing leases in a one-shot environment,
-		/// such as an AWS Lambda function.
-		/// </summary>
-		/// <param name="shutdownAfterFinishedLease">Stop after one lease has finished</param>
-		/// <param name="ensureSingleInstance">Ensure only one instance of agent is running</param>
-		/// <param name="stoppingToken">Indicates that the service is trying to stop</param>
-		/// <returns>Async task</returns>
-		internal async Task HandleSessionAsync(bool shutdownAfterFinishedLease, bool ensureSingleInstance, CancellationToken stoppingToken)
-		{
-			using Mutex singleInstanceMutex = new (false, "Global\\HordeAgent-DB828ACB-0AA5-4D32-A62A-21D4429B1014");
-			if (ensureSingleInstance)
-			{
-				// Make sure there's only one instance of the agent running
-				await WaitForMutexAsync(singleInstanceMutex, stoppingToken);				
-			}
-
-			// Terminate any remaining child processes from other instances
-			TerminateProcesses(_logger, stoppingToken);
-
-			// Show the worker capabilities
-			AgentCapabilities capabilities = await GetAgentCapabilities(_workingDir, _logger, _settings.Properties);
-			if (capabilities.Properties.Count > 0)
-			{
-				_logger.LogInformation("Global:");
-				foreach (string property in capabilities.Properties)
-				{
-					_logger.LogInformation("  {AgentProperty}", property);
-				}
-			}
-			foreach (DeviceCapabilities device in capabilities.Devices)
-			{
-				_logger.LogInformation("{DeviceName} Device:", device.Handle);
-				foreach (string property in device.Properties)
-				{
-					_logger.LogInformation("   {DeviceProperty}", property);
-				}
-			}
-
-			// Mount all the necessary network shares. Currently only supported on Windows.
-			if (_settings.ShareMountingEnabled && RuntimeInformation.IsOSPlatform(OSPlatform.Windows))
-			{
-				foreach (MountNetworkShare share in _settings.Shares)
-				{
-					if (share.MountPoint != null && share.RemotePath != null)
-					{
-						_logger.LogInformation("Mounting {RemotePath} as {MountPoint}", share.RemotePath, share.MountPoint);
-						NetworkShare.Mount(share.MountPoint, share.RemotePath);
-					}
-				}
-			}
-
-			// Create the session
-			CreateSessionResponse createSessionResponse;
-			using (GrpcChannel channel = _grpcService.CreateGrpcChannel(_serverProfile.Token))
-			{
-				HordeRpc.HordeRpcClient rpcClient = _createHordeRpcClient(channel);
-
-				// Create the session information
-				CreateSessionRequest sessionRequest = new CreateSessionRequest();
-				sessionRequest.Name = _settings.GetAgentName();
-				sessionRequest.Status = AgentStatus.Ok;
-				sessionRequest.Capabilities = capabilities;
-				sessionRequest.Version = Program.Version;
-
-				// Create a session
-				createSessionResponse = await rpcClient.CreateSessionAsync(sessionRequest, null, null, stoppingToken);
-				_logger.LogInformation("Session started. AgentName={AgentName} SessionId={SessionId}", _settings.GetAgentName(), createSessionResponse.SessionId);
-			}
-
-			Func<GrpcChannel> createGrpcChannel = () => _grpcService.CreateGrpcChannel(createSessionResponse.Token);
-			
-			// Open a connection to the server
-			await using (IRpcConnection rpcCon = RpcConnection.Create(createGrpcChannel, _createHordeRpcClient, _logger))
-			{
-				// Track how many updates we get in 10 seconds. We'll start rate limiting this if it looks like we've got a problem that's causing us to spam the server.
-				Stopwatch updateTimer = Stopwatch.StartNew();
-				Queue<TimeSpan> updateTimes = new Queue<TimeSpan>();
-
-				// Loop until we're ready to exit
-				Stopwatch updateCapabilitiesTimer = Stopwatch.StartNew();
-				for (; ; )
-				{
-					Task waitTask = _updateLeasesEvent.Task;
-
-					// Flag for whether the service is stopping
-					bool stopping = false;
-					if (stoppingToken.IsCancellationRequested)
-					{
-						_logger.LogInformation("Cancellation from token requested");
-						stopping = true;
-					}
-					
-					if (_requestShutdown)
-					{
-						_logger.LogInformation("Shutdown requested");
-						stopping = true;
-					}
-
-					// Build the next update request
-					UpdateSessionRequest updateSessionRequest = new UpdateSessionRequest();
-					updateSessionRequest.AgentId = createSessionResponse.AgentId;
-					updateSessionRequest.SessionId = createSessionResponse.SessionId;
-
-					// Get the new the lease states. If a restart is requested and we have no active leases, signal to the server that we're stopping.
-					lock (_lockObject)
-					{
-						foreach (LeaseInfo leaseInfo in _activeLeases)
-						{
-							updateSessionRequest.Leases.Add(new Lease(leaseInfo.Lease));
-						}
-						if (_requestShutdown && _activeLeases.Count == 0)
-						{
-							stopping = true;
-						}
-					}
-
-					// Get the new agent status
-					if (stopping)
-					{
-						updateSessionRequest.Status = AgentStatus.Stopping;
-					}
-					else if (_unhealthy)
-					{
-						updateSessionRequest.Status = AgentStatus.Unhealthy;
-					}
-					else
-					{
-						updateSessionRequest.Status = AgentStatus.Ok;
-					}
-
-					// Update the capabilities every 5m
-					if (updateCapabilitiesTimer.Elapsed > TimeSpan.FromMinutes(5.0))
-					{
-						try
-						{
-							updateSessionRequest.Capabilities = await GetAgentCapabilities(_workingDir, _logger, _settings.Properties);
-						}
-						catch (Exception ex)
-						{
-							_logger.LogWarning(ex, "Unable to query agent capabilities. Ignoring.");
-						}
-						updateCapabilitiesTimer.Restart();
-					}
-
-					// Complete the wait task if we subsequently stop
-					using (stopping ? (CancellationTokenRegistration?)null : stoppingToken.Register(() => _updateLeasesEvent.Set()))
-					{
-						// Update the state with the server
-						UpdateSessionResponse? updateSessionResponse = null;
-						using (IRpcClientRef? rpcClientRef = rpcCon.TryGetClientRef(new RpcContext()))
-						{
-							if (rpcClientRef == null)
-							{
-								// An RpcConnection has not yet been established, wait a period of time and try again
-								await Task.WhenAny(Task.Delay(_rpcConnectionRetryDelay, stoppingToken), waitTask);
-							}
-							else
-							{
-								updateSessionResponse = await UpdateSessionAsync(rpcClientRef, updateSessionRequest, waitTask);
-							}
-						}
-
-						lock (_lockObject)
-						{
-							// Now reconcile the local state to match what the server reports
-							if (updateSessionResponse != null)
-							{
-								bool atLeastOneLeaseFinished = _activeLeases.Any(x => x.Lease.State is LeaseState.Completed or LeaseState.Cancelled);
-								if (atLeastOneLeaseFinished && shutdownAfterFinishedLease)
-								{
-									_logger.LogInformation("At least one lease executed. Requesting shutdown.");
-									_requestShutdown = true;
-								}
-
-								// Remove any leases which have completed
-								int numRemoved = _activeLeases.RemoveAll(x => (x.Lease.State == LeaseState.Completed || x.Lease.State == LeaseState.Cancelled) && !updateSessionResponse.Leases.Any(y => y.Id == x.Lease.Id && y.State != LeaseState.Cancelled));
-								NumLeasesCompleted += numRemoved;
- 
-								// Create any new leases and cancel any running leases
-								foreach (Lease serverLease in updateSessionResponse.Leases)
-								{
-									if (serverLease.State == LeaseState.Cancelled)
-									{
-										LeaseInfo? info = _activeLeases.FirstOrDefault(x => x.Lease.Id == serverLease.Id);
-										if (info != null)
-										{
-											_logger.LogInformation("Cancelling lease {LeaseId}", serverLease.Id);
-											info.CancellationTokenSource.Cancel();
-										}
-									}
-									if (serverLease.State == LeaseState.Pending && !_activeLeases.Any(x => x.Lease.Id == serverLease.Id))
-									{
-										serverLease.State = LeaseState.Active;
-
-										_logger.LogInformation("Adding lease {LeaseId}", serverLease.Id);
-										LeaseInfo info = new LeaseInfo(serverLease);
-										info.Task = Task.Run(() => HandleLeaseAsync(rpcCon, createSessionResponse.AgentId, info), CancellationToken.None);
-										_activeLeases.Add(info);
-										OnLeaseActive?.Invoke(serverLease);
-									}
-								}
-							}
-
-							// If there's nothing still running and cancellation was requested, exit
-							if (_activeLeases.Count == 0 && updateSessionRequest.Status == AgentStatus.Stopping)
-							{
-								_logger.LogInformation("No leases are active. Agent is stopping.");
-								break;
-							}
-						}
-					}
-
-					// Update the historical update times
-					TimeSpan updateTime = updateTimer.Elapsed;
-					while (updateTimes.TryPeek(out TimeSpan firstTime) && firstTime < updateTime - TimeSpan.FromMinutes(1.0))
-					{
-						updateTimes.Dequeue();
-					}
-					updateTimes.Enqueue(updateTime);
-
-					// If we're updating too much, introduce an artificial delay
-					if (updateTimes.Count > 60)
-					{
-						_logger.LogWarning("Agent is issuing large number of UpdateSession() calls. Delaying for 10 seconds.");
-						await Task.Delay(TimeSpan.FromSeconds(10.0), stoppingToken);
-					}
-				}
-
-				_logger.LogInformation("Disposing RpcConnection");
-			}
-		}
-
-		/// <summary>
-		/// Wrapper for <see cref="UpdateSessionInternalAsync"/> which filters/logs exceptions
-		/// </summary>
-		/// <param name="rpcClientRef">The RPC client connection</param>
-		/// <param name="updateSessionRequest">The session update request</param>
-		/// <param name="waitTask">Task which can be used to jump out of the update early</param>
-		/// <returns>Response from the call</returns>
-		async Task<UpdateSessionResponse?> UpdateSessionAsync(IRpcClientRef rpcClientRef, UpdateSessionRequest updateSessionRequest, Task waitTask)
-		{
-			UpdateSessionResponse? updateSessionResponse = null;
-			try
-			{
-				updateSessionResponse = await UpdateSessionInternalAsync(rpcClientRef, updateSessionRequest, waitTask);
-				_updateSessionFailures = 0;
-			}
-			catch (RpcException ex)
-			{
-				if (++_updateSessionFailures >= 3)
-				{
-					throw;
-				}
-				else if (ex.StatusCode == StatusCode.Unavailable)
-				{
-					_logger.LogInformation(ex, "Service unavailable while calling UpdateSessionAsync(), will retry");
-				}
-				else
-				{
-					_logger.LogError(ex, "Error while executing RPC. Will retry.");
-				}
-			}
-			return updateSessionResponse;
-		}
-
-		/// <summary>
-		/// Tries to update the session state on the server.
-		/// 
-		/// This operation is a little gnarly due to the fact that we want to long-poll for the result.
-		/// Since we're doing the update via a gRPC call, the way to do that without using cancellation tokens is to keep the request stream open
-		/// until we want to terminate the call (see [REDACTED_PROFILE_URL]/grpc/issues/8277). In order to do that, we need to make a 
-		/// bidirectional streaming call, even though we only expect one response/response.
-		/// </summary>
-		/// <param name="rpcClientRef">The RPC client</param>
-		/// <param name="request">The session update request</param>
-		/// <param name="waitTask">Task to use to terminate the wait</param>
-		/// <returns>The response object</returns>
-		async Task<UpdateSessionResponse?> UpdateSessionInternalAsync(IRpcClientRef rpcClientRef, UpdateSessionRequest request, Task waitTask)
-		{
-			DateTime deadline = DateTime.UtcNow + TimeSpan.FromMinutes(2.0);
-			using (AsyncDuplexStreamingCall<UpdateSessionRequest, UpdateSessionResponse> call = rpcClientRef.Client.UpdateSession(deadline: deadline))
-			{
-				_logger.LogDebug("Updating session {SessionId} (Status={Status})", request.SessionId, request.Status);
-
-				// Write the request to the server
-				await call.RequestStream.WriteAsync(request);
-
-				// Wait until the server responds or we need to trigger a new update
-				Task<bool> moveNextAsync = call.ResponseStream.MoveNext();
-
-				Task task = await Task.WhenAny(moveNextAsync, waitTask, rpcClientRef.DisposingTask);
-				if(task == waitTask)
-				{
-					_logger.LogDebug("Cancelling long poll from client side (new update)");
-				}
-				else if (task == rpcClientRef.DisposingTask)
-				{
-					_logger.LogDebug("Cancelling long poll from client side (server migration)");
-				}
-
-				// Close the request stream to indicate that we're finished
-				await call.RequestStream.CompleteAsync();
-
-				// Wait for a response or a new update to come in, then close the request stream
-				UpdateSessionResponse? response = null;
-				while (await moveNextAsync)
-				{
-					response = call.ResponseStream.Current;
-					moveNextAsync = call.ResponseStream.MoveNext();
-				}
-				return response;
-			}
-		}
-
-		/// <summary>
-		/// Handle a lease request
-		/// </summary>
-		/// <param name="rpcConnection">The RPC connection to the server</param>
-		/// <param name="agentId">The agent id</param>
-		/// <param name="leaseInfo">Information about the lease</param>
-		/// <returns>Async task</returns>
-		async Task HandleLeaseAsync(IRpcConnection rpcConnection, string agentId, LeaseInfo leaseInfo)
-		{
-			using IScope scope = GlobalTracer.Instance.BuildSpan("HandleLease").WithResourceName(leaseInfo.Lease.Id).StartActive();
-			scope.Span.SetTag("LeaseId", leaseInfo.Lease.Id);
-			scope.Span.SetTag("AgentId", agentId);
-//			using IDisposable TraceProperty = LogContext.PushProperty("dd.trace_id", CorrelationIdentifier.TraceId.ToString());
-//			using IDisposable SpanProperty = LogContext.PushProperty("dd.span_id", CorrelationIdentifier.SpanId.ToString());
-
-			_logger.LogInformation("Handling lease {LeaseId}", leaseInfo.Lease.Id);
-
-			// Get the lease outcome
-			LeaseResult result = LeaseResult.Failed;
-			try
-			{
-				result = await HandleLeasePayloadAsync(rpcConnection, agentId, leaseInfo);
-			}
-			catch (Exception ex)
-			{
-				_logger.LogError(ex, "Unhandled exception while executing lease {LeaseId}", leaseInfo.Lease.Id);
-			}
-
-			// Update the state of the lease
-			lock (_lockObject)
-			{
-				if (leaseInfo.CancellationTokenSource.IsCancellationRequested)
-				{
-					leaseInfo.Lease.State = LeaseState.Cancelled;
-					leaseInfo.Lease.Outcome = LeaseOutcome.Failed;
-					leaseInfo.Lease.Output = ByteString.Empty;
-				}
-				else
-				{
-					leaseInfo.Lease.State = (result.Outcome == LeaseOutcome.Cancelled) ? LeaseState.Cancelled : LeaseState.Completed;
-					leaseInfo.Lease.Outcome = result.Outcome;
-					leaseInfo.Lease.Output = (result.Output != null)? ByteString.CopyFrom(result.Output) : ByteString.Empty;
-				}
-				_logger.LogInformation("Transitioning lease {LeaseId} to {State}, outcome={Outcome}", leaseInfo.Lease.Id, leaseInfo.Lease.State, leaseInfo.Lease.Outcome);
-			}
-			_updateLeasesEvent.Set();
-		}
-
-		/// <summary>
-		/// Dispatch a lease payload to the appropriate handler
-		/// </summary>
-		/// <param name="rpcConnection">The RPC connection to the server</param>
-		/// <param name="agentId">The current agent id</param>
-		/// <param name="leaseInfo">Information about the lease</param>
-		/// <returns>Outcome from the lease</returns>
-		internal async Task<LeaseResult> HandleLeasePayloadAsync(IRpcConnection rpcConnection, string agentId, LeaseInfo leaseInfo)
-		{
-			Any payload = leaseInfo.Lease.Payload;
-
-			ComputeTaskMessage computeTask;
-			if (leaseInfo.Lease.Payload.TryUnpack(out computeTask))
-			{
-				GlobalTracer.Instance.ActiveSpan?.SetTag("task", "Compute");
-				return await ComputeAsync(rpcConnection, leaseInfo.Lease.Id, computeTask, _logger, leaseInfo.CancellationTokenSource.Token);
-			}
-
-			ConformTask conformTask;
-			if (leaseInfo.Lease.Payload.TryUnpack(out conformTask))
-			{
-				GlobalTracer.Instance.ActiveSpan?.SetTag("task", "Conform");
-				Task<LeaseResult> Handler(ILogger newLogger) => ConformAsync(rpcConnection, agentId, leaseInfo.Lease.Id, conformTask, newLogger, leaseInfo.CancellationTokenSource.Token);
-				return await HandleLeasePayloadWithLogAsync(rpcConnection, conformTask.LogId, null, null, Handler, leaseInfo.CancellationTokenSource.Token);
-			}
-
-			ExecuteJobTask jobTask;
-			if (leaseInfo.Lease.Payload.TryUnpack(out jobTask))
-			{
-				GlobalTracer.Instance.ActiveSpan?.SetTag("task", "Job");
-				Task<LeaseResult> Handler(ILogger newLogger) => ExecuteJobAsync(rpcConnection, leaseInfo.Lease.Id, jobTask, newLogger, leaseInfo.CancellationTokenSource.Token);
-				return await HandleLeasePayloadWithLogAsync(rpcConnection, jobTask.LogId, jobTask.JobId, jobTask.BatchId, Handler, leaseInfo.CancellationTokenSource.Token);
-			}
-
-			UpgradeTask upgradeTask;
-			if (leaseInfo.Lease.Payload.TryUnpack(out upgradeTask) || TryUnpackUpgradeTask(leaseInfo.Lease.Payload, out upgradeTask))
-			{
-				GlobalTracer.Instance.ActiveSpan?.SetTag("task", "Upgrade");
-				Task<LeaseResult> Handler(ILogger newLogger) => UpgradeAsync(rpcConnection, upgradeTask, newLogger, leaseInfo.CancellationTokenSource.Token);
-				return await HandleLeasePayloadWithLogAsync(rpcConnection, upgradeTask.LogId, null, null, Handler, leaseInfo.CancellationTokenSource.Token);
-			}
-
-			ShutdownTask shutdownTask;
-			if (leaseInfo.Lease.Payload.TryUnpack(out shutdownTask))
-			{
-				GlobalTracer.Instance.ActiveSpan?.SetTag("task", "Shutdown");
-				Task<LeaseResult> Handler(ILogger newLogger) => ShutdownAsync(newLogger);
-				return await HandleLeasePayloadWithLogAsync(rpcConnection, shutdownTask.LogId, null, null, Handler, leaseInfo.CancellationTokenSource.Token);
-			}
-
-			RestartTask restartTask;
-			if (leaseInfo.Lease.Payload.TryUnpack(out restartTask))
-			{
-				GlobalTracer.Instance.ActiveSpan?.SetTag("task", "Restart");
-				Task<LeaseResult> Handler(ILogger newLogger) => RestartAsync(newLogger);
-				return await HandleLeasePayloadWithLogAsync(rpcConnection, restartTask.LogId, null, null, Handler, leaseInfo.CancellationTokenSource.Token);
-			}
-			
-			TestTask testTask;
-			if (leaseInfo.Lease.Payload.TryUnpack(out testTask))
-			{
-				GlobalTracer.Instance.ActiveSpan?.SetTag("task", "Test");
-				return LeaseResult.Success;
-			}
-
-			_logger.LogError("Invalid lease payload type ({PayloadType})", payload.TypeUrl);
-			return LeaseResult.Failed;
-		}
-
-		/// <summary>
-		/// Attempts to parse an upgrade task, given either name
-		/// </summary>
-		/// <param name="payload"></param>
-		/// <param name="upgradeTask"></param>
-		/// <returns></returns>
-		static bool TryUnpackUpgradeTask(Any payload, out UpgradeTask upgradeTask)
-		{
-			if (payload.TypeUrl == "type.googleapis.com/UpgradeTask" || payload.TypeUrl == "type.googleapis.com/Horde.UpgradeTask")
-			{
-				upgradeTask = UpgradeTask.Parser.ParseFrom(payload.Value);
-				return true;
-			}
-			else
-			{
-				upgradeTask = null!;
-				return false;
-			}
-		}
-
-		/// <summary>
-		/// Terminates any processes which are still running under the given directory
-		/// </summary>
-		/// <param name="logger">Logger device</param>
-		/// <param name="cancellationToken">Cancellation token</param>
-		void TerminateProcesses(ILogger logger, CancellationToken cancellationToken)
-		{
-			// Terminate child processes from any previous runs
-			ProcessUtils.TerminateProcesses(ShouldTerminateProcess, logger, cancellationToken);
-		}
-
-		/// <summary>
-		/// Callback for determining whether a process should be terminated
-		/// </summary>
-		/// <param name="imageFile">The file to terminate</param>
-		/// <returns>True if the process should be terminated</returns>
-		bool ShouldTerminateProcess(FileReference imageFile)
-		{
-			if (imageFile.IsUnderDirectory(_workingDir))
-			{
-				return true;
-			}
-
-			string fileName = imageFile.GetFileName();
-			if (_processNamesToTerminate.Contains(fileName))
-			{
-				return true;
-			}
-
-			return false;
-		}
-
-		/// <summary>
-		/// Dispatch a lease payload to the appropriate handler
-		/// </summary>
-		/// <param name="rpcConnection">The RPC connection to the server</param>
-		/// <param name="logId">Unique id for the log to create</param>
-		/// <param name="jobId">The job being executed</param>
-		/// <param name="jobBatchId">The batch being executed</param>
-		/// <param name="executeLease">Action to perform to execute the lease</param>
-		/// <param name="cancellationToken">The cancellation token for this lease</param>
-		/// <returns>Outcome from the lease</returns>
-		async Task<LeaseResult> HandleLeasePayloadWithLogAsync(IRpcConnection rpcConnection, string logId, string? jobId, string? jobBatchId, Func<ILogger, Task<LeaseResult>> executeLease, CancellationToken cancellationToken)
-		{
-			await using (JsonRpcLogger? conformLogger = String.IsNullOrEmpty(logId) ? null : new JsonRpcLogger(rpcConnection, logId, jobId, jobBatchId, null, null, _logger))
-			{
-				ILogger leaseLogger = (conformLogger == null) ? (ILogger)_logger : new ForwardingLogger(_logger, new DefaultLoggerIndentHandler(conformLogger));
-				try
-				{
-					LeaseResult result = await executeLease(leaseLogger);
-					return result;
-				}
-				catch (Exception ex)
-				{
-					if (cancellationToken.IsCancellationRequested && IsCancellationException(ex))
-					{
-						leaseLogger.LogInformation("Lease was cancelled");
-						return LeaseResult.Cancelled;
-					}
-					else
-					{
-						leaseLogger.LogError(ex, "Caught unhandled exception while attempting to execute lease:\n{Exception}", ex);
-						return LeaseResult.Failed;
-					}
-				}
-			}
-		}
-
-		/// <summary>
-		/// Execute a remote execution aciton
-		/// </summary>
-		/// <param name="rpcConnection">RPC client for communicating with the server</param>
-		/// <param name="leaseId">The lease id</param>
-		/// <param name="computeTask">The action task parameters</param>
-		/// <param name="logger">Logger for the task</param>
-		/// <param name="cancellationToken">Token used to cancel the operation</param>
-		/// <returns></returns>
-		internal async Task<LeaseResult> ComputeAsync(IRpcConnection rpcConnection, string leaseId, ComputeTaskMessage computeTask, ILogger logger, CancellationToken cancellationToken)
-		{
-			logger.LogInformation("Starting compute task (lease {LeaseId})", leaseId);
-
-			ComputeTaskResultMessage result;
-			try
-			{
-				DateTimeOffset actionTaskStartTime = DateTimeOffset.UtcNow;
-				using (IRpcClientRef client = await rpcConnection.GetClientRef(new RpcContext(), cancellationToken))
-				{
-					DirectoryReference leaseDir = DirectoryReference.Combine(_workingDir, "Compute", leaseId);
-					DirectoryReference.CreateDirectory(leaseDir);
-
-					ComputeTaskExecutor executor = new ComputeTaskExecutor(_storageClient, logger);
-					try
-					{
-						result = await executor.ExecuteAsync(leaseId, computeTask, leaseDir, cancellationToken);
-					}
-					finally
-					{
-						try
-						{
-							DirectoryReference.Delete(leaseDir, true);
-						}
-						catch
-						{
-						}
-					}
-				}
-			}
-			catch (LegacyBlobNotFoundException ex)
-			{
-				logger.LogError(ex, "Blob not found: {Hash}", ex.Hash);
-				result = new ComputeTaskResultMessage(ComputeTaskOutcome.BlobNotFound, ex.Hash.ToString());
-			}
-			catch (Exception ex)
-			{
-				logger.LogError(ex, "Exception while executing compute task");
-				result = new ComputeTaskResultMessage(ComputeTaskOutcome.Exception, ex.ToString());
-			}
-
-			return new LeaseResult(result.ToByteArray());
-		}
-
-		/// <summary>
-		/// Conform a machine
-		/// </summary>
-		/// <param name="rpcConnection">RPC client for communicating with the server</param>
-		/// <param name="agentId">The current agent id</param>
-		/// <param name="leaseId">The current lease id</param>
-		/// <param name="conformTask">The conform task parameters</param>
-		/// <param name="conformLogger">Logger for the task</param>
-		/// <param name="cancellationToken">Token used to cancel the operation</param>
-		/// <returns>Async task</returns>
-		async Task<LeaseResult> ConformAsync(IRpcConnection rpcConnection, string agentId, string leaseId, ConformTask conformTask, ILogger conformLogger, CancellationToken cancellationToken)
-		{
-			conformLogger.LogInformation("Conforming, lease {LeaseId}", leaseId);
-			TerminateProcesses(conformLogger, cancellationToken);
-
-			bool removeUntrackedFiles = conformTask.RemoveUntrackedFiles;
-			IList<AgentWorkspace> pendingWorkspaces = conformTask.Workspaces;
-			for (; ;)
-			{
-				// Run the conform task
-				if (_settings.Executor == ExecutorType.Perforce && _settings.PerforceExecutor.RunConform)
-				{
-					await PerforceExecutor.ConformAsync(_workingDir, pendingWorkspaces, removeUntrackedFiles, conformLogger, cancellationToken);
-				}
-				else
-				{
-					conformLogger.LogInformation("Skipping due to Settings.RunConform flag");
-				}
-
-				// Update the new set of workspaces
-				UpdateAgentWorkspacesRequest request = new UpdateAgentWorkspacesRequest();
-				request.AgentId = agentId;
-				request.Workspaces.AddRange(pendingWorkspaces);
-				request.RemoveUntrackedFiles = removeUntrackedFiles;
-
-				UpdateAgentWorkspacesResponse response = await rpcConnection.InvokeAsync(x => x.UpdateAgentWorkspacesAsync(request, null, null, cancellationToken), new RpcContext(), cancellationToken);
-				if (!response.Retry)
-				{
-					conformLogger.LogInformation("Conform finished");
-					break;
-				}
-
-				conformLogger.LogInformation("Pending workspaces have changed - running conform again...");
-				pendingWorkspaces = response.PendingWorkspaces;
-				removeUntrackedFiles = response.RemoveUntrackedFiles;
-			}
-
-			return LeaseResult.Success;
-		}
-
-		/// <summary>
-		/// Execute part of a job
-		/// </summary>
-		/// <param name="rpcClient">RPC client for communicating with the server</param>
-		/// <param name="leaseId">The current lease id</param>
-		/// <param name="executeTask">The task to execute</param>
-		/// <param name="logger">The logger to use for this lease</param>
-		/// <param name="cancellationToken">Token used to cancel the operation</param>
-		/// <returns>Async task</returns>
-		internal async Task<LeaseResult> ExecuteJobAsync(IRpcConnection rpcClient, string leaseId, ExecuteJobTask executeTask, ILogger logger, CancellationToken cancellationToken)
-		{
-			logger.LogInformation("Executing job \"{JobName}\", jobId {JobId}, batchId {BatchId}, leaseId {LeaseId}", executeTask.JobName, executeTask.JobId, executeTask.BatchId, leaseId);
-			GlobalTracer.Instance.ActiveSpan?.SetTag("jobId", executeTask.JobId.ToString());
-			GlobalTracer.Instance.ActiveSpan?.SetTag("jobName", executeTask.JobName.ToString());
-			GlobalTracer.Instance.ActiveSpan?.SetTag("batchId", executeTask.BatchId.ToString());	
-
-			// Start executing the current batch
-			BeginBatchResponse batch = await rpcClient.InvokeAsync(x => x.BeginBatchAsync(new BeginBatchRequest(executeTask.JobId, executeTask.BatchId, leaseId), null, null, cancellationToken), new RpcContext(), cancellationToken);
-
-			// Execute the batch
-			try
-			{
-				await ExecuteBatchAsync(rpcClient, leaseId, executeTask, batch, logger, cancellationToken);
-			}
-			catch (Exception ex)
-			{
-				if (cancellationToken.IsCancellationRequested && IsCancellationException(ex))
-				{
-					logger.LogError("Step was aborted");
-					throw;
-				}
-				else
-				{
-					logger.LogError(ex, "Exception while executing batch: {Ex}", ex);
-				}
-			}
-
-			// If this lease was cancelled, don't bother updating the job state.
-			if (cancellationToken.IsCancellationRequested)
-			{
-				return LeaseResult.Cancelled;
-			}
-
-			// Mark the batch as complete
-			await rpcClient.InvokeAsync(x => x.FinishBatchAsync(new FinishBatchRequest(executeTask.JobId, executeTask.BatchId, leaseId), null, null, cancellationToken), new RpcContext(), cancellationToken);
-			logger.LogInformation("Done.");
-
-			return LeaseResult.Success;
-		}
-
-		/// <summary>
-		/// Executes a batch
-		/// </summary>
-		/// <param name="rpcClient">RPC client for communicating with the server</param>
-		/// <param name="leaseId">The current lease id</param>
-		/// <param name="executeTask">The task to execute</param>
-		/// <param name="batch">The batch to execute</param>
-		/// <param name="batchLogger">Output log for the batch</param>
-		/// <param name="cancellationToken">Cancellation token to abort the batch</param>
-		/// <returns>Async task</returns>
-		async Task ExecuteBatchAsync(IRpcConnection rpcClient, string leaseId, ExecuteJobTask executeTask, BeginBatchResponse batch, ILogger batchLogger, CancellationToken cancellationToken)
-		{
-			batchLogger.LogInformation("Executing batch {BatchId} using {Executor} executor", executeTask.BatchId, _settings.Executor.ToString());
-			TerminateProcesses(batchLogger, cancellationToken);
-
-			// Get the working directory for this lease
-			DirectoryReference scratchDir = DirectoryReference.Combine(_workingDir, "Scratch");
-
-			// Create an executor for this job
-			IExecutor executor = _createExecutor(rpcClient, executeTask, batch);
-
-			// Try to initialize the executor
-			batchLogger.LogInformation("Initializing...");
-			using (batchLogger.BeginIndentScope("  "))
-			{
-				using IScope scope = GlobalTracer.Instance.BuildSpan("Initialize").StartActive();
-				await executor.InitializeAsync(batchLogger, cancellationToken);
-			}
-
-			try
-			{
-				// Execute the steps
-				for (; ; )
-				{
-					// Get the next step to execute
-					BeginStepResponse step = await rpcClient.InvokeAsync(x => x.BeginStepAsync(new BeginStepRequest(executeTask.JobId, executeTask.BatchId, leaseId), null, null, cancellationToken), new RpcContext(), cancellationToken);
-					if (step.State == BeginStepResponse.Types.Result.Waiting)
-					{
-						batchLogger.LogInformation("Waiting for dependency to be ready");
-						await Task.Delay(TimeSpan.FromSeconds(20.0), cancellationToken);
-						continue;
-					}
-					else if (step.State == BeginStepResponse.Types.Result.Complete)
-					{
-						break;
-					}
-					else if (step.State != BeginStepResponse.Types.Result.Ready)
-					{
-						batchLogger.LogError("Unexpected step state: {StepState}", step.State);
-						break;
-					}
-
-					// Get current disk space available. This will allow us to more easily spot steps that eat up a lot of disk space.
-					string? driveName;
-					if (RuntimeInformation.IsOSPlatform(OSPlatform.Windows))
-					{
-						driveName = Path.GetPathRoot(_workingDir.FullName);
-					}
-					else
-					{
-						driveName = _workingDir.FullName;
-					}
-
-					float availableFreeSpace = 0;
-					if (driveName != null)
-					{
-						try
-						{
-							DriveInfo info = new DriveInfo(driveName);
-							availableFreeSpace = (1.0f * info.AvailableFreeSpace) / 1024 / 1024 / 1024;
-						}
-						catch (Exception ex)
-						{
-							batchLogger.LogWarning(ex, "Unable to query disk info for path '{DriveName}'", driveName);
-						}
-					}
-
-					// Print the new state
-					Stopwatch stepTimer = Stopwatch.StartNew();
-
-					batchLogger.LogInformation("Starting job {JobId}, batch {BatchId}, step {StepId} (Drive Space Left: {DriveSpaceRemaining} GB)", executeTask.JobId, executeTask.BatchId, step.StepId, availableFreeSpace.ToString("F1"));
-
-					// Create a trace span
-					using IScope scope = GlobalTracer.Instance.BuildSpan("Execute").WithResourceName(step.Name).StartActive();
-					scope.Span.SetTag("stepId", step.StepId);
-					scope.Span.SetTag("logId", step.LogId);
-					//				using IDisposable TraceProperty = LogContext.PushProperty("dd.trace_id", CorrelationIdentifier.TraceId.ToString());
-					//				using IDisposable SpanProperty = LogContext.PushProperty("dd.span_id", CorrelationIdentifier.SpanId.ToString());
-
-					// Update the context to include information about this step
-					JobStepOutcome stepOutcome;
-					JobStepState stepState;
-					using (batchLogger.BeginIndentScope("  "))
-					{
-						// Start writing to the log file
-#pragma warning disable CA2000 // Dispose objects before losing scope
-						await using (JsonRpcLogger stepLogger = new JsonRpcLogger(rpcClient, step.LogId, executeTask.JobId, executeTask.BatchId, step.StepId, step.Warnings, _logger))
-						{
-							// Execute the task
-							ILogger forwardingLogger = new DefaultLoggerIndentHandler(stepLogger);
-							if (_settings.WriteStepOutputToLogger)
-							{
-								forwardingLogger = new ForwardingLogger(_logger, forwardingLogger);
-							}
-
-							using CancellationTokenSource stepPollCancelSource = new CancellationTokenSource();
-							using CancellationTokenSource stepAbortSource = new CancellationTokenSource();
-							TaskCompletionSource<bool> stepFinishedSource = new TaskCompletionSource<bool>();
-							Task stepPollTask = Task.Run(() => PollForStepAbort(rpcClient, executeTask.JobId, executeTask.BatchId, step.StepId, stepAbortSource, stepFinishedSource.Task, stepPollCancelSource.Token), cancellationToken);
-
-							try
-							{
-								(stepOutcome, stepState) = await ExecuteStepAsync(executor, step, forwardingLogger, cancellationToken, stepAbortSource.Token);
-							}
-							finally
-							{
-								// Will get called even when cancellation token for the lease/batch fires
-								stepFinishedSource.SetResult(true); // Tell background poll task to stop
-								await stepPollTask;
-							}
-
-							// Kill any processes spawned by the step
-							TerminateProcesses(stepLogger, cancellationToken);
-
-							// Wait for the logger to finish
-							await stepLogger.StopAsync();
-
-							// Reflect the warnings/errors in the step outcome
-							if (stepOutcome > stepLogger.Outcome)
-							{
-								stepOutcome = stepLogger.Outcome;
-							}
-						}
-#pragma warning restore CA2000 // Dispose objects before losing scope
-
-						// Update the server with the outcome from the step
-						batchLogger.LogInformation("Marking step as complete (Outcome={Outcome}, State={StepState})", stepOutcome, stepState);
-						await rpcClient.InvokeAsync(x => x.UpdateStepAsync(new UpdateStepRequest(executeTask.JobId, executeTask.BatchId, step.StepId, stepState, stepOutcome), null, null, cancellationToken), new RpcContext(), cancellationToken);
-					}
-
-					// Print the finishing state
-					stepTimer.Stop();
-					batchLogger.LogInformation("Completed in {Time}", stepTimer.Elapsed);
-				}
-			}
-			catch (Exception ex)
-			{
-				if (cancellationToken.IsCancellationRequested && IsCancellationException(ex))
-				{
-					_logger.LogError("Step was aborted");
-				}
-				else
-				{
-					_logger.LogError(ex, "Exception while executing batch: {Ex}", ex);
-				}
-			}
-
-			// Clean the environment
-			batchLogger.LogInformation("Finalizing...");
-			using (batchLogger.BeginIndentScope("  "))
-			{
-				using IScope scope = GlobalTracer.Instance.BuildSpan("Finalize").StartActive();
-				await executor.FinalizeAsync(batchLogger, CancellationToken.None);
-			}
-		}
-
-		/// <summary>
-		/// Executes a step
-		/// </summary>
-		/// <param name="executor">The executor to run this step</param>
-		/// <param name="step">Step to execute</param>
-		/// <param name="stepLogger">Logger for the step</param>
-		/// <param name="cancellationToken">Cancellation token to abort the batch</param>
-		/// <param name="stepCancellationToken">Cancellation token to abort only this individual step</param>
-		/// <returns>Async task</returns>
-		internal static async Task<(JobStepOutcome, JobStepState)> ExecuteStepAsync(IExecutor executor, BeginStepResponse step, ILogger stepLogger, CancellationToken cancellationToken, CancellationToken stepCancellationToken)
-		{
-			using CancellationTokenSource combined = CancellationTokenSource.CreateLinkedTokenSource(cancellationToken, stepCancellationToken);
-			try
-			{
-				JobStepOutcome stepOutcome = await executor.RunAsync(step, stepLogger, combined.Token);
-				return (stepOutcome, JobStepState.Completed);
-			}
-			catch (Exception ex)
-			{
-				if (cancellationToken.IsCancellationRequested && IsCancellationException(ex))
-				{
-					stepLogger.LogError("The step was cancelled by batch/lease");
-					throw;
-				}
-				
-				if (stepCancellationToken.IsCancellationRequested && IsCancellationException(ex))
-				{
-					stepLogger.LogError("The step was intentionally cancelled");
-					return (JobStepOutcome.Failure, JobStepState.Aborted);
-				}
-				
-				stepLogger.LogError(ex, "Exception while executing step: {Ex}", ex);
-				return (JobStepOutcome.Failure, JobStepState.Completed);
-			}
-		}
-
-		internal async Task PollForStepAbort(IRpcConnection epcClient, string jobId, string batchId, string stepId, CancellationTokenSource stepCancelSource, Task finishedTask, CancellationToken cancellationToken)
-		{
-			Stopwatch timer = Stopwatch.StartNew();
-			while (!finishedTask.IsCompleted)
-			{
-				try
-				{
-					GetStepResponse res = await epcClient.InvokeAsync(x => x.GetStepAsync(new GetStepRequest(jobId, batchId, stepId), null, null, cancellationToken), new RpcContext(), cancellationToken);
-					if (res.AbortRequested)
-					{
-						_logger.LogDebug("Step was aborted by server (JobId={JobId} BatchId={BatchId} StepId={StepId})", jobId, batchId, stepId);
-						stepCancelSource.Cancel();
-						break;
-					}
-				}
-				catch (RpcException ex)
-				{
-					_logger.LogError(ex, "Poll for step abort has failed. Aborting (JobId={JobId} BatchId={BatchId} StepId={StepId})", jobId, batchId, stepId);
-					stepCancelSource.Cancel();
-					break;
-				}
-
-				await Task.WhenAny(Task.Delay(_stepAbortPollInterval, cancellationToken), finishedTask);
-			}
-		}
-
-		/// <summary>
-		/// Determine if the given exception was triggered due to a cancellation event
-		/// </summary>
-		/// <param name="ex">The exception to check</param>
-		/// <returns>True if the exception is a cancellation exception</returns>
-		static bool IsCancellationException(Exception ex)
-		{
-			if(ex is OperationCanceledException)
-			{
-				return true;
-			}
-
-			RpcException? rpcException = ex as RpcException;
-			if(rpcException != null && rpcException.StatusCode == StatusCode.Cancelled)
-			{
-				return true;
-			}
-
-			return false;
-		}
-
-		/// <summary>
-		/// Check for an update of the agent software
-		/// </summary>
-		/// <param name="rpcClient">RPC client for communicating with the server</param>
-		/// <param name="upgradeTask">The upgrade task</param>
-		/// <param name="logger">Logging device</param>
-		/// <param name="cancellationToken">Token used to cancel the operation</param>
-		/// <returns>Outcome of this operation</returns>
-		async Task<LeaseResult> UpgradeAsync(IRpcConnection rpcClient, UpgradeTask upgradeTask, ILogger logger, CancellationToken cancellationToken)
-		{
-			string requiredVersion = upgradeTask.SoftwareId;
-
-			// Check if we're running the right version
-			if (requiredVersion != null && requiredVersion != Program.Version)
-			{
-				logger.LogInformation("Upgrading from {CurrentVersion} to {TargetVersion}", Program.Version, requiredVersion);
-
-				// Clear out the working directory
-				DirectoryReference upgradeDir = DirectoryReference.Combine(_workingDir, "Upgrade");
-				DirectoryReference.CreateDirectory(upgradeDir);
-				await DeleteDirectoryContentsAsync(new DirectoryInfo(upgradeDir.FullName));
-
-				// Download the new software
-				FileInfo outputFile = new FileInfo(Path.Combine(upgradeDir.FullName, "Agent.zip"));
-				using (IRpcClientRef rpcClientRef = await rpcClient.GetClientRef(new RpcContext(), cancellationToken))
-				using (AsyncServerStreamingCall<DownloadSoftwareResponse> cursor = rpcClientRef.Client.DownloadSoftware(new DownloadSoftwareRequest(requiredVersion), null, null, cancellationToken))
-				{
-					using (Stream outputStream = outputFile.Open(FileMode.Create))
-					{
-						while (await cursor.ResponseStream.MoveNext(cancellationToken))
-						{
-							outputStream.Write(cursor.ResponseStream.Current.Data.Span);
-						}
-					}
-				}
-
-				// Extract it to a temporary directory
-				DirectoryReference extractedDir = DirectoryReference.Combine(upgradeDir, "Extracted");
-				DirectoryReference.CreateDirectory(extractedDir);
-				ZipFile.ExtractToDirectory(outputFile.FullName, extractedDir.FullName);
-
-				//				// Debug code for updating an agent with the local version
-				//				foreach (FileInfo SourceFile in new FileInfo(Assembly.GetExecutingAssembly().Location).Directory.EnumerateFiles())
-				//				{
-				//					SourceFile.CopyTo(Path.Combine(ExtractedDir.FullName, SourceFile.Name), true);
-				//				}
-
-				// Get the current process and assembly. This may be different if running through dotnet.exe rather than a native PE image.
-				FileReference assemblyFileName = new FileReference(Assembly.GetExecutingAssembly().Location);
-
-				// Spawn the other process
-				using (Process process = new Process())
-				{
-					StringBuilder arguments = new StringBuilder();
-
-					DirectoryReference targetDir = assemblyFileName.Directory;
-
-					// We were launched via an external application (presumably dotnet.exe). Do the same thing again.
-					FileReference newAssemblyFileName = FileReference.Combine(extractedDir, assemblyFileName.MakeRelativeTo(targetDir));
-					if (!FileReference.Exists(newAssemblyFileName))
-					{
-						logger.LogError("Unable to find {AgentExe} in extracted archive", newAssemblyFileName);
-						return LeaseResult.Failed;
-					}
-
-					process.StartInfo.FileName = "dotnet";
-
-					StringBuilder currentArguments = new StringBuilder();
-					foreach (string arg in Program.Args)
-					{
-						currentArguments.AppendArgument(arg);
-					}
-
-					arguments.AppendArgument(newAssemblyFileName.FullName);
-					arguments.AppendArgument("Service");
-					arguments.AppendArgument("Upgrade");
-					arguments.AppendArgument("-ProcessId=", Environment.ProcessId.ToString());
-					arguments.AppendArgument("-TargetDir=", targetDir.FullName);
-					arguments.AppendArgument("-Arguments=", currentArguments.ToString());
-
-					process.StartInfo.Arguments = arguments.ToString();
-					process.StartInfo.UseShellExecute = false;
-					process.EnableRaisingEvents = true;
-
-					StringBuilder launchCommand = new StringBuilder();
-					launchCommand.AppendArgument(process.StartInfo.FileName);
-					launchCommand.Append(' ');
-					launchCommand.Append(arguments);
-					logger.LogInformation("Launching: {Launch}", launchCommand.ToString());
-
-					TaskCompletionSource<int> exitCodeSource = new TaskCompletionSource<int>();
-					process.Exited += (sender, args) => { exitCodeSource.SetResult(process.ExitCode); };
-
-					process.Start();
-
-					using (cancellationToken.Register(() => { exitCodeSource.SetResult(0); }))
-					{
-						await exitCodeSource.Task;
-					}
-				}
-			}
-
-			return LeaseResult.Success;
-		}
-
-		/// <summary>
-		/// Delete the contents of a directory without deleting it itself
-		/// </summary>
-		/// <param name="baseDir">Directory to clean</param>
-		/// <returns>Async task</returns>
-		static async Task DeleteDirectoryContentsAsync(DirectoryInfo baseDir)
-		{
-			List<Task> childTasks = new List<Task>();
-			foreach (DirectoryInfo subDir in baseDir.EnumerateDirectories())
-			{
-				childTasks.Add(Task.Run(() => DeleteDirectory(subDir)));
-			}
-			foreach (FileInfo file in baseDir.EnumerateFiles())
-			{
-				file.Attributes = FileAttributes.Normal;
-				file.Delete();
-			}
-			await Task.WhenAll(childTasks);
-		}
-
-		/// <summary>
-		/// Deletes a directory and its contents
-		/// </summary>
-		/// <param name="baseDir">Directory to delete</param>
-		/// <returns>Async task</returns>
-		static async Task DeleteDirectory(DirectoryInfo baseDir)
-		{
-			await DeleteDirectoryContentsAsync(baseDir);
-			baseDir.Delete();
-		}
-
-		/// <summary>
-		/// Check for an update of the agent software
-		/// </summary>
-		/// <param name="logger">Logging device</param>
-		/// <returns>Outcome of this operation</returns>
-		Task<LeaseResult> RestartAsync(ILogger logger)
-		{
-			logger.LogInformation("Setting restart flag");
-			_requestShutdown = true;
-			_restartAfterShutdown = true;
-			return Task.FromResult(LeaseResult.Success);
-		}
-
-		/// <summary>
-		/// Shutdown the agent
-		/// </summary>
-		/// <param name="logger">Logging device</param>
-		/// <returns>Outcome of this operation</returns>
-		internal Task<LeaseResult> ShutdownAsync(ILogger logger)
-		{
-			logger.LogInformation("Setting shutdown flag");
-			_requestShutdown = true;
-			_restartAfterShutdown = false;
-			return Task.FromResult(LeaseResult.Success);
-		}
-
-		/// <summary>
-		/// Gets the hardware capabilities of this worker
-		/// </summary>
-		/// <returns>Worker object for advertising to the server</returns>
-		public static async Task<AgentCapabilities> GetAgentCapabilities(DirectoryReference workingDir, ILogger logger, Dictionary<string, string>? extraProperties = null)
-		{
-			// Create the primary device
-			DeviceCapabilities primaryDevice = new DeviceCapabilities();
-			primaryDevice.Handle = "Primary";
-
-			List<DeviceCapabilities> otherDevices = new List<DeviceCapabilities>();
-			if (RuntimeInformation.IsOSPlatform(OSPlatform.Windows))
-			{
-				primaryDevice.Properties.Add("Platform=Win64");
-				primaryDevice.Properties.Add("PlatformGroup=Windows");
-				primaryDevice.Properties.Add("PlatformGroup=Microsoft");
-				primaryDevice.Properties.Add("PlatformGroup=Desktop");
-
-				primaryDevice.Properties.Add("OSFamily=Windows");
-
-				// Add OS info
-				using (ManagementObjectSearcher searcher = new ManagementObjectSearcher("select * from Win32_OperatingSystem"))
-				{
-					foreach (ManagementObject row in searcher.Get())
-					{
-						using (row)
-						{
-							Dictionary<string, object> properties = GetWmiProperties(row);
-
-							object? name;
-							if (properties.TryGetValue("Caption", out name))
-							{
-								primaryDevice.Properties.Add($"OSDistribution={name}");
-							}
-
-							object? version;
-							if (properties.TryGetValue("Version", out version))
-							{
-								primaryDevice.Properties.Add($"OSKernelVersion={version}");
-							}
-						}
-					}
-				}
-
-				// Add CPU info
-				using (ManagementObjectSearcher searcher = new ManagementObjectSearcher("select * from Win32_Processor"))
-				{
-					Dictionary<string, int> nameToCount = new Dictionary<string, int>();
-					int totalPhysicalCores = 0;
-					int totalLogicalCores = 0;
-
-					foreach (ManagementObject row in searcher.Get())
-					{
-						using (row)
-						{
-							Dictionary<string, object> properties = GetWmiProperties(row);
-
-							object? nameObject;
-							if (properties.TryGetValue("Name", out nameObject))
-							{
-								string name = nameObject.ToString() ?? String.Empty;
-								int count;
-								nameToCount.TryGetValue(name, out count);
-								nameToCount[name] = count + 1;
-							}
-
-							object? numPhysicalCores;
-							if ((properties.TryGetValue("NumberOfEnabledCore", out numPhysicalCores) && numPhysicalCores is uint) || (properties.TryGetValue("NumberOfCores", out numPhysicalCores) && numPhysicalCores is uint))
-							{
-								totalPhysicalCores += (int)(uint)numPhysicalCores;
-							}
-
-							object? numLogicalCores;
-							if (properties.TryGetValue("NumberOfLogicalProcessors", out numLogicalCores) && numLogicalCores is uint numLogicalCoresUint)
-							{
-								totalLogicalCores += (int)numLogicalCoresUint;
-							}
-						}
-					}
-
-					AddCpuInfo(primaryDevice, nameToCount, totalLogicalCores, totalPhysicalCores);
-				}
-
-				// Add RAM info
-				using (ManagementObjectSearcher searcher = new ManagementObjectSearcher("select Capacity from Win32_PhysicalMemory"))
-				{
-					ulong totalCapacity = 0;
-					foreach (ManagementObject row in searcher.Get())
-					{
-						using (row)
-						{
-							object? capacity = row.GetPropertyValue("Capacity");
-							if (capacity is ulong ulongCapacity)
-							{
-								totalCapacity += ulongCapacity;
-							}
-						}
-					}
-
-					if (totalCapacity > 0)
-					{
-						primaryDevice.Properties.Add($"RAM={totalCapacity / (1024 * 1024 * 1024)}");
-					}
-				}
-
-				// Add GPU info
-				using (ManagementObjectSearcher searcher = new ManagementObjectSearcher("select Name, DriverVersion, AdapterRAM from Win32_VideoController"))
-				{
-					int index = 0;
-					foreach (ManagementObject row in searcher.Get())
-					{
-						using (row)
-						{
-							WmiProperties properties = new WmiProperties(row);
-							if (properties.TryGetValue("Name", out string? name) && properties.TryGetValue("DriverVersion", out string? driverVersion))
-							{
-								string prefix = $"GPU-{++index}";
-								primaryDevice.Properties.Add($"{prefix}-Name={name}");
-								primaryDevice.Properties.Add($"{prefix}-DriverVersion={driverVersion}");								
-							}
-						}
-					}
-				}
-
-				// Add EC2 properties if needed
-				await AddAwsProperties(primaryDevice.Properties, logger);
-
-				// Add session information
-				primaryDevice.Properties.Add($"User={Environment.UserName}");
-				primaryDevice.Properties.Add($"Domain={Environment.UserDomainName}");
-				primaryDevice.Properties.Add($"Interactive={Environment.UserInteractive}");
-				primaryDevice.Properties.Add($"Elevated={BuildGraphExecutor.IsUserAdministrator()}");
-			}
-			else if (RuntimeInformation.IsOSPlatform(OSPlatform.Linux))
-			{
-				primaryDevice.Properties.Add("Platform=Linux");
-				primaryDevice.Properties.Add("PlatformGroup=Linux");
-				primaryDevice.Properties.Add("PlatformGroup=Unix");
-				primaryDevice.Properties.Add("PlatformGroup=Desktop");
-
-				primaryDevice.Properties.Add("OSFamily=Linux");
-				primaryDevice.Properties.Add("OSVersion=Linux");
-
-				// Add EC2 properties if needed
-				await AddAwsProperties(primaryDevice.Properties, logger);
-
-				// Parse the CPU info
-				List<Dictionary<string, string>>? cpuRecords = await ReadLinuxHwPropsAsync("/proc/cpuinfo", logger);
-				if (cpuRecords != null)
-				{
-					Dictionary<string, string> cpuNames = new Dictionary<string, string>(StringComparer.Ordinal);
-					foreach (Dictionary<string, string> cpuRecord in cpuRecords)
-					{
-						if (cpuRecord.TryGetValue("physical id", out string? physicalId) && cpuRecord.TryGetValue("model name", out string? modelName))
-						{
-							cpuNames[physicalId] = modelName;
-						}
-					}
-
-					Dictionary<string, int> nameToCount = new Dictionary<string, int>(StringComparer.Ordinal);
-					foreach (string cpuName in cpuNames.Values)
-					{
-						nameToCount.TryGetValue(cpuName, out int count);
-						nameToCount[cpuName] = count + 1;
-					}
-
-					HashSet<string> logicalCores = new HashSet<string>();
-					HashSet<string> physicalCores = new HashSet<string>();
-					foreach (Dictionary<string, string> cpuRecord in cpuRecords)
-					{
-						if (cpuRecord.TryGetValue("processor", out string? logicalCoreId))
-						{
-							logicalCores.Add(logicalCoreId);
-						}
-						if (cpuRecord.TryGetValue("core id", out string? physicalCoreId))
-						{
-							physicalCores.Add(physicalCoreId);
-						}
-					}
-
-					AddCpuInfo(primaryDevice, nameToCount, logicalCores.Count, physicalCores.Count);
-				}
-
-				// Parse the RAM info
-				List<Dictionary<string, string>>? memRecords = await ReadLinuxHwPropsAsync("/proc/meminfo", logger);
-				if (memRecords != null && memRecords.Count > 0 && memRecords[0].TryGetValue("MemTotal", out string? memTotal))
-				{
-					Match match = Regex.Match(memTotal, @"(\d+)\s+kB");
-					if (match.Success)
-					{
-						long totalCapacity = Int64.Parse(match.Groups[1].Value) * 1024;
-						primaryDevice.Properties.Add($"RAM={totalCapacity / (1024 * 1024 * 1024)}");
-					}
-				}
-
-				// Add session information
-				primaryDevice.Properties.Add($"User={Environment.UserName}");
-				primaryDevice.Properties.Add($"Domain={Environment.UserDomainName}");
-				primaryDevice.Properties.Add($"Interactive={Environment.UserInteractive}");
-				primaryDevice.Properties.Add($"Elevated={BuildGraphExecutor.IsUserAdministrator()}");
-			}
-			else if(RuntimeInformation.IsOSPlatform(OSPlatform.OSX))
-			{
-				primaryDevice.Properties.Add("Platform=Mac");
-				primaryDevice.Properties.Add("PlatformGroup=Apple");
-				primaryDevice.Properties.Add("PlatformGroup=Desktop");
-
-				primaryDevice.Properties.Add("OSFamily=MacOS");
-				primaryDevice.Properties.Add("OSVersion=MacOS");
-
-				string output;
-				using(Process process = new Process())
-				{
-					process.StartInfo.FileName = "system_profiler";
-					process.StartInfo.Arguments = "SPHardwareDataType SPSoftwareDataType -xml";
-					process.StartInfo.CreateNoWindow = true;
-					process.StartInfo.RedirectStandardOutput = true;
-					process.StartInfo.RedirectStandardInput = false;
-					process.StartInfo.UseShellExecute = false;
-					process.Start();
-
-					output = await process.StandardOutput.ReadToEndAsync();
-				}
-
-				XmlDocument xml = new XmlDocument();
-				xml.LoadXml(output);
-
-				XmlNode? hardwareNode = xml.SelectSingleNode("/plist/array/dict/key[. = '_dataType']/following-sibling::string[. = 'SPHardwareDataType']/../key[. = '_items']/following-sibling::array/dict");
-				if(hardwareNode != null)
-				{
-					XmlNode? model = hardwareNode.SelectSingleNode("key[. = 'machine_model']/following-sibling::string");
-					if(model != null)
-					{
-						primaryDevice.Properties.Add($"Model={model.InnerText}");
-					}
-
-					XmlNode? cpuTypeNode = hardwareNode.SelectSingleNode("key[. = 'cpu_type']/following-sibling::string");
-					XmlNode? cpuSpeedNode = hardwareNode.SelectSingleNode("key[. = 'current_processor_speed']/following-sibling::string");
-					XmlNode? cpuPackagesNode = hardwareNode.SelectSingleNode("key[. = 'packages']/following-sibling::integer");
-					if(cpuTypeNode != null && cpuSpeedNode != null && cpuPackagesNode != null)
-					{
-						primaryDevice.Properties.Add((cpuPackagesNode.InnerText != "1")? $"CPU={cpuPackagesNode.InnerText} x {cpuTypeNode.InnerText} @ {cpuSpeedNode.InnerText}" : $"CPU={cpuTypeNode.InnerText} @ {cpuSpeedNode.InnerText}");
-					}
-
-					primaryDevice.Properties.Add($"LogicalCores={Environment.ProcessorCount}");
-
-					XmlNode? cpuCountNode = hardwareNode.SelectSingleNode("key[. = 'number_processors']/following-sibling::integer");
-					if(cpuCountNode != null)
-					{
-						primaryDevice.Properties.Add($"PhysicalCores={cpuCountNode.InnerText}");
-					}
-
-					XmlNode? memoryNode = hardwareNode.SelectSingleNode("key[. = 'physical_memory']/following-sibling::string");
-					if(memoryNode != null)
-					{
-						string[] parts = memoryNode.InnerText.Split(new char[]{' '}, StringSplitOptions.RemoveEmptyEntries);
-						if(parts.Length == 2 && parts[1] == "GB")
-						{
-							primaryDevice.Properties.Add($"RAM={parts[0]}");
-						}
-					}
-				}
-
-				XmlNode? softwareNode = xml.SelectSingleNode("/plist/array/dict/key[. = '_dataType']/following-sibling::string[. = 'SPSoftwareDataType']/../key[. = '_items']/following-sibling::array/dict");
-				if(softwareNode != null)
-				{
-					XmlNode? osVersionNode = softwareNode.SelectSingleNode("key[. = 'os_version']/following-sibling::string");
-					if(osVersionNode != null)
-					{
-						primaryDevice.Properties.Add($"OSDistribution={osVersionNode.InnerText}");
-					}
-
-					XmlNode? kernelVersionNode = softwareNode.SelectSingleNode("key[. = 'kernel_version']/following-sibling::string");
-					if(kernelVersionNode != null)
-					{
-						primaryDevice.Properties.Add($"OSKernelVersion={kernelVersionNode.InnerText}");
-					}
-				}
-			}
-
-			// Get the IP addresses
-			try
-			{
-				using CancellationTokenSource dnsCts = new(3000);
-				IPHostEntry entry = await Dns.GetHostEntryAsync(Dns.GetHostName(), dnsCts.Token);
-				foreach (IPAddress address in entry.AddressList)
-				{
-					if (address.AddressFamily == System.Net.Sockets.AddressFamily.InterNetwork)
-					{
-						primaryDevice.Properties.Add($"Ipv4={address}");
-					}
-					else if (address.AddressFamily == System.Net.Sockets.AddressFamily.InterNetworkV6)
-					{
-						primaryDevice.Properties.Add($"Ipv6={address}");
-					}
-				}
-			}
-			catch (Exception ex)
-			{
-				logger.LogDebug(ex, "Unable to get local IP address");
-			}
-
-			// Get the time that the machine booted
-			primaryDevice.Properties.Add($"BootTime={s_bootTime}");
-			primaryDevice.Properties.Add($"StartTime={s_startTime}");
-
-			// Add disk info based on platform
-			string? driveName;
-			if (RuntimeInformation.IsOSPlatform(OSPlatform.Windows))
-			{
-				driveName = Path.GetPathRoot(workingDir.FullName);
-			}
-			else
-			{
-				driveName = workingDir.FullName;
-			}
-
-			if (driveName != null)
-			{
-				try
-				{
-					DriveInfo info = new DriveInfo(driveName);
-					primaryDevice.Properties.Add($"DiskFreeSpace={info.AvailableFreeSpace}");
-					primaryDevice.Properties.Add($"DiskTotalSize={info.TotalSize}");
-				}
-				catch (Exception ex)
-				{
-					logger.LogWarning(ex, "Unable to query disk info for path '{DriveName}'", driveName);
-				}
-			}
-			primaryDevice.Properties.Add($"WorkingDir={workingDir}");
-
-			// Add any horde. env vars for custom properties.
-			IEnumerable<string> envVars = Environment.GetEnvironmentVariables().Keys.Cast<string>();
-			foreach (string envVar in envVars.Where(x => x.StartsWith("horde.", StringComparison.InvariantCultureIgnoreCase)))
-			{
-				primaryDevice.Properties.Add($"{envVar}={Environment.GetEnvironmentVariable(envVar)}");
-			}
-
-			// Create the worker
-			AgentCapabilities agent = new ();
-			agent.Devices.Add(primaryDevice);
-			agent.Devices.AddRange(otherDevices);
-			
-			if (extraProperties != null)
-			{
-				agent.Properties.AddRange(extraProperties.Select(kvp => $"{kvp.Key}={kvp.Value}"));				
-			}
-
-			return agent;
-		}
-
-		static void AddCpuInfo(DeviceCapabilities primaryDevice, Dictionary<string, int> nameToCount, int numLogicalCores, int numPhysicalCores)
-		{
-			if (nameToCount.Count > 0)
-			{
-				primaryDevice.Properties.Add("CPU=" + String.Join(", ", nameToCount.Select(x => (x.Value > 1) ? $"{x.Key} x {x.Value}" : x.Key)));
-			}
-
-			if (numLogicalCores > 0)
-			{
-				primaryDevice.Properties.Add($"LogicalCores={numLogicalCores}");
-			}
-
-			if (numPhysicalCores > 0)
-			{
-				primaryDevice.Properties.Add($"PhysicalCores={numPhysicalCores}");
-			}
-		}
-
-		static async Task<List<Dictionary<string, string>>?> ReadLinuxHwPropsAsync(string fileName, ILogger logger)
-		{
-			List<Dictionary<string, string>>? records = null;
-			if (File.Exists(fileName))
-			{
-				records = new List<Dictionary<string, string>>();
-				using (StreamReader reader = new StreamReader(fileName))
-				{
-					Dictionary<string, string> record = new Dictionary<string, string>(StringComparer.Ordinal);
-
-					string? line;
-					while ((line = await reader.ReadLineAsync()) != null)
-					{
-						int idx = line.IndexOf(':', StringComparison.Ordinal);
-						if (idx == -1)
-						{
-							if (record.Count > 0)
-							{
-								records.Add(record);
-								record = new Dictionary<string, string>(StringComparer.Ordinal);
-							}
-						}
-						else
-						{
-							string key = line.Substring(0, idx).Trim();
-							string value = line.Substring(idx + 1).Trim();
-
-							if (record.TryGetValue(key, out string? prevValue))
-							{
-								logger.LogWarning("Multiple entries for {Key} in {File} (was '{Prev}', now '{Next}')", key, fileName, prevValue, value);
-							}
-							else
-							{
-								record.Add(key, value);
-							}
-						}
-					}
-
-					if (record.Count > 0)
-					{
-						records.Add(record);
-					}
-				}
-			}
-			return records;
-		}
-
-		[SupportedOSPlatform("windows")]
-		class WmiProperties
-		{
-			readonly Dictionary<string, object> _properties = new Dictionary<string, object>(StringComparer.Ordinal);
-
-			public WmiProperties(ManagementObject row)
-			{
-				foreach (PropertyData property in row.Properties)
-				{
-					_properties[property.Name] = property.Value;
-				}
-			}
-
-			public bool TryGetValue(string name, [NotNullWhen(true)] out string? value)
-			{
-				object? @object;
-				if (_properties.TryGetValue(name, out @object))
-				{
-					value = @object.ToString();
-					return value != null;
-				}
-				else
-				{
-					value = null!;
-					return false;
-				}
-			}
-
-			public bool TryGetValue(string name, out long value)
-			{
-				object? objectValue;
-				if (_properties.TryGetValue(name, out objectValue))
-				{
-					if (objectValue is int intValue)
-					{
-						value = intValue;
-						return true;
-					}
-					else if (objectValue is uint uintValue)
-					{
-						value = uintValue;
-						return true;
-					}
-					else if (objectValue is long longValue)
-					{
-						value = longValue;
-						return true;
-					}
-					else if (objectValue is ulong ulongValue)
-					{
-						value = (long)ulongValue;
-						return true;
-					}
-				}
-
-				value = 0;
-				return false;
-			}
-		}
-
-		[SupportedOSPlatform("windows")]
-		static Dictionary<string, object> GetWmiProperties(ManagementObject row)
-		{
-			Dictionary<string, object> properties = new Dictionary<string, object>(StringComparer.Ordinal);
-			foreach (PropertyData property in row.Properties)
-			{
-				properties[property.Name] = property.Value;
-			}
-			return properties;
-		}
-
-		static async Task AddAwsProperties(IList<string> properties, ILogger logger)
-		{
-			if (EC2InstanceMetadata.IdentityDocument != null)
-			{
-				properties.Add("EC2=1");
-				AddAwsProperty("aws-instance-id", "/instance-id", properties);
-				AddAwsProperty("aws-instance-type", "/instance-type", properties);
-				AddAwsProperty("aws-region", "/region", properties);
-
-				try
-				{
-					using (AmazonEC2Client client = new AmazonEC2Client())
-					{
-						DescribeTagsRequest request = new DescribeTagsRequest();
-						request.Filters = new List<Filter>();
-						request.Filters.Add(new Filter("resource-id", new List<string> { EC2InstanceMetadata.InstanceId }));
-
-						DescribeTagsResponse response = await client.DescribeTagsAsync(request);
-						foreach (TagDescription tag in response.Tags)
-						{
-							properties.Add($"aws-tag={tag.Key}:{tag.Value}");
-						}
-					}
-				}
-				catch (Exception ex)
-				{
-					logger.LogDebug(ex, "Unable to query EC2 tags.");
-				}
-			}
-		}			
-
-		static void AddAwsProperty(string name, string awsKey, IList<string> properties)
-		{
-			string? value = EC2InstanceMetadata.GetData(awsKey);
-			if (value != null)
-			{
-				properties.Add($"{name}={value}");
-			}
-		}
 	}
 }
diff --git a/Engine/Source/Programs/Horde/Horde.Agent/Utility/ExceptionUtils.cs b/Engine/Source/Programs/Horde/Horde.Agent/Utility/ExceptionUtils.cs
new file mode 100644
--- /dev/null
+++ b/Engine/Source/Programs/Horde/Horde.Agent/Utility/ExceptionUtils.cs
@@ -0,0 +1,34 @@
+// Copyright Epic Games, Inc. All Rights Reserved.
+
+using System;
+using Grpc.Core;
+
+namespace Horde.Agent.Utility
+{
+	/// <summary>
+	/// Utility methods for exceptions
+	/// </summary>
+	static class ExceptionUtils
+	{
+		/// <summary>
+		/// Determine if the given exception was triggered due to a cancellation event
+		/// </summary>
+		/// <param name="ex">The exception to check</param>
+		/// <returns>True if the exception is a cancellation exception</returns>
+		public static bool IsCancellationException(this Exception ex)
+		{
+			if (ex is OperationCanceledException)
+			{
+				return true;
+			}
+
+			RpcException? rpcException = ex as RpcException;
+			if (rpcException != null && rpcException.StatusCode == StatusCode.Cancelled)
+			{
+				return true;
+			}
+
+			return false;
+		}
+	}
+}
diff --git a/Engine/Source/Programs/Horde/Horde.Agent/Utility/JsonRpcLogger.cs b/Engine/Source/Programs/Horde/Horde.Agent/Utility/JsonRpcLogger.cs
--- a/Engine/Source/Programs/Horde/Horde.Agent/Utility/JsonRpcLogger.cs
+++ b/Engine/Source/Programs/Horde/Horde.Agent/Utility/JsonRpcLogger.cs
@@ -130,6 +130,29 @@ namespace Horde.Agent.Parser
 			Outcome = JobStepOutcome.Success;
 		}
 
+		/// <summary>
+		/// Constructor
+		/// </summary>
+		/// <param name="rpcClient">RPC client to use for server requests</param>
+		/// <param name="logId">The log id to write to</param>
+		/// <param name="inner">Additional logger to write to</param>
+		public JsonRpcLogger(IRpcConnection rpcClient, string logId, ILogger inner)
+			: this(rpcClient, logId, null, null, null, null, inner)
+		{
+		}
+
+		/// <summary>
+		/// Constructor
+		/// </summary>
+		/// <param name="rpcClient">RPC client to use for server requests</param>
+		/// <param name="logId">The log id to write to</param>
+		/// <param name="warnings">Whether to include warnings in the output</param>
+		/// <param name="inner">Additional logger to write to</param>
+		public JsonRpcLogger(IRpcConnection rpcClient, string logId, bool? warnings, ILogger inner)
+			: this(rpcClient, logId, null, null, null, warnings, inner)
+		{
+		}
+
 		/// <summary>
 		/// Constructor
 		/// </summary>
diff --git a/Engine/Source/Programs/Horde/Horde.Agent/Utility/RpcConnection.cs b/Engine/Source/Programs/Horde/Horde.Agent/Utility/RpcConnection.cs
--- a/Engine/Source/Programs/Horde/Horde.Agent/Utility/RpcConnection.cs
+++ b/Engine/Source/Programs/Horde/Horde.Agent/Utility/RpcConnection.cs
@@ -68,7 +68,7 @@ namespace Horde.Agent.Utility
 		}
 	}
 
-	internal interface IRpcConnection : IAsyncDisposable
+	interface IRpcConnection : IAsyncDisposable
 	{
 		/// <summary>
 		/// Attempts to get a client reference, returning immediately if there's not one available
diff --git a/Engine/Source/Programs/Horde/Horde.Agent/appsettings.Local.json b/Engine/Source/Programs/Horde/Horde.Agent/appsettings.Local.json
new file mode 100644
--- /dev/null
+++ b/Engine/Source/Programs/Horde/Horde.Agent/appsettings.Local.json
@@ -0,0 +1,6 @@
+{
+	"Horde": 
+	{
+		"EnableAwsEc2Support": false
+	}
+}
