Horde: Missing files.

diff --git a/Engine/Source/Programs/Horde/Horde.Agent/Services/CapabilitiesService.cs b/Engine/Source/Programs/Horde/Horde.Agent/Services/CapabilitiesService.cs
new file mode 100644
--- /dev/null
+++ b/Engine/Source/Programs/Horde/Horde.Agent/Services/CapabilitiesService.cs
@@ -0,0 +1,590 @@
+// Copyright Epic Games, Inc. All Rights Reserved.
+
+using System;
+using System.Collections.Generic;
+using System.Diagnostics;
+using System.Diagnostics.CodeAnalysis;
+using System.IO;
+using System.Linq;
+using System.Management;
+using System.Net;
+using System.Runtime.InteropServices;
+using System.Runtime.Versioning;
+using System.Text.RegularExpressions;
+using System.Threading;
+using System.Threading.Tasks;
+using System.Xml;
+using Amazon.EC2;
+using Amazon.EC2.Model;
+using Amazon.Util;
+using EpicGames.Core;
+using Horde.Agent.Execution;
+using HordeCommon.Rpc.Messages;
+using Microsoft.Extensions.Logging;
+using Microsoft.Extensions.Options;
+
+namespace Horde.Agent.Services
+{
+	/// <summary>
+	/// Implements the message handling loop for an agent. Runs asynchronously until disposed.
+	/// </summary>
+	class CapabilitiesService
+	{
+		static readonly DateTimeOffset s_startTime = DateTimeOffset.Now;
+		static readonly DateTimeOffset s_bootTime = DateTimeOffset.Now - TimeSpan.FromTicks(Environment.TickCount64 * TimeSpan.TicksPerMillisecond);
+
+		readonly AgentSettings _settings;
+		readonly ILogger _logger;
+
+		/// <summary>
+		/// Constructor
+		/// </summary>
+		public CapabilitiesService(IOptions<AgentSettings> settings, ILogger<CapabilitiesService> logger)
+		{
+			_settings = settings.Value;
+			_logger = logger;
+		}
+
+		/// <summary>
+		/// Gets the hardware capabilities of this worker
+		/// </summary>
+		/// <param name="workingDir">Working directory for the agent</param>
+		/// <returns>Worker object for advertising to the server</returns>
+		public async Task<AgentCapabilities> GetCapabilitiesAsync(DirectoryReference? workingDir)
+		{
+			ILogger logger = _logger;
+
+			// Create the primary device
+			DeviceCapabilities primaryDevice = new DeviceCapabilities();
+			primaryDevice.Handle = "Primary";
+
+			List<DeviceCapabilities> otherDevices = new List<DeviceCapabilities>();
+			if (RuntimeInformation.IsOSPlatform(OSPlatform.Windows))
+			{
+				primaryDevice.Properties.Add("Platform=Win64");
+				primaryDevice.Properties.Add("PlatformGroup=Windows");
+				primaryDevice.Properties.Add("PlatformGroup=Microsoft");
+				primaryDevice.Properties.Add("PlatformGroup=Desktop");
+
+				primaryDevice.Properties.Add("OSFamily=Windows");
+
+				// WMI doesn't work currently on ARM64, due to reliance on a NET Framework assembly.
+				if (RuntimeInformation.OSArchitecture != Architecture.Arm64)
+				{
+					// Add OS info
+					using (ManagementObjectSearcher searcher = new ManagementObjectSearcher("select * from Win32_OperatingSystem"))
+					{
+						foreach (ManagementObject row in searcher.Get())
+						{
+							using (row)
+							{
+								Dictionary<string, object> properties = GetWmiProperties(row);
+
+								object? name;
+								if (properties.TryGetValue("Caption", out name))
+								{
+									primaryDevice.Properties.Add($"OSDistribution={name}");
+								}
+
+								object? version;
+								if (properties.TryGetValue("Version", out version))
+								{
+									primaryDevice.Properties.Add($"OSKernelVersion={version}");
+								}
+							}
+						}
+					}
+
+					// Add CPU info
+					using (ManagementObjectSearcher searcher = new ManagementObjectSearcher("select * from Win32_Processor"))
+					{
+						Dictionary<string, int> nameToCount = new Dictionary<string, int>();
+						int totalPhysicalCores = 0;
+						int totalLogicalCores = 0;
+
+						foreach (ManagementObject row in searcher.Get())
+						{
+							using (row)
+							{
+								Dictionary<string, object> properties = GetWmiProperties(row);
+
+								object? nameObject;
+								if (properties.TryGetValue("Name", out nameObject))
+								{
+									string name = nameObject.ToString() ?? String.Empty;
+									int count;
+									nameToCount.TryGetValue(name, out count);
+									nameToCount[name] = count + 1;
+								}
+
+								object? numPhysicalCores;
+								if ((properties.TryGetValue("NumberOfEnabledCore", out numPhysicalCores) && numPhysicalCores is uint) || (properties.TryGetValue("NumberOfCores", out numPhysicalCores) && numPhysicalCores is uint))
+								{
+									totalPhysicalCores += (int)(uint)numPhysicalCores;
+								}
+
+								object? numLogicalCores;
+								if (properties.TryGetValue("NumberOfLogicalProcessors", out numLogicalCores) && numLogicalCores is uint numLogicalCoresUint)
+								{
+									totalLogicalCores += (int)numLogicalCoresUint;
+								}
+							}
+						}
+
+						AddCpuInfo(primaryDevice, nameToCount, totalLogicalCores, totalPhysicalCores);
+					}
+
+					// Add RAM info
+					using (ManagementObjectSearcher searcher = new ManagementObjectSearcher("select Capacity from Win32_PhysicalMemory"))
+					{
+						ulong totalCapacity = 0;
+						foreach (ManagementObject row in searcher.Get())
+						{
+							using (row)
+							{
+								object? capacity = row.GetPropertyValue("Capacity");
+								if (capacity is ulong ulongCapacity)
+								{
+									totalCapacity += ulongCapacity;
+								}
+							}
+						}
+
+						if (totalCapacity > 0)
+						{
+							primaryDevice.Properties.Add($"RAM={totalCapacity / (1024 * 1024 * 1024)}");
+						}
+					}
+
+					// Add GPU info
+					using (ManagementObjectSearcher searcher = new ManagementObjectSearcher("select Name, DriverVersion, AdapterRAM from Win32_VideoController"))
+					{
+						int index = 0;
+						foreach (ManagementObject row in searcher.Get())
+						{
+							using (row)
+							{
+								WmiProperties properties = new WmiProperties(row);
+								if (properties.TryGetValue("Name", out string? name) && properties.TryGetValue("DriverVersion", out string? driverVersion))
+								{
+									string prefix = $"GPU-{++index}";
+									primaryDevice.Properties.Add($"{prefix}-Name={name}");
+									primaryDevice.Properties.Add($"{prefix}-DriverVersion={driverVersion}");
+								}
+							}
+						}
+					}
+				}
+
+				// Add EC2 properties if needed
+				if (_settings.EnableAwsEc2Support)
+				{
+					await AddAwsProperties(primaryDevice.Properties, logger);
+				}
+
+				// Add session information
+				primaryDevice.Properties.Add($"User={Environment.UserName}");
+				primaryDevice.Properties.Add($"Domain={Environment.UserDomainName}");
+				primaryDevice.Properties.Add($"Interactive={Environment.UserInteractive}");
+				primaryDevice.Properties.Add($"Elevated={BuildGraphExecutor.IsUserAdministrator()}");
+			}
+			else if (RuntimeInformation.IsOSPlatform(OSPlatform.Linux))
+			{
+				primaryDevice.Properties.Add("Platform=Linux");
+				primaryDevice.Properties.Add("PlatformGroup=Linux");
+				primaryDevice.Properties.Add("PlatformGroup=Unix");
+				primaryDevice.Properties.Add("PlatformGroup=Desktop");
+
+				primaryDevice.Properties.Add("OSFamily=Linux");
+				primaryDevice.Properties.Add("OSVersion=Linux");
+
+				// Add EC2 properties if needed
+				if (_settings.EnableAwsEc2Support)
+				{
+					await AddAwsProperties(primaryDevice.Properties, logger);
+				}
+
+				// Parse the CPU info
+				List<Dictionary<string, string>>? cpuRecords = await ReadLinuxHwPropsAsync("/proc/cpuinfo", logger);
+				if (cpuRecords != null)
+				{
+					Dictionary<string, string> cpuNames = new Dictionary<string, string>(StringComparer.Ordinal);
+					foreach (Dictionary<string, string> cpuRecord in cpuRecords)
+					{
+						if (cpuRecord.TryGetValue("physical id", out string? physicalId) && cpuRecord.TryGetValue("model name", out string? modelName))
+						{
+							cpuNames[physicalId] = modelName;
+						}
+					}
+
+					Dictionary<string, int> nameToCount = new Dictionary<string, int>(StringComparer.Ordinal);
+					foreach (string cpuName in cpuNames.Values)
+					{
+						nameToCount.TryGetValue(cpuName, out int count);
+						nameToCount[cpuName] = count + 1;
+					}
+
+					HashSet<string> logicalCores = new HashSet<string>();
+					HashSet<string> physicalCores = new HashSet<string>();
+					foreach (Dictionary<string, string> cpuRecord in cpuRecords)
+					{
+						if (cpuRecord.TryGetValue("processor", out string? logicalCoreId))
+						{
+							logicalCores.Add(logicalCoreId);
+						}
+						if (cpuRecord.TryGetValue("core id", out string? physicalCoreId))
+						{
+							physicalCores.Add(physicalCoreId);
+						}
+					}
+
+					AddCpuInfo(primaryDevice, nameToCount, logicalCores.Count, physicalCores.Count);
+				}
+
+				// Parse the RAM info
+				List<Dictionary<string, string>>? memRecords = await ReadLinuxHwPropsAsync("/proc/meminfo", logger);
+				if (memRecords != null && memRecords.Count > 0 && memRecords[0].TryGetValue("MemTotal", out string? memTotal))
+				{
+					Match match = Regex.Match(memTotal, @"(\d+)\s+kB");
+					if (match.Success)
+					{
+						long totalCapacity = Int64.Parse(match.Groups[1].Value) * 1024;
+						primaryDevice.Properties.Add($"RAM={totalCapacity / (1024 * 1024 * 1024)}");
+					}
+				}
+
+				// Add session information
+				primaryDevice.Properties.Add($"User={Environment.UserName}");
+				primaryDevice.Properties.Add($"Domain={Environment.UserDomainName}");
+				primaryDevice.Properties.Add($"Interactive={Environment.UserInteractive}");
+				primaryDevice.Properties.Add($"Elevated={BuildGraphExecutor.IsUserAdministrator()}");
+			}
+			else if (RuntimeInformation.IsOSPlatform(OSPlatform.OSX))
+			{
+				primaryDevice.Properties.Add("Platform=Mac");
+				primaryDevice.Properties.Add("PlatformGroup=Apple");
+				primaryDevice.Properties.Add("PlatformGroup=Desktop");
+
+				primaryDevice.Properties.Add("OSFamily=MacOS");
+				primaryDevice.Properties.Add("OSVersion=MacOS");
+
+				string output;
+				using (Process process = new Process())
+				{
+					process.StartInfo.FileName = "system_profiler";
+					process.StartInfo.Arguments = "SPHardwareDataType SPSoftwareDataType -xml";
+					process.StartInfo.CreateNoWindow = true;
+					process.StartInfo.RedirectStandardOutput = true;
+					process.StartInfo.RedirectStandardInput = false;
+					process.StartInfo.UseShellExecute = false;
+					process.Start();
+
+					output = await process.StandardOutput.ReadToEndAsync();
+				}
+
+				XmlDocument xml = new XmlDocument();
+				xml.LoadXml(output);
+
+				XmlNode? hardwareNode = xml.SelectSingleNode("/plist/array/dict/key[. = '_dataType']/following-sibling::string[. = 'SPHardwareDataType']/../key[. = '_items']/following-sibling::array/dict");
+				if (hardwareNode != null)
+				{
+					XmlNode? model = hardwareNode.SelectSingleNode("key[. = 'machine_model']/following-sibling::string");
+					if (model != null)
+					{
+						primaryDevice.Properties.Add($"Model={model.InnerText}");
+					}
+
+					XmlNode? cpuTypeNode = hardwareNode.SelectSingleNode("key[. = 'cpu_type']/following-sibling::string");
+					XmlNode? cpuSpeedNode = hardwareNode.SelectSingleNode("key[. = 'current_processor_speed']/following-sibling::string");
+					XmlNode? cpuPackagesNode = hardwareNode.SelectSingleNode("key[. = 'packages']/following-sibling::integer");
+					if (cpuTypeNode != null && cpuSpeedNode != null && cpuPackagesNode != null)
+					{
+						primaryDevice.Properties.Add((cpuPackagesNode.InnerText != "1") ? $"CPU={cpuPackagesNode.InnerText} x {cpuTypeNode.InnerText} @ {cpuSpeedNode.InnerText}" : $"CPU={cpuTypeNode.InnerText} @ {cpuSpeedNode.InnerText}");
+					}
+
+					primaryDevice.Properties.Add($"LogicalCores={Environment.ProcessorCount}");
+
+					XmlNode? cpuCountNode = hardwareNode.SelectSingleNode("key[. = 'number_processors']/following-sibling::integer");
+					if (cpuCountNode != null)
+					{
+						primaryDevice.Properties.Add($"PhysicalCores={cpuCountNode.InnerText}");
+					}
+
+					XmlNode? memoryNode = hardwareNode.SelectSingleNode("key[. = 'physical_memory']/following-sibling::string");
+					if (memoryNode != null)
+					{
+						string[] parts = memoryNode.InnerText.Split(new char[] { ' ' }, StringSplitOptions.RemoveEmptyEntries);
+						if (parts.Length == 2 && parts[1] == "GB")
+						{
+							primaryDevice.Properties.Add($"RAM={parts[0]}");
+						}
+					}
+				}
+
+				XmlNode? softwareNode = xml.SelectSingleNode("/plist/array/dict/key[. = '_dataType']/following-sibling::string[. = 'SPSoftwareDataType']/../key[. = '_items']/following-sibling::array/dict");
+				if (softwareNode != null)
+				{
+					XmlNode? osVersionNode = softwareNode.SelectSingleNode("key[. = 'os_version']/following-sibling::string");
+					if (osVersionNode != null)
+					{
+						primaryDevice.Properties.Add($"OSDistribution={osVersionNode.InnerText}");
+					}
+
+					XmlNode? kernelVersionNode = softwareNode.SelectSingleNode("key[. = 'kernel_version']/following-sibling::string");
+					if (kernelVersionNode != null)
+					{
+						primaryDevice.Properties.Add($"OSKernelVersion={kernelVersionNode.InnerText}");
+					}
+				}
+			}
+
+			// Get the IP addresses
+			try
+			{
+				using CancellationTokenSource dnsCts = new(3000);
+				IPHostEntry entry = await Dns.GetHostEntryAsync(Dns.GetHostName(), dnsCts.Token);
+				foreach (IPAddress address in entry.AddressList)
+				{
+					if (address.AddressFamily == System.Net.Sockets.AddressFamily.InterNetwork)
+					{
+						primaryDevice.Properties.Add($"Ipv4={address}");
+					}
+					else if (address.AddressFamily == System.Net.Sockets.AddressFamily.InterNetworkV6)
+					{
+						primaryDevice.Properties.Add($"Ipv6={address}");
+					}
+				}
+			}
+			catch (Exception ex)
+			{
+				logger.LogDebug(ex, "Unable to get local IP address");
+			}
+
+			// Get the time that the machine booted
+			primaryDevice.Properties.Add($"BootTime={s_bootTime}");
+			primaryDevice.Properties.Add($"StartTime={s_startTime}");
+
+			// Add information about the current session
+			if (workingDir != null)
+			{
+				// Add disk info based on platform
+				string? driveName;
+				if (RuntimeInformation.IsOSPlatform(OSPlatform.Windows))
+				{
+					driveName = Path.GetPathRoot(workingDir.FullName);
+				}
+				else
+				{
+					driveName = workingDir.FullName;
+				}
+
+				if (driveName != null)
+				{
+					try
+					{
+						DriveInfo info = new DriveInfo(driveName);
+						primaryDevice.Properties.Add($"DiskFreeSpace={info.AvailableFreeSpace}");
+						primaryDevice.Properties.Add($"DiskTotalSize={info.TotalSize}");
+					}
+					catch (Exception ex)
+					{
+						logger.LogWarning(ex, "Unable to query disk info for path '{DriveName}'", driveName);
+					}
+				}
+				primaryDevice.Properties.Add($"WorkingDir={workingDir}");
+			}
+
+			// Add any horde. env vars for custom properties.
+			IEnumerable<string> envVars = Environment.GetEnvironmentVariables().Keys.Cast<string>();
+			foreach (string envVar in envVars.Where(x => x.StartsWith("horde.", StringComparison.InvariantCultureIgnoreCase)))
+			{
+				primaryDevice.Properties.Add($"{envVar}={Environment.GetEnvironmentVariable(envVar)}");
+			}
+
+			// Create the worker
+			AgentCapabilities agent = new();
+			agent.Devices.Add(primaryDevice);
+			agent.Devices.AddRange(otherDevices);
+
+			// Add any additional properties from the config file
+			agent.Properties.AddRange(_settings.Properties.Select(kvp => $"{kvp.Key}={kvp.Value}"));
+			return agent;
+		}
+
+		static void AddCpuInfo(DeviceCapabilities primaryDevice, Dictionary<string, int> nameToCount, int numLogicalCores, int numPhysicalCores)
+		{
+			if (nameToCount.Count > 0)
+			{
+				primaryDevice.Properties.Add("CPU=" + String.Join(", ", nameToCount.Select(x => (x.Value > 1) ? $"{x.Key} x {x.Value}" : x.Key)));
+			}
+
+			if (numLogicalCores > 0)
+			{
+				primaryDevice.Properties.Add($"LogicalCores={numLogicalCores}");
+			}
+
+			if (numPhysicalCores > 0)
+			{
+				primaryDevice.Properties.Add($"PhysicalCores={numPhysicalCores}");
+			}
+		}
+
+		static async Task<List<Dictionary<string, string>>?> ReadLinuxHwPropsAsync(string fileName, ILogger logger)
+		{
+			List<Dictionary<string, string>>? records = null;
+			if (File.Exists(fileName))
+			{
+				records = new List<Dictionary<string, string>>();
+				using (StreamReader reader = new StreamReader(fileName))
+				{
+					Dictionary<string, string> record = new Dictionary<string, string>(StringComparer.Ordinal);
+
+					string? line;
+					while ((line = await reader.ReadLineAsync()) != null)
+					{
+						int idx = line.IndexOf(':', StringComparison.Ordinal);
+						if (idx == -1)
+						{
+							if (record.Count > 0)
+							{
+								records.Add(record);
+								record = new Dictionary<string, string>(StringComparer.Ordinal);
+							}
+						}
+						else
+						{
+							string key = line.Substring(0, idx).Trim();
+							string value = line.Substring(idx + 1).Trim();
+
+							if (record.TryGetValue(key, out string? prevValue))
+							{
+								logger.LogWarning("Multiple entries for {Key} in {File} (was '{Prev}', now '{Next}')", key, fileName, prevValue, value);
+							}
+							else
+							{
+								record.Add(key, value);
+							}
+						}
+					}
+
+					if (record.Count > 0)
+					{
+						records.Add(record);
+					}
+				}
+			}
+			return records;
+		}
+
+		[SupportedOSPlatform("windows")]
+		class WmiProperties
+		{
+			readonly Dictionary<string, object> _properties = new Dictionary<string, object>(StringComparer.Ordinal);
+
+			public WmiProperties(ManagementObject row)
+			{
+				foreach (PropertyData property in row.Properties)
+				{
+					_properties[property.Name] = property.Value;
+				}
+			}
+
+			public bool TryGetValue(string name, [NotNullWhen(true)] out string? value)
+			{
+				object? @object;
+				if (_properties.TryGetValue(name, out @object))
+				{
+					value = @object.ToString();
+					return value != null;
+				}
+				else
+				{
+					value = null!;
+					return false;
+				}
+			}
+
+			public bool TryGetValue(string name, out long value)
+			{
+				object? objectValue;
+				if (_properties.TryGetValue(name, out objectValue))
+				{
+					if (objectValue is int intValue)
+					{
+						value = intValue;
+						return true;
+					}
+					else if (objectValue is uint uintValue)
+					{
+						value = uintValue;
+						return true;
+					}
+					else if (objectValue is long longValue)
+					{
+						value = longValue;
+						return true;
+					}
+					else if (objectValue is ulong ulongValue)
+					{
+						value = (long)ulongValue;
+						return true;
+					}
+				}
+
+				value = 0;
+				return false;
+			}
+		}
+
+		[SupportedOSPlatform("windows")]
+		static Dictionary<string, object> GetWmiProperties(ManagementObject row)
+		{
+			Dictionary<string, object> properties = new Dictionary<string, object>(StringComparer.Ordinal);
+			foreach (PropertyData property in row.Properties)
+			{
+				properties[property.Name] = property.Value;
+			}
+			return properties;
+		}
+
+		static async Task AddAwsProperties(IList<string> properties, ILogger logger)
+		{
+			if (EC2InstanceMetadata.IdentityDocument != null)
+			{
+				properties.Add("EC2=1");
+				AddAwsProperty("aws-instance-id", "/instance-id", properties);
+				AddAwsProperty("aws-instance-type", "/instance-type", properties);
+				AddAwsProperty("aws-region", "/region", properties);
+
+				try
+				{
+					using (AmazonEC2Client client = new AmazonEC2Client())
+					{
+						DescribeTagsRequest request = new DescribeTagsRequest();
+						request.Filters = new List<Filter>();
+						request.Filters.Add(new Filter("resource-id", new List<string> { EC2InstanceMetadata.InstanceId }));
+
+						DescribeTagsResponse response = await client.DescribeTagsAsync(request);
+						foreach (TagDescription tag in response.Tags)
+						{
+							properties.Add($"aws-tag={tag.Key}:{tag.Value}");
+						}
+					}
+				}
+				catch (Exception ex)
+				{
+					logger.LogDebug(ex, "Unable to query EC2 tags.");
+				}
+			}
+		}
+
+		static void AddAwsProperty(string name, string awsKey, IList<string> properties)
+		{
+			string? value = EC2InstanceMetadata.GetData(awsKey);
+			if (value != null)
+			{
+				properties.Add($"{name}={value}");
+			}
+		}
+	}
+}
diff --git a/Engine/Source/Programs/Horde/Horde.Agent/Services/SessionFactoryService.cs b/Engine/Source/Programs/Horde/Horde.Agent/Services/SessionFactoryService.cs
new file mode 100644
--- /dev/null
+++ b/Engine/Source/Programs/Horde/Horde.Agent/Services/SessionFactoryService.cs
@@ -0,0 +1,279 @@
+// Copyright Epic Games, Inc. All Rights Reserved.
+
+using System;
+using System.Collections.Generic;
+using System.Linq;
+using System.Reflection;
+using System.Runtime.InteropServices;
+using System.Threading;
+using System.Threading.Tasks;
+using EpicGames.Core;
+using Grpc.Net.Client;
+using Horde.Agent.Utility;
+using HordeCommon;
+using HordeCommon.Rpc;
+using HordeCommon.Rpc.Messages;
+using Microsoft.Extensions.Logging;
+using Microsoft.Extensions.Options;
+
+namespace Horde.Agent.Services
+{
+	/// <summary>
+	/// Interface for about the current session. 
+	/// </summary>
+	interface ISession : IAsyncDisposable
+	{
+		/// <summary>
+		/// URL of the server
+		/// </summary>
+		Uri ServerUrl { get; }
+
+		/// <summary>
+		/// The agent identifier
+		/// </summary>
+		string AgentId { get; }
+
+		/// <summary>
+		/// Identifier for the current session
+		/// </summary>
+		string SessionId { get; }
+
+		/// <summary>
+		/// Token to use for connection to the server
+		/// </summary>
+		string Token { get; }
+
+		/// <summary>
+		/// Connection to the server
+		/// </summary>
+		IRpcConnection RpcConnection { get; }
+
+		/// <summary>
+		/// Working directory for sandboxes etc..
+		/// </summary>
+		DirectoryReference WorkingDir { get; }
+
+		/// <summary>
+		/// Terminate all processes in running in the working directory
+		/// </summary>
+		/// <param name="logger">Logger for any diagnostic messages</param>
+		/// <param name="cancellationToken">Cancellation token for the operation</param>
+		Task TerminateProcessesAsync(ILogger logger, CancellationToken cancellationToken);
+	}
+
+	/// <summary>
+	/// Interface for a factory to create sessions
+	/// </summary>
+	interface ISessionFactoryService
+	{
+		/// <summary>
+		/// Creates a new session
+		/// </summary>
+		/// <param name="cancellationToken">Cancellation token for the operation</param>
+		/// <returns>New session information</returns>
+		public Task<ISession> CreateAsync(CancellationToken cancellationToken);
+	}
+
+	/// <summary>
+	/// Information about the current session. 
+	/// </summary>
+	sealed class Session : ISession
+	{
+		/// <summary>
+		/// URL of the server
+		/// </summary>
+		public Uri ServerUrl { get; }
+
+		/// <summary>
+		/// The agent identifier
+		/// </summary>
+		public string AgentId { get; }
+
+		/// <summary>
+		/// Identifier for the current session
+		/// </summary>
+		public string SessionId { get; }
+
+		/// <summary>
+		/// Token to use for connection to the server
+		/// </summary>
+		public string Token { get; }
+
+		/// <summary>
+		/// Connection to the server
+		/// </summary>
+		public IRpcConnection RpcConnection { get; }
+
+		/// <summary>
+		/// Working directory for sandboxes etc..
+		/// </summary>
+		public DirectoryReference WorkingDir { get; }
+
+		readonly HashSet<string> _processNamesToTerminate;
+		readonly ILogger _logger;
+
+		/// <summary>
+		/// Constructor
+		/// </summary>
+		public Session(Uri serverUrl, string agentId, string sessionId, string token, IRpcConnection rpcConnection, DirectoryReference workingDir, IEnumerable<string> processNamesToTerminate, ILogger logger)
+		{
+			ServerUrl = serverUrl;
+			AgentId = agentId;
+			SessionId = sessionId;
+			Token = token;
+			RpcConnection = rpcConnection;
+			WorkingDir = workingDir;
+
+			_processNamesToTerminate = new HashSet<string>(processNamesToTerminate, StringComparer.OrdinalIgnoreCase);
+			_logger = logger;
+		}
+
+		/// <summary>
+		/// Creates a new agent session
+		/// </summary>
+		/// <param name="capabilitiesService"></param>
+		/// <param name="grpcService"></param>
+		/// <param name="settings"></param>
+		/// <param name="logger"></param>
+		/// <param name="cancellationToken">Indicates that the service is trying to stop</param>
+		/// <returns>Async task</returns>
+		public static async Task<Session> CreateAsync(CapabilitiesService capabilitiesService, GrpcService grpcService, IOptions<AgentSettings> settings, ILogger logger, CancellationToken cancellationToken)
+		{
+			AgentSettings currentSettings = settings.Value;
+
+			// Get the working directory
+			if (currentSettings.WorkingDir == null)
+			{
+				throw new Exception("WorkingDir is not set. Unable to run service.");
+			}
+
+			DirectoryReference baseDir = new FileReference(Assembly.GetExecutingAssembly().Location).Directory;
+			DirectoryReference workingDir = DirectoryReference.Combine(baseDir, currentSettings.WorkingDir);
+			logger.LogInformation("WorkingDir: {WorkingDir}", workingDir);
+			DirectoryReference.CreateDirectory(workingDir);
+
+			// Print the server info
+			ServerProfile serverProfile = currentSettings.GetCurrentServerProfile();
+			logger.LogInformation("Server: {Server}", serverProfile.Url);
+
+			// Show the worker capabilities
+			AgentCapabilities capabilities = await capabilitiesService.GetCapabilitiesAsync(workingDir);
+			if (capabilities.Properties.Count > 0)
+			{
+				logger.LogInformation("Global:");
+				foreach (string property in capabilities.Properties)
+				{
+					logger.LogInformation("  {AgentProperty}", property);
+				}
+			}
+			foreach (DeviceCapabilities device in capabilities.Devices)
+			{
+				logger.LogInformation("{DeviceName} Device:", device.Handle);
+				foreach (string property in device.Properties)
+				{
+					logger.LogInformation("   {DeviceProperty}", property);
+				}
+			}
+
+			// Mount all the necessary network shares. Currently only supported on Windows.
+			if (currentSettings.ShareMountingEnabled && RuntimeInformation.IsOSPlatform(OSPlatform.Windows))
+			{
+				foreach (MountNetworkShare share in currentSettings.Shares)
+				{
+					if (share.MountPoint != null && share.RemotePath != null)
+					{
+						logger.LogInformation("Mounting {RemotePath} as {MountPoint}", share.RemotePath, share.MountPoint);
+						NetworkShare.Mount(share.MountPoint, share.RemotePath);
+					}
+				}
+			}
+
+			// Create the session
+			CreateSessionResponse createSessionResponse;
+			using (GrpcChannel channel = grpcService.CreateGrpcChannel(serverProfile.Token))
+			{
+				HordeRpc.HordeRpcClient rpcClient = new HordeRpc.HordeRpcClient(channel);
+
+				// Create the session information
+				CreateSessionRequest sessionRequest = new CreateSessionRequest();
+				sessionRequest.Name = currentSettings.GetAgentName();
+				sessionRequest.Status = AgentStatus.Ok;
+				sessionRequest.Capabilities = capabilities;
+				sessionRequest.Version = Program.Version;
+
+				// Create a session
+				createSessionResponse = await rpcClient.CreateSessionAsync(sessionRequest, null, null, cancellationToken);
+				logger.LogInformation("Created session. AgentName={AgentName} SessionId={SessionId}", currentSettings.GetAgentName(), createSessionResponse.SessionId);
+			}
+
+			Func<GrpcChannel> createGrpcChannel = () => grpcService.CreateGrpcChannel(createSessionResponse.Token);
+
+			// Open a connection to the server
+#pragma warning disable CA2000 // False positive; ownership is transferred to new Session object.
+			IRpcConnection rpcConnection = Utility.RpcConnection.Create(createGrpcChannel, channel => new HordeRpc.HordeRpcClient(channel), logger);
+			return new Session(serverProfile.Url, createSessionResponse.AgentId, createSessionResponse.SessionId, createSessionResponse.Token, rpcConnection, workingDir, currentSettings.ProcessNamesToTerminate, logger);
+#pragma warning restore CA2000
+		}
+
+		/// <summary>
+		/// Dispose of the current session
+		/// </summary>
+		/// <returns></returns>
+		public async ValueTask DisposeAsync()
+		{
+			await RpcConnection.DisposeAsync();
+		}
+
+		/// <inheritdoc/>
+		public Task TerminateProcessesAsync(ILogger logger, CancellationToken cancellationToken)
+		{
+			// Terminate child processes from any previous runs
+			ProcessUtils.TerminateProcesses(ShouldTerminateProcess, logger, cancellationToken);
+			return Task.CompletedTask;
+		}
+
+		/// <summary>
+		/// Callback for determining whether a process should be terminated
+		/// </summary>
+		bool ShouldTerminateProcess(FileReference imageFile)
+		{
+			if (imageFile.IsUnderDirectory(WorkingDir))
+			{
+				return true;
+			}
+
+			string fileName = imageFile.GetFileName();
+			if (_processNamesToTerminate.Contains(fileName))
+			{
+				return true;
+			}
+
+			return false;
+		}
+	}
+
+	/// <summary>
+	/// Creates session objects
+	/// </summary>
+	class SessionFactoryService : ISessionFactoryService
+	{
+		readonly CapabilitiesService _capabilitiesService;
+		readonly GrpcService _grpcService;
+		readonly IOptions<AgentSettings> _settings;
+		readonly ILogger _logger;
+
+		/// <summary>
+		/// Constructor
+		/// </summary>
+		public SessionFactoryService(CapabilitiesService capabilitiesService, GrpcService grpcService, IOptions<AgentSettings> settings, ILogger<SessionFactoryService> logger)
+		{
+			_capabilitiesService = capabilitiesService;
+			_grpcService = grpcService;
+			_settings = settings;
+			_logger = logger;
+		}
+
+		/// <inheritdoc/>
+		public async Task<ISession> CreateAsync(CancellationToken cancellationToken) => await Session.CreateAsync(_capabilitiesService, _grpcService, _settings, _logger, cancellationToken);
+	}
+}
