Horde: Allow overriding the server profile globally via the command line. Should fix deployment issues.

diff --git a/Engine/Source/Programs/Horde/Horde.Agent/Commands/Service/RunCommand.cs b/Engine/Source/Programs/Horde/Horde.Agent/Commands/Service/RunCommand.cs
--- a/Engine/Source/Programs/Horde/Horde.Agent/Commands/Service/RunCommand.cs
+++ b/Engine/Source/Programs/Horde/Horde.Agent/Commands/Service/RunCommand.cs
@@ -1,34 +1,16 @@
 // Copyright Epic Games, Inc. All Rights Reserved.
 
 using System;
-using System.Collections.Generic;
-using System.Net.Http;
 using System.Threading.Tasks;
-using Datadog.Trace;
-using Datadog.Trace.Configuration;
-using Datadog.Trace.OpenTracing;
 using EpicGames.Core;
-using EpicGames.Horde.Storage;
-using Horde.Agent.Execution;
-using Horde.Agent.Execution.Interfaces;
-using Horde.Agent.Leases;
-using Horde.Agent.Leases.Handlers;
-using Horde.Agent.Services;
-using Horde.Agent.Utility;
 using Microsoft.Extensions.Configuration;
 using Microsoft.Extensions.DependencyInjection;
-using Microsoft.Extensions.DependencyInjection.Extensions;
 using Microsoft.Extensions.Hosting;
 using Microsoft.Extensions.Logging;
-using Microsoft.Extensions.Options;
-using OpenTracing;
-using OpenTracing.Util;
-using Polly;
 using Serilog.Events;
 
 namespace Horde.Agent.Modes.Service
 {
-	using ITracer = OpenTracing.ITracer;
 
 	/// <summary>
 	/// 
@@ -36,18 +18,6 @@ namespace Horde.Agent.Modes.Service
 	[Command("Service", "Run", "Runs the service in listen mode")]
 	class RunCommand : Command
 	{
-		/// <summary>
-		/// Override for the server to use
-		/// </summary>
-		[CommandLine("-Server=")]
-		string? Server { get; set; } = null;
-
-		/// <summary>
-		/// Override the working directory
-		/// </summary>
-		[CommandLine("-WorkingDir=")]
-		string? WorkingDir { get; set; } = null;
-
 		/// <summary>
 		/// Log verbosity level (use normal Serilog levels such as debug, warning or info)
 		/// </summary>
@@ -96,18 +66,6 @@ namespace Horde.Agent.Modes.Service
 				.ConfigureAppConfiguration(builder =>
 				{
 					builder.AddConfiguration(_defaultServices.Configuration);
-
-					Dictionary<string, string> overrides = new Dictionary<string, string>();
-					if (Server != null)
-					{
-						overrides.Add($"{AgentSettings.SectionName}:{nameof(AgentSettings.Server)}", Server);
-					}
-					if (WorkingDir != null)
-					{
-						overrides.Add($"{AgentSettings.SectionName}:{nameof(AgentSettings.WorkingDir)}", WorkingDir);
-					}
-
-					builder.AddInMemoryCollection(overrides);
 				})
 				.ConfigureLogging(builder =>
 				{
diff --git a/Engine/Source/Programs/Horde/Horde.Agent/Program.cs b/Engine/Source/Programs/Horde/Horde.Agent/Program.cs
--- a/Engine/Source/Programs/Horde/Horde.Agent/Program.cs
+++ b/Engine/Source/Programs/Horde/Horde.Agent/Program.cs
@@ -96,12 +96,31 @@ namespace Horde.Agent
 
 			using Logging.HordeLoggerProvider loggerProvider = new Logging.HordeLoggerProvider();
 
+			string? environment = Environment.GetEnvironmentVariable("DOTNET_ENVIRONMENT");
+			if (String.IsNullOrEmpty(environment))
+			{
+				environment = "Production";
+			}
+
+			CommandLineArguments arguments = new CommandLineArguments(args);
+
+			Dictionary<string, string> configOverrides = new Dictionary<string, string>();
+			if (arguments.TryGetValue("-Server=", out string? serverOverride))
+			{
+				configOverrides.Add($"{AgentSettings.SectionName}:{nameof(AgentSettings.Server)}", serverOverride);
+			}
+			if (arguments.TryGetValue("-WorkingDir=", out string? workingDirOverride))
+			{
+				configOverrides.Add($"{AgentSettings.SectionName}:{nameof(AgentSettings.WorkingDir)}", workingDirOverride);
+			}
+
 			IConfiguration config = new ConfigurationBuilder()
 				.SetBasePath(AppDir.FullName)
 				.AddJsonFile("appsettings.json", optional: false)
 				.AddJsonFile("appsettings.Build.json", optional: true) // specific settings for builds (installer/dockerfile)
-				.AddJsonFile($"appsettings.{Environment.GetEnvironmentVariable("DOTNET_ENVIRONMENT")}.json", optional: true) // environment variable overrides, also used in k8s setups with Helm
+				.AddJsonFile($"appsettings.{environment}.json", optional: true) // environment variable overrides, also used in k8s setups with Helm
 				.AddJsonFile("appsettings.User.json", optional: true)
+				.AddInMemoryCollection(configOverrides)
 				.AddEnvironmentVariables()
 				.Build();
 
