Horde: Fix issue with device automation pivot detail list dancing scrollbar
#rnx
#skipci

diff --git a/Engine/Source/Programs/Horde/HordeDashboard/src/components/DeviceView.tsx b/Engine/Source/Programs/Horde/HordeDashboard/src/components/DeviceView.tsx
--- a/Engine/Source/Programs/Horde/HordeDashboard/src/components/DeviceView.tsx
+++ b/Engine/Source/Programs/Horde/HordeDashboard/src/components/DeviceView.tsx
@@ -546,7 +546,7 @@ const DevicePanel: React.FC = observer(() => {
 
          <Stack tokens={{ childrenGap: 12 }}>
             <FocusZone direction={FocusZoneDirection.vertical}>
-               <div className={customStyles.details} style={{ height: "calc(100vh - 270px)", position: 'relative' }} data-is-scrollable>
+               <div className={customStyles.details} style={{ height: "calc(100vh - 280px)", position: 'relative' }} data-is-scrollable>
                   <ScrollablePane scrollbarVisibility={ScrollbarVisibility.always} onScroll={() => { }}>
                      <DetailsList
                         styles={{ root: { overflowX: "hidden", width: 1676 } }}
@@ -556,6 +556,7 @@ const DevicePanel: React.FC = observer(() => {
                         selectionMode={SelectionMode.none}
                         layoutMode={DetailsListLayoutMode.justified}
                         compact={true}
+                        onShouldVirtualize={() => false}
                         onRenderRow={onRenderRow}
                      />
                   </ScrollablePane>
