Horde: Mark unawaited transaction commands as fire-and-forget

Prevents task cancellation exceptions from being thrown.

diff --git a/Engine/Source/Programs/Horde/Horde.Build/Logs/Builder/RedisLogBuilder.cs b/Engine/Source/Programs/Horde/Horde.Build/Logs/Builder/RedisLogBuilder.cs
--- a/Engine/Source/Programs/Horde/Horde.Build/Logs/Builder/RedisLogBuilder.cs
+++ b/Engine/Source/Programs/Horde/Horde.Build/Logs/Builder/RedisLogBuilder.cs
@@ -91,13 +91,13 @@ namespace Horde.Build.Logs.Builder
 				scope.Span.SetTag("Offset", chunkOffset.ToString(CultureInfo.InvariantCulture));
 				scope.Span.SetTag("WriteOffset", writeOffset.ToString(CultureInfo.InvariantCulture));
 
-				ITransaction? createTransaction = redisDb.CreateTransaction();
+				ITransaction createTransaction = redisDb.CreateTransaction();
 				createTransaction.AddCondition(Condition.SortedSetNotContains(ItemsKey, keys.Prefix));
-				_ = createTransaction.SortedSetAddAsync(ItemsKey, keys.Prefix, DateTime.UtcNow.Ticks);
-				_ = createTransaction.StringSetAsync(keys.Type, (int)type);
-				_ = createTransaction.StringSetAsync(keys.Length, data.Length);
-				_ = createTransaction.StringSetAsync(keys.LineIndex, writeLineIndex + writeLineCount);
-				_ = createTransaction.StringSetAsync(keys.SubChunkData, data);
+				_ = createTransaction.SortedSetAddAsync(ItemsKey, keys.Prefix, DateTime.UtcNow.Ticks, flags: CommandFlags.FireAndForget);
+				_ = createTransaction.StringSetAsync(keys.Type, (int)type, flags: CommandFlags.FireAndForget);
+				_ = createTransaction.StringSetAsync(keys.Length, data.Length, flags: CommandFlags.FireAndForget);
+				_ = createTransaction.StringSetAsync(keys.LineIndex, writeLineIndex + writeLineCount, flags: CommandFlags.FireAndForget);
+				_ = createTransaction.StringSetAsync(keys.SubChunkData, data, flags: CommandFlags.FireAndForget);
 				if (await createTransaction.ExecuteAsync())
 				{
 					_logger.LogTrace("Created {Size} bytes in {Key}", data.Length, keys.SubChunkData);
@@ -111,14 +111,14 @@ namespace Horde.Build.Logs.Builder
 				scope.Span.SetTag("Offset", chunkOffset.ToString(CultureInfo.InvariantCulture));
 				scope.Span.SetTag("WriteOffset", writeOffset.ToString(CultureInfo.InvariantCulture));
 
-				ITransaction? appendTransaction = redisDb.CreateTransaction();
+				ITransaction appendTransaction = redisDb.CreateTransaction();
 				appendTransaction.AddCondition(Condition.SortedSetContains(ItemsKey, keys.Prefix));
 				appendTransaction.AddCondition(Condition.KeyNotExists(keys.Complete));
 				appendTransaction.AddCondition(Condition.StringEqual(keys.Type, (int)type));
 				appendTransaction.AddCondition(Condition.StringEqual(keys.Length, (int)(writeOffset - chunkOffset)));
-				_ = appendTransaction.StringAppendAsync(keys.SubChunkData, data);
-				_ = appendTransaction.StringSetAsync(keys.Length, (int)(writeOffset - chunkOffset) + data.Length);
-				_ = appendTransaction.StringSetAsync(keys.LineIndex, writeLineIndex + writeLineCount);
+				_ = appendTransaction.StringAppendAsync(keys.SubChunkData, data, flags: CommandFlags.FireAndForget);
+				_ = appendTransaction.StringSetAsync(keys.Length, (int)(writeOffset - chunkOffset) + data.Length, flags: CommandFlags.FireAndForget);
+				_ = appendTransaction.StringSetAsync(keys.LineIndex, writeLineIndex + writeLineCount, flags: CommandFlags.FireAndForget);
 				if (await appendTransaction.ExecuteAsync())
 				{
 					_logger.LogTrace("Appended {Size} to {Key}", data.Length, keys.SubChunkData);
@@ -174,13 +174,14 @@ namespace Horde.Build.Logs.Builder
 				writeTransaction.AddCondition(Condition.StringEqual(keys.Length, length));
 				writeTransaction.AddCondition(Condition.StringEqual(keys.LineIndex, lineIndex));
 				Task<long> newLength = writeTransaction.StringAppendAsync(keys.ChunkData, subChunkDataBytes);
-				_ = writeTransaction.KeyDeleteAsync(keys.SubChunkData);
+				_ = writeTransaction.KeyDeleteAsync(keys.SubChunkData, flags: CommandFlags.FireAndForget);
 
 				if (await writeTransaction.ExecuteAsync())
 				{
 					_logger.LogDebug("Completed sub-chunk for log {LogId} chunk offset {Offset} -> sub-chunk size {SubChunkSize}, chunk size {ChunkSize}", logId, offset, subChunkDataBytes.Length, await newLength);
 					return;
 				}
+				await writeTransaction.WaitAndIgnoreCancellations(newLength);
 
 				// Cool down before retrying
 				await Task.Delay(100);
@@ -202,7 +203,7 @@ namespace Horde.Build.Logs.Builder
 
 			ITransaction transaction = redisDb.CreateTransaction();
 			transaction.AddCondition(Condition.SortedSetContains(ItemsKey, keys.Prefix));
-			_ = transaction.StringSetAsync(keys.Complete, true);
+			_ = transaction.StringSetAsync(keys.Complete, true, flags: CommandFlags.FireAndForget);
 			if(!await transaction.ExecuteAsync())
 			{
 				_logger.LogDebug("Log {LogId} chunk offset {Offset} is not in Redis builder", logId, offset);
@@ -223,14 +224,14 @@ namespace Horde.Build.Logs.Builder
 			ChunkKeys keys = new ChunkKeys(logId, offset);
 
 			ITransaction transaction = redisDb.CreateTransaction();
-			_ = transaction.KeyDeleteAsync(keys.Type);
-			_ = transaction.KeyDeleteAsync(keys.LineIndex);
-			_ = transaction.KeyDeleteAsync(keys.Length);
-			_ = transaction.KeyDeleteAsync(keys.ChunkData);
-			_ = transaction.KeyDeleteAsync(keys.SubChunkData);
-			_ = transaction.KeyDeleteAsync(keys.Complete);
-
-			_ = transaction.SortedSetRemoveAsync(ItemsKey, keys.Prefix);
+			_ = transaction.KeyDeleteAsync(keys.Type, flags: CommandFlags.FireAndForget);
+			_ = transaction.KeyDeleteAsync(keys.LineIndex, flags: CommandFlags.FireAndForget);
+			_ = transaction.KeyDeleteAsync(keys.Length, flags: CommandFlags.FireAndForget);
+			_ = transaction.KeyDeleteAsync(keys.ChunkData, flags: CommandFlags.FireAndForget);
+			_ = transaction.KeyDeleteAsync(keys.SubChunkData, flags: CommandFlags.FireAndForget);
+			_ = transaction.KeyDeleteAsync(keys.Complete, flags: CommandFlags.FireAndForget);
+
+			_ = transaction.SortedSetRemoveAsync(ItemsKey, keys.Prefix, flags: CommandFlags.FireAndForget);
 			await transaction.ExecuteAsync();
 		}
 
@@ -289,6 +290,8 @@ namespace Horde.Build.Logs.Builder
 					return new LogChunkData(offset, lineIndex, subChunks);
 				}
 
+				await transaction.WaitAndIgnoreCancellations(typeTask, lastSubChunkDataTask);
+
 				// Cool down before retrying
 				await Task.Delay(100);
 			}
@@ -314,7 +317,7 @@ namespace Horde.Build.Logs.Builder
 				{
 					ITransaction transaction = redisDb.CreateTransaction();
 					transaction.AddCondition(Condition.SortedSetEqual(ItemsKey, entry.Element, entry.Score));
-					_ = transaction.SortedSetAddAsync(ItemsKey, entry.Element, utcNow.Ticks);
+					_ = transaction.SortedSetAddAsync(ItemsKey, entry.Element, utcNow.Ticks, flags: CommandFlags.FireAndForget);
 
 					if (!await transaction.ExecuteAsync())
 					{
diff --git a/Engine/Source/Programs/Horde/Horde.Build/Perforce/ReplicationService.cs b/Engine/Source/Programs/Horde/Horde.Build/Perforce/ReplicationService.cs
--- a/Engine/Source/Programs/Horde/Horde.Build/Perforce/ReplicationService.cs
+++ b/Engine/Source/Programs/Horde/Horde.Build/Perforce/ReplicationService.cs
@@ -376,7 +376,7 @@ namespace Horde.Build.Perforce
 				// Remove this stream from the dirty list if it's empty
 				ITransaction transaction = _redisConnectionPool.GetDatabase().CreateTransaction();
 				transaction.AddCondition(Condition.ListLengthLessThan(streamChanges.Key, 2));
-				_ = transaction.With(_redisDirtyStreams).RemoveAsync(streamId);
+				_ = transaction.With(_redisDirtyStreams).RemoveAsync(streamId, flags: CommandFlags.FireAndForget);
 				if (await transaction.ExecuteAsync())
 				{
 					break;
@@ -427,7 +427,7 @@ namespace Horde.Build.Perforce
 					ITransaction transaction = _redisConnectionPool.GetDatabase().CreateTransaction();
 					transaction.AddCondition(Condition.ListIndexEqual(changes.Key, 0, values[0]));
 					transaction.AddCondition(Condition.ListIndexEqual(changes.Key, 1, values[1]));
-					_ = transaction.With(changes).LeftPopAsync();
+					_ = transaction.With(changes).LeftPopAsync(flags: CommandFlags.FireAndForget);
 					await transaction.ExecuteAsync();
 				}
 				else
diff --git a/Engine/Source/Programs/Horde/Horde.Build/Tools/ToolCollection.cs b/Engine/Source/Programs/Horde/Horde.Build/Tools/ToolCollection.cs
--- a/Engine/Source/Programs/Horde/Horde.Build/Tools/ToolCollection.cs
+++ b/Engine/Source/Programs/Horde/Horde.Build/Tools/ToolCollection.cs
@@ -220,7 +220,7 @@ namespace Horde.Build.Tools
 				{
 					ITransaction transaction = redis.CreateTransaction();
 					transaction.AddCondition(Condition.StringEqual(s_indexKey, value));
-					_ = transaction.StringSetAsync(s_indexKey, newValue);
+					_ = transaction.StringSetAsync(s_indexKey, newValue, flags: CommandFlags.FireAndForget);
 
 					if (await transaction.ExecuteAsync())
 					{
diff --git a/Engine/Source/Programs/Horde/Horde.Build/Utilities/RedisExtensions.cs b/Engine/Source/Programs/Horde/Horde.Build/Utilities/RedisExtensions.cs
new file mode 100644
--- /dev/null
+++ b/Engine/Source/Programs/Horde/Horde.Build/Utilities/RedisExtensions.cs
@@ -0,0 +1,23 @@
+// Copyright Epic Games, Inc. All Rights Reserved.
+
+using System.Threading.Tasks;
+using StackExchange.Redis;
+
+namespace Horde.Build.Utilities;
+
+/// <summary>
+/// Extensions for Redis classes
+/// </summary>
+public static class RedisExtensions
+{
+	/// <summary>
+	/// Await given tasks and swallow any task cancellation exceptions.
+	/// Useful to do after (failing) transaction to ensure no trailing tasks are left unawaited
+	/// </summary>
+	/// <param name="transaction"></param>
+	/// <param name="tasks">Tasks to await</param>
+	public static async Task WaitAndIgnoreCancellations(this ITransaction transaction, params Task[] tasks)
+	{
+		try { await Task.WhenAll(tasks); } catch (TaskCanceledException) { /* Ignore */ }
+	}
+}
\ No newline at end of file
diff --git a/Engine/Source/Programs/Horde/Horde.Build/Utilities/VersionedCollection.cs b/Engine/Source/Programs/Horde/Horde.Build/Utilities/VersionedCollection.cs
--- a/Engine/Source/Programs/Horde/Horde.Build/Utilities/VersionedCollection.cs
+++ b/Engine/Source/Programs/Horde/Horde.Build/Utilities/VersionedCollection.cs
@@ -238,7 +238,7 @@ namespace Horde.Build.Utilities
 			pairs[0] = new KeyValuePair<RedisKey, RedisValue>(docKey, doc.ToBson(typeof(VersionedDocument<TId, TLatest>)));
 			pairs[1] = new KeyValuePair<RedisKey, RedisValue>(docKey.Append(s_timeSuffix), doc.LastUpdateTime.Ticks);
 
-			_ = transaction.StringSetAsync(pairs);
+			_ = transaction.StringSetAsync(pairs, flags: CommandFlags.FireAndForget);
 
 			if (!await transaction.ExecuteAsync())
 			{
