Horde: Move settings for commit replication into the global Horde settings object.

Add an appsettings.Local.json file which disables replication for local runs of the server from VS.

diff --git a/Engine/Source/Programs/Horde/Horde.Build/Perforce/PerforceServiceCache.cs b/Engine/Source/Programs/Horde/Horde.Build/Perforce/PerforceServiceCache.cs
--- a/Engine/Source/Programs/Horde/Horde.Build/Perforce/PerforceServiceCache.cs
+++ b/Engine/Source/Programs/Horde/Horde.Build/Perforce/PerforceServiceCache.cs
@@ -27,17 +27,6 @@ namespace Horde.Build.Perforce
 	using StreamId = StringId<IStream>;
 	using UserId = ObjectId<IUser>;
 
-	/// <summary>
-	/// Options for the commit service
-	/// </summary>
-	public class CommitServiceOptions
-	{
-		/// <summary>
-		/// Whether to mirror commit metadata to the database
-		/// </summary>
-		public bool Enable { get; set; } = true;
-	}
-
 	/// <summary>
 	/// Service which mirrors changes from Perforce
 	/// </summary>
@@ -164,8 +153,8 @@ namespace Horde.Build.Perforce
 			}
 		}
 
-		public CommitServiceOptions Options { get; }
 		readonly MongoService _mongoService;
+		readonly ServerSettings _settings;
 		readonly IMongoCollection<CachedCommitDoc> _commits;
 		readonly IStreamCollection _streamCollection;
 		readonly ILogger _logger;
@@ -174,12 +163,11 @@ namespace Horde.Build.Perforce
 		/// <summary>
 		/// Constructor
 		/// </summary>
-		public PerforceServiceCache(PerforceLoadBalancer loadBalancer, MongoService mongoService, GlobalsService globalsService, IUserCollection userCollection, IStreamCollection streamCollection, IClock clock, IOptions<ServerSettings> settings, IOptions<CommitServiceOptions> options, ILogger<PerforceService> logger)
+		public PerforceServiceCache(PerforceLoadBalancer loadBalancer, MongoService mongoService, GlobalsService globalsService, IUserCollection userCollection, IStreamCollection streamCollection, IClock clock, IOptions<ServerSettings> settings, ILogger<PerforceService> logger)
 			: base(loadBalancer, globalsService, userCollection, settings, logger)
 		{
-			Options = options.Value;
-
 			_mongoService = mongoService;
+			_settings = settings.Value;
 
 			List<MongoIndex<CachedCommitDoc>> indexes = new List<MongoIndex<CachedCommitDoc>>();
 			indexes.Add(MongoIndex.Create<CachedCommitDoc>(keys => keys.Ascending(x => x.StreamId).Descending(x => x.Number), true));
@@ -203,7 +191,7 @@ namespace Horde.Build.Perforce
 		/// <inheritdoc/>
 		public async Task StartAsync(CancellationToken cancellationToken)
 		{
-			if (Options.Enable)
+			if (_settings.Commits.ReplicateMetadata)
 			{
 				await _updateCommitsTicker.StartAsync();
 			}
@@ -212,7 +200,7 @@ namespace Horde.Build.Perforce
 		/// <inheritdoc/>
 		public async Task StopAsync(CancellationToken cancellationToken)
 		{
-			if (Options.Enable)
+			if (_settings.Commits.ReplicateMetadata)
 			{
 				await _updateCommitsTicker.StopAsync();
 			}
diff --git a/Engine/Source/Programs/Horde/Horde.Build/Perforce/ReplicationService.cs b/Engine/Source/Programs/Horde/Horde.Build/Perforce/ReplicationService.cs
--- a/Engine/Source/Programs/Horde/Horde.Build/Perforce/ReplicationService.cs
+++ b/Engine/Source/Programs/Horde/Horde.Build/Perforce/ReplicationService.cs
@@ -39,27 +39,6 @@ namespace Horde.Build.Perforce
 		}
 	}
 
-	/// <summary>
-	/// Options for the commit service
-	/// </summary>
-	public class ReplicationServiceOptions
-	{
-		/// <summary>
-		/// Whether to enable replication. Must also be enabled on a per-stream basis (see <see cref="StreamConfig.ReplicationMode"/>).
-		/// </summary>
-		public bool Enable { get; set; } = true;
-
-		/// <summary>
-		/// Options for how objects are packed together
-		/// </summary>
-		public TreeOptions Bundle { get; set; } = new TreeOptions();
-
-		/// <summary>
-		/// Options for how objects are sliced
-		/// </summary>
-		public ChunkingOptions Chunking { get; set; } = new ChunkingOptions();
-	}
-
 	/// <summary>
 	/// Root node for a commit snapshot
 	/// </summary>
@@ -136,7 +115,7 @@ namespace Horde.Build.Perforce
 			}
 		}
 
-		public ReplicationServiceOptions Options { get; }
+		readonly ServerSettings _settings;
 
 		// Redis
 		readonly RedisConnectionPool _redisConnectionPool;
@@ -167,9 +146,9 @@ namespace Horde.Build.Perforce
 		/// <summary>
 		/// Constructor
 		/// </summary>
-		public ReplicationService(RedisService redisService, /* CommitService commitService, */IStreamCollection streamCollection, IPerforceService perforceService, StorageService storageService, IOptions<ReplicationServiceOptions> options, ILogger<ReplicationService> logger)
+		public ReplicationService(RedisService redisService, /* CommitService commitService, */IStreamCollection streamCollection, IPerforceService perforceService, StorageService storageService, IOptions<ServerSettings> settings, ILogger<ReplicationService> logger)
 		{
-			Options = options.Value;
+			_settings = settings.Value;
 
 			_redisConnectionPool = redisService.ConnectionPool;
 			_redisDirtyStreams = new RedisSet<StreamId>(_redisConnectionPool, RedisBaseKey.Append("streams"));
@@ -193,7 +172,7 @@ namespace Horde.Build.Perforce
 		/// <inheritdoc/>
 		public async Task StartAsync(CancellationToken cancellationToken)
 		{
-			if (Options.Enable)
+			if (_settings.Commits.ReplicateContent)
 			{
 				_stopping = false;
 				_updateStreamsEvent.Reset();
@@ -205,7 +184,7 @@ namespace Horde.Build.Perforce
 		/// <inheritdoc/>
 		public async Task StopAsync(CancellationToken cancellationToken)
 		{
-			if (Options.Enable)
+			if (_settings.Commits.ReplicateContent)
 			{
 				_stopping = true;
 				_updateStreamsEvent.Latch();
@@ -628,7 +607,7 @@ namespace Horde.Build.Perforce
 					using ReadOnlyMemoryStream dataStream = new ReadOnlyMemoryStream(data);
 
 					FileEntry entry = await root.AddFileByPathAsync(path, FileEntryFlags.PerforceDepotPathAndRevision, cancellationToken);
-					await entry.AppendAsync(dataStream, Options.Chunking, writer, cancellationToken);
+					await entry.AppendAsync(dataStream, _settings.Commits.Chunking, writer, cancellationToken);
 				}
 			}
 			else
@@ -783,7 +762,7 @@ namespace Horde.Build.Perforce
 								else if (io.Command == PerforceIoCommand.Write)
 								{
 									FileWriter file = handles[io.File];
-									await file.AppendAsync(io.Payload, Options.Chunking, writer, cancellationToken);
+									await file.AppendAsync(io.Payload, _settings.Commits.Chunking, writer, cancellationToken);
 								}
 								else if (io.Command == PerforceIoCommand.Close)
 								{
diff --git a/Engine/Source/Programs/Horde/Horde.Build/Program.cs b/Engine/Source/Programs/Horde/Horde.Build/Program.cs
--- a/Engine/Source/Programs/Horde/Horde.Build/Program.cs
+++ b/Engine/Source/Programs/Horde/Horde.Build/Program.cs
@@ -141,11 +141,17 @@ namespace Horde.Build
 
 			CommandLineArguments arguments = new CommandLineArguments(args);
 
+			string? environment = Environment.GetEnvironmentVariable("ASPNETCORE_ENVIRONMENT");
+			if (String.IsNullOrEmpty(environment))
+			{
+				environment = "Production";
+			}
+
 			IConfiguration config = new ConfigurationBuilder()
 				.SetBasePath(AppDir.FullName)
 				.AddJsonFile("appsettings.json", optional: false) 
 				.AddJsonFile("appsettings.Build.json", optional: true) // specific settings for builds (installer/dockerfile)
-				.AddJsonFile($"appsettings.{Environment.GetEnvironmentVariable("ASPNETCORE_ENVIRONMENT")}.json", optional: true) // environment variable overrides, also used in k8s setups with Helm
+				.AddJsonFile($"appsettings.{environment}.json", optional: true) // environment variable overrides, also used in k8s setups with Helm
 				.AddJsonFile("appsettings.User.json", optional: true)
 				.AddJsonFile(UserConfigFile.FullName, optional: true, reloadOnChange: true)
 				.AddEnvironmentVariables()
diff --git a/Engine/Source/Programs/Horde/Horde.Build/ServerSettings.cs b/Engine/Source/Programs/Horde/Horde.Build/ServerSettings.cs
--- a/Engine/Source/Programs/Horde/Horde.Build/ServerSettings.cs
+++ b/Engine/Source/Programs/Horde/Horde.Build/ServerSettings.cs
@@ -5,6 +5,7 @@ using System.Linq;
 using EpicGames.Horde.Storage;
 using EpicGames.Horde.Storage.Nodes;
 using Horde.Build.Agents.Fleet;
+using Horde.Build.Perforce;
 using Horde.Build.Storage.Backends;
 using Horde.Build.Utilities;
 using Serilog.Events;
@@ -178,6 +179,32 @@ namespace Horde.Build
 		public bool LimitConcurrentLogChunkWriting { get; set; } = false;
 	}
 
+	/// <summary>
+	/// Options for the commit service
+	/// </summary>
+	public class CommitSettings
+	{
+		/// <summary>
+		/// Whether to mirror commit metadata to the database
+		/// </summary>
+		public bool ReplicateMetadata { get; set; } = true;
+
+		/// <summary>
+		/// Whether to mirror commit metadata to the database
+		/// </summary>
+		public bool ReplicateContent { get; set; } = false;
+
+		/// <summary>
+		/// Options for how objects are packed together
+		/// </summary>
+		public TreeOptions Bundle { get; set; } = new TreeOptions();
+
+		/// <summary>
+		/// Options for how objects are sliced
+		/// </summary>
+		public ChunkingOptions Chunking { get; set; } = new ChunkingOptions();
+	}
+
 	/// <summary>
 	/// Global settings for the application
 	/// </summary>
@@ -590,6 +617,11 @@ namespace Horde.Build
 		/// <inheritdoc cref="FeatureFlags" />
 		public FeatureFlagSettings FeatureFlags { get; set; } = new ();
 
+		/// <summary>
+		/// Options for the commit service
+		/// </summary>
+		public CommitSettings Commits { get; set; } = new CommitSettings();
+
 		/// <summary>
 		/// Helper method to check if this process has activated the given mode
 		/// </summary>
diff --git a/Engine/Source/Programs/Horde/Horde.Build/Startup.cs b/Engine/Source/Programs/Horde/Horde.Build/Startup.cs
--- a/Engine/Source/Programs/Horde/Horde.Build/Startup.cs
+++ b/Engine/Source/Programs/Horde/Horde.Build/Startup.cs
@@ -292,12 +292,8 @@ namespace Horde.Build
 			IConfigurationSection configSection = Configuration.GetSection("Horde");
 			ServerSettings settings = new ServerSettings();
 			configSection.Bind(settings);
-
 			settings.Validate();
 
-			services.Configure<CommitServiceOptions>(configSection.GetSection("Commits"));
-			services.Configure<ReplicationServiceOptions>(configSection.GetSection("Replication"));
-
 			if (settings.GlobalThreadPoolMinSize != null)
 			{
 				// Min thread pool size is set to combat timeouts seen with the Redis client.
diff --git a/Engine/Source/Programs/Horde/Horde.Build/appsettings.Local.json b/Engine/Source/Programs/Horde/Horde.Build/appsettings.Local.json
new file mode 100644
--- /dev/null
+++ b/Engine/Source/Programs/Horde/Horde.Build/appsettings.Local.json
@@ -0,0 +1,10 @@
+{
+	"Horde":
+	{
+		"Commits":
+		{
+			"ReplicateMetadata": false,
+			"ReplicateContent": false
+		}
+	}
+}
