Horde: Migrate to Microsoft.Management.Infrastructure from System.Management. Latter is deprecated, and not supported on Windows for ARM.

diff --git a/Engine/Source/Programs/Horde/Horde.Agent/Commands/Service/UpgradeCommand.cs b/Engine/Source/Programs/Horde/Horde.Agent/Commands/Service/UpgradeCommand.cs
--- a/Engine/Source/Programs/Horde/Horde.Agent/Commands/Service/UpgradeCommand.cs
+++ b/Engine/Source/Programs/Horde/Horde.Agent/Commands/Service/UpgradeCommand.cs
@@ -5,7 +5,6 @@ using System.Collections.Generic;
 using System.Diagnostics;
 using System.IO;
 using System.Linq;
-using System.Management;
 using System.Reflection;
 using System.Runtime.InteropServices;
 using System.Runtime.Versioning;
@@ -14,6 +13,7 @@ using System.Text;
 using System.Threading.Tasks;
 using EpicGames.Core;
 using Microsoft.Extensions.Logging;
+using Microsoft.Management.Infrastructure;
 
 namespace Horde.Agent.Commands
 {
@@ -199,21 +199,26 @@ namespace Horde.Agent.Commands
 		/// <param name="processId">The process id to search for</param>
 		/// <returns>The service controller corresponding to this process</returns>
 		[SupportedOSPlatform("windows")]
-		static ServiceController? GetServiceForProcess(int processId)
+		public static ServiceController? GetServiceForProcess(int processId)
 		{
 			try
 			{
-				using (ManagementObjectSearcher searcher = new ManagementObjectSearcher($"SELECT Name FROM Win32_Service WHERE ProcessId={processId}"))
+				using (CimSession session = CimSession.Create(null))
 				{
-					foreach (ManagementObject service in searcher.Get())
+					const string QueryNamespace = @"root\cimv2";
+					const string QueryDialect = "WQL";
+
+					foreach (CimInstance instance in session.QueryInstances(QueryNamespace, QueryDialect, $"SELECT Name FROM Win32_Service WHERE ProcessId={processId}"))
 					{
-						PropertyData property = service.Properties["Name"];
-						if (property != null)
+						foreach (CimProperty property in instance.CimInstanceProperties)
 						{
-							string? name = property.Value as string;
-							if (name != null)
+							if (property.Name.Equals("Name", StringComparison.OrdinalIgnoreCase))
 							{
-								return new ServiceController(name);
+								string? value = property.Value as string;
+								if (value != null)
+								{
+									return new ServiceController(value);
+								}
 							}
 						}
 					}
diff --git a/Engine/Source/Programs/Horde/Horde.Agent/Services/CapabilitiesService.cs b/Engine/Source/Programs/Horde/Horde.Agent/Services/CapabilitiesService.cs
--- a/Engine/Source/Programs/Horde/Horde.Agent/Services/CapabilitiesService.cs
+++ b/Engine/Source/Programs/Horde/Horde.Agent/Services/CapabilitiesService.cs
@@ -3,13 +3,10 @@
 using System;
 using System.Collections.Generic;
 using System.Diagnostics;
-using System.Diagnostics.CodeAnalysis;
 using System.IO;
 using System.Linq;
-using System.Management;
 using System.Net;
 using System.Runtime.InteropServices;
-using System.Runtime.Versioning;
 using System.Text.RegularExpressions;
 using System.Threading;
 using System.Threading.Tasks;
@@ -22,6 +19,7 @@ using Horde.Agent.Execution;
 using HordeCommon.Rpc.Messages;
 using Microsoft.Extensions.Logging;
 using Microsoft.Extensions.Options;
+using Microsoft.Management.Infrastructure;
 
 namespace Horde.Agent.Services
 {
@@ -69,84 +67,73 @@ namespace Horde.Agent.Services
 				primaryDevice.Properties.Add("OSFamily=Windows");
 
 				// WMI doesn't work currently on ARM64, due to reliance on a NET Framework assembly.
-				if (RuntimeInformation.OSArchitecture != Architecture.Arm64)
+				using (CimSession session = CimSession.Create(null))
 				{
+					const string QueryNamespace = @"root\cimv2";
+					const string QueryDialect = "WQL";
+
 					// Add OS info
-					using (ManagementObjectSearcher searcher = new ManagementObjectSearcher("select * from Win32_OperatingSystem"))
+					foreach (CimInstance instance in session.QueryInstances(QueryNamespace, QueryDialect, "select * from Win32_OperatingSystem"))
 					{
-						foreach (ManagementObject row in searcher.Get())
+						foreach (CimProperty property in instance.CimInstanceProperties)
 						{
-							using (row)
+							string name = property.Name;
+							if (name.Equals("Caption", StringComparison.OrdinalIgnoreCase))
 							{
-								Dictionary<string, object> properties = GetWmiProperties(row);
-
-								object? name;
-								if (properties.TryGetValue("Caption", out name))
-								{
-									primaryDevice.Properties.Add($"OSDistribution={name}");
-								}
-
-								object? version;
-								if (properties.TryGetValue("Version", out version))
-								{
-									primaryDevice.Properties.Add($"OSKernelVersion={version}");
-								}
+								primaryDevice.Properties.Add($"OSDistribution={property.Value}");
+							}
+							else if (name.Equals("Version", StringComparison.OrdinalIgnoreCase))
+							{
+								primaryDevice.Properties.Add($"OSKernelVersion={property.Value}");
 							}
 						}
 					}
 
 					// Add CPU info
-					using (ManagementObjectSearcher searcher = new ManagementObjectSearcher("select * from Win32_Processor"))
-					{
-						Dictionary<string, int> nameToCount = new Dictionary<string, int>();
-						int totalPhysicalCores = 0;
-						int totalLogicalCores = 0;
+					Dictionary<string, int> cpuNameToCount = new Dictionary<string, int>();
+					int totalPhysicalCores = 0;
+					int totalLogicalCores = 0;
 
-						foreach (ManagementObject row in searcher.Get())
+					foreach (CimInstance instance in session.QueryInstances(QueryNamespace, QueryDialect, "select * from Win32_Processor"))
+					{
+						foreach (CimProperty property in instance.CimInstanceProperties)
 						{
-							using (row)
+							string name = property.Name;
+							if (name.Equals("Name", StringComparison.OrdinalIgnoreCase))
 							{
-								Dictionary<string, object> properties = GetWmiProperties(row);
-
-								object? nameObject;
-								if (properties.TryGetValue("Name", out nameObject))
-								{
-									string name = nameObject.ToString() ?? String.Empty;
-									int count;
-									nameToCount.TryGetValue(name, out count);
-									nameToCount[name] = count + 1;
-								}
-
-								object? numPhysicalCores;
-								if ((properties.TryGetValue("NumberOfEnabledCore", out numPhysicalCores) && numPhysicalCores is uint) || (properties.TryGetValue("NumberOfCores", out numPhysicalCores) && numPhysicalCores is uint))
+								string cpuName = property.Value.ToString() ?? String.Empty;
+								int count;
+								cpuNameToCount.TryGetValue(cpuName, out count);
+								cpuNameToCount[cpuName] = count + 1;
+							}
+							else if (name.Equals("NumberOfEnabledCore", StringComparison.OrdinalIgnoreCase) || name.Equals("NumberOfCores", StringComparison.OrdinalIgnoreCase))
+							{
+								if (property.Value is uint numCores)
 								{
-									totalPhysicalCores += (int)(uint)numPhysicalCores;
+									totalPhysicalCores += (int)numCores;
 								}
-
-								object? numLogicalCores;
-								if (properties.TryGetValue("NumberOfLogicalProcessors", out numLogicalCores) && numLogicalCores is uint numLogicalCoresUint)
+							}
+							else if (name.Equals("NumberOfLogicalProcessors", StringComparison.OrdinalIgnoreCase))
+							{
+								if (property.Value is uint numCores)
 								{
-									totalLogicalCores += (int)numLogicalCoresUint;
+									totalLogicalCores += (int)numCores;
 								}
 							}
 						}
-
-						AddCpuInfo(primaryDevice, nameToCount, totalLogicalCores, totalPhysicalCores);
 					}
 
+					AddCpuInfo(primaryDevice, cpuNameToCount, totalLogicalCores, totalPhysicalCores);
+
 					// Add RAM info
-					using (ManagementObjectSearcher searcher = new ManagementObjectSearcher("select Capacity from Win32_PhysicalMemory"))
+					foreach (CimInstance instance in session.QueryInstances(QueryNamespace, QueryDialect, "select Capacity from Win32_PhysicalMemory"))
 					{
 						ulong totalCapacity = 0;
-						foreach (ManagementObject row in searcher.Get())
+						foreach (CimProperty property in instance.CimInstanceProperties)
 						{
-							using (row)
+							if (property.Name.Equals("Capacity", StringComparison.OrdinalIgnoreCase) && property.Value is ulong capacity)
 							{
-								object? capacity = row.GetPropertyValue("Capacity");
-								if (capacity is ulong ulongCapacity)
-								{
-									totalCapacity += ulongCapacity;
-								}
+								totalCapacity += capacity;
 							}
 						}
 
@@ -157,20 +144,32 @@ namespace Horde.Agent.Services
 					}
 
 					// Add GPU info
-					using (ManagementObjectSearcher searcher = new ManagementObjectSearcher("select Name, DriverVersion, AdapterRAM from Win32_VideoController"))
+					int index = 0;
+					foreach (CimInstance instance in session.QueryInstances(QueryNamespace, QueryDialect, "select Name, DriverVersion, AdapterRAM from Win32_VideoController"))
 					{
-						int index = 0;
-						foreach (ManagementObject row in searcher.Get())
+						string? name = null;
+						string? driverVersion = null;
+
+						foreach (CimProperty property in instance.CimInstanceProperties)
 						{
-							using (row)
+							if (property.Name.Equals("Name", StringComparison.OrdinalIgnoreCase))
 							{
-								WmiProperties properties = new WmiProperties(row);
-								if (properties.TryGetValue("Name", out string? name) && properties.TryGetValue("DriverVersion", out string? driverVersion))
-								{
-									string prefix = $"GPU-{++index}";
-									primaryDevice.Properties.Add($"{prefix}-Name={name}");
-									primaryDevice.Properties.Add($"{prefix}-DriverVersion={driverVersion}");
-								}
+								name = property.Value.ToString();
+							}
+							else if (property.Name.Equals("DriverVersion", StringComparison.OrdinalIgnoreCase))
+							{
+								driverVersion = property.Value.ToString();
+							}
+						}
+
+						if (name != null)
+						{
+							string prefix = $"GPU-{++index}";
+							primaryDevice.Properties.Add($"{prefix}-Name={name}");
+
+							if (driverVersion != null)
+							{
+								primaryDevice.Properties.Add($"{prefix}-DriverVersion={driverVersion}");
 							}
 						}
 					}
@@ -476,77 +475,6 @@ namespace Horde.Agent.Services
 			return records;
 		}
 
-		[SupportedOSPlatform("windows")]
-		class WmiProperties
-		{
-			readonly Dictionary<string, object> _properties = new Dictionary<string, object>(StringComparer.Ordinal);
-
-			public WmiProperties(ManagementObject row)
-			{
-				foreach (PropertyData property in row.Properties)
-				{
-					_properties[property.Name] = property.Value;
-				}
-			}
-
-			public bool TryGetValue(string name, [NotNullWhen(true)] out string? value)
-			{
-				object? @object;
-				if (_properties.TryGetValue(name, out @object))
-				{
-					value = @object.ToString();
-					return value != null;
-				}
-				else
-				{
-					value = null!;
-					return false;
-				}
-			}
-
-			public bool TryGetValue(string name, out long value)
-			{
-				object? objectValue;
-				if (_properties.TryGetValue(name, out objectValue))
-				{
-					if (objectValue is int intValue)
-					{
-						value = intValue;
-						return true;
-					}
-					else if (objectValue is uint uintValue)
-					{
-						value = uintValue;
-						return true;
-					}
-					else if (objectValue is long longValue)
-					{
-						value = longValue;
-						return true;
-					}
-					else if (objectValue is ulong ulongValue)
-					{
-						value = (long)ulongValue;
-						return true;
-					}
-				}
-
-				value = 0;
-				return false;
-			}
-		}
-
-		[SupportedOSPlatform("windows")]
-		static Dictionary<string, object> GetWmiProperties(ManagementObject row)
-		{
-			Dictionary<string, object> properties = new Dictionary<string, object>(StringComparer.Ordinal);
-			foreach (PropertyData property in row.Properties)
-			{
-				properties[property.Name] = property.Value;
-			}
-			return properties;
-		}
-
 		static async Task AddAwsProperties(IList<string> properties, ILogger logger)
 		{
 			if (EC2InstanceMetadata.IdentityDocument != null)
