Horde: Merge the compute task executor logic with the lease handling logic.

diff --git a/Engine/Source/Programs/Horde/Horde.Agent.Tests/ComputeTaskExecutorTest.cs b/Engine/Source/Programs/Horde/Horde.Agent.Tests/ComputeTaskExecutorTest.cs
--- a/Engine/Source/Programs/Horde/Horde.Agent.Tests/ComputeTaskExecutorTest.cs
+++ b/Engine/Source/Programs/Horde/Horde.Agent.Tests/ComputeTaskExecutorTest.cs
@@ -8,7 +8,7 @@ using EpicGames.Core;
 using EpicGames.Horde.Compute;
 using EpicGames.Horde.Storage;
 using EpicGames.Horde.Storage.Impl;
-using Horde.Agent.Execution;
+using Horde.Agent.Leases.Handlers;
 using Microsoft.Extensions.Logging.Abstractions;
 using Microsoft.VisualStudio.TestTools.UnitTesting;
 
@@ -19,13 +19,13 @@ namespace Horde.Agent.Tests
 	{
 		private readonly MemoryStorageClient _storageClient = new ();
 		private readonly FakeCompressor _compressor = new ();
-		private readonly ComputeTaskExecutor _executor;
+		private readonly ComputeHandler _executor;
 		private readonly NamespaceId _namespaceId = new ("my-namespace");
 		private readonly DirectoryReference _sandboxDir;
 		
 		public ComputeTaskExecutorTest()
 		{
-			_executor = new (_storageClient, NullLogger.Instance);
+			_executor = new (_storageClient, NullLogger<ComputeHandler>.Instance);
 			_sandboxDir = new (Path.Join(Path.GetTempPath(), "horde-agent", Guid.NewGuid().ToString()));
 			DirectoryReference.CreateDirectory(_sandboxDir);
 		}
diff --git a/Engine/Source/Programs/Horde/Horde.Agent/Execution/ComputeTaskExecutor.cs b/Engine/Source/Programs/Horde/Horde.Agent/Execution/ComputeTaskExecutor.cs
deleted file mode 100644
diff --git a/Engine/Source/Programs/Horde/Horde.Agent/Leases/Handlers/ComputeHandler.cs b/Engine/Source/Programs/Horde/Horde.Agent/Leases/Handlers/ComputeHandler.cs
--- a/Engine/Source/Programs/Horde/Horde.Agent/Leases/Handlers/ComputeHandler.cs
+++ b/Engine/Source/Programs/Horde/Horde.Agent/Leases/Handlers/ComputeHandler.cs
@@ -1,13 +1,20 @@
 // Copyright Epic Games, Inc. All Rights Reserved.
 
 using System;
+using System.Collections.Generic;
+using System.Diagnostics;
+using System.IO;
+using System.Linq;
+using System.Runtime.InteropServices;
+using System.Text;
 using System.Threading;
 using System.Threading.Tasks;
 using EpicGames.Core;
 using EpicGames.Horde.Compute;
 using EpicGames.Horde.Storage;
+using EpicGames.Serialization;
 using Google.Protobuf;
-using Horde.Agent.Execution;
+using Grpc.Core;
 using Horde.Agent.Services;
 using Horde.Agent.Utility;
 using HordeCommon.Rpc.Tasks;
@@ -17,7 +24,13 @@ namespace Horde.Agent.Leases.Handlers
 {
 	class ComputeHandler : LeaseHandler<ComputeTaskMessage>
 	{
-		readonly ILegacyStorageClient _legacyStorageClient;
+		class OutputTree
+		{
+			public Dictionary<string, Task<IoHash>> _files = new Dictionary<string, Task<IoHash>>();
+			public Dictionary<string, OutputTree> _subDirs = new Dictionary<string, OutputTree>();
+		}
+
+		readonly ILegacyStorageClient _storageClient;
 		readonly ILogger _logger;
 
 		/// <summary>
@@ -25,7 +38,7 @@ namespace Horde.Agent.Leases.Handlers
 		/// </summary>
 		public ComputeHandler(ILegacyStorageClient legacyStorageClient, ILogger<ComputeHandler> logger)
 		{
-			_legacyStorageClient = legacyStorageClient;
+			_storageClient = legacyStorageClient;
 			_logger = logger;
 		}
 
@@ -43,10 +56,9 @@ namespace Horde.Agent.Leases.Handlers
 					DirectoryReference leaseDir = DirectoryReference.Combine(session.WorkingDir, "Compute", leaseId);
 					DirectoryReference.CreateDirectory(leaseDir);
 
-					ComputeTaskExecutor executor = new ComputeTaskExecutor(_legacyStorageClient, _logger);
 					try
 					{
-						result = await executor.ExecuteAsync(leaseId, computeTask, leaseDir, cancellationToken);
+						result = await ExecuteAsync(leaseId, computeTask, leaseDir, cancellationToken);
 					}
 					finally
 					{
@@ -73,6 +85,279 @@ namespace Horde.Agent.Leases.Handlers
 
 			return new LeaseResult(result.ToByteArray());
 		}
+
+		/// <summary>
+		/// Read stdout/stderr and wait until process completes
+		/// </summary>
+		/// <param name="process">The process reading from</param>
+		/// <param name="timeout">Max execution timeout for process</param>
+		/// <param name="cancelToken">Cancellation token</param>
+		/// <returns>stdout and stderr data</returns>
+		/// <exception cref="RpcException">Raised for either timeout or cancel</exception>
+		private static async Task<(byte[] stdOutData, byte[] stdOutErr)> ReadProcessStreams(ManagedProcess process, TimeSpan timeout, CancellationToken cancelToken)
+		{
+			// Read stdout/stderr without cancellation token.
+			using MemoryStream stdOutStream = new MemoryStream();
+			Task stdOutReadTask = process.StdOut.CopyToAsync(stdOutStream, cancelToken);
+
+			using MemoryStream stdErrStream = new MemoryStream();
+			Task stdErrReadTask = process.StdErr.CopyToAsync(stdErrStream, cancelToken);
+
+			Task outputTask = Task.WhenAll(stdOutReadTask, stdErrReadTask);
+
+			// Instead, create a separate task that will wait for either timeout or cancellation
+			// as cancel interruptions are not reliable with Stream.CopyToAsync()
+			Task timeoutTask = Task.Delay(timeout, cancelToken);
+
+			Task waitTask = await Task.WhenAny(outputTask, timeoutTask);
+			if (waitTask == timeoutTask)
+			{
+				throw new RpcException(new Grpc.Core.Status(StatusCode.DeadlineExceeded, $"Action timed out after {timeout.TotalMilliseconds} ms"));
+			}
+
+			return (stdOutStream.ToArray(), stdErrStream.ToArray());
+		}
+
+		static int GetMarkerTime(Stopwatch timer)
+		{
+			int elapsedMs = (int)timer.ElapsedMilliseconds;
+			timer.Restart();
+			return elapsedMs;
+		}
+
+		/// <summary>
+		/// Execute an action
+		/// </summary>
+		/// <param name="leaseId">The lease id</param>
+		/// <param name="computeTaskMessage">Task to execute</param>
+		/// <param name="sandboxDir">Directory to use as a sandbox for execution</param>
+		/// <param name="cancellationToken">Token used to cancel the operation</param>
+		/// <returns>The action result</returns>
+		public async Task<ComputeTaskResultMessage> ExecuteAsync(string leaseId, ComputeTaskMessage computeTaskMessage, DirectoryReference sandboxDir, CancellationToken cancellationToken)
+		{
+			Stopwatch timer = Stopwatch.StartNew();
+			ComputeTaskExecutionStatsMessage stats = new ComputeTaskExecutionStatsMessage();
+			stats.StartTime = Google.Protobuf.WellKnownTypes.Timestamp.FromDateTime(DateTime.UtcNow);
+
+			NamespaceId namespaceId = new NamespaceId(computeTaskMessage.NamespaceId);
+			BucketId inputBucketId = new BucketId(computeTaskMessage.InputBucketId);
+			BucketId outputBucketId = new BucketId(computeTaskMessage.OutputBucketId);
+
+			ComputeTask task = await _storageClient.GetRefAsync<ComputeTask>(namespaceId, inputBucketId, computeTaskMessage.TaskRefId.AsRefId(), cancellationToken);
+			_logger.LogInformation("Executing task {Hash} for lease ID {LeaseId}", computeTaskMessage.TaskRefId, leaseId);
+			stats.DownloadRefMs = GetMarkerTime(timer);
+
+			DirectoryReference.CreateDirectory(sandboxDir);
+			FileUtils.ForceDeleteDirectoryContents(sandboxDir);
+
+			DirectoryTree inputDirectory = await _storageClient.ReadBlobAsync<DirectoryTree>(namespaceId, task.SandboxHash, cancellationToken: cancellationToken);
+			await SetupSandboxAsync(namespaceId, inputDirectory, sandboxDir);
+			stats.DownloadInputMs = GetMarkerTime(timer);
+
+			using (ManagedProcessGroup processGroup = new ManagedProcessGroup())
+			{
+				string fileName = FileReference.Combine(sandboxDir, task.WorkingDirectory.ToString(), task.Executable.ToString()).FullName;
+				string workingDirectory = DirectoryReference.Combine(sandboxDir, task.WorkingDirectory.ToString()).FullName;
+				EnsureFileIsExecutable(fileName);
+
+				Dictionary<string, string> newEnvironment = new Dictionary<string, string>();
+				foreach (System.Collections.DictionaryEntry? entry in Environment.GetEnvironmentVariables())
+				{
+					newEnvironment[entry!.Value.Key.ToString()!] = entry!.Value.Value!.ToString()!;
+				}
+				foreach (KeyValuePair<Utf8String, Utf8String> pair in task.EnvVars)
+				{
+					newEnvironment[pair.Key.ToString()] = pair.Value.ToString();
+				}
+
+				string arguments = CommandLineArguments.Join(task.Arguments.Select(x => x.ToString()));
+				_logger.LogInformation("Executing {FileName} with arguments {Arguments}", fileName, arguments);
+
+				TimeSpan timeout = TimeSpan.FromMinutes(5.0);// Action.Timeout == null ? TimeSpan.FromMinutes(5) : Action.Timeout.ToTimeSpan();
+				using (ManagedProcess process = new ManagedProcess(processGroup, fileName, arguments, workingDirectory, newEnvironment, ProcessPriorityClass.Normal, ManagedProcessFlags.None))
+				{
+					(byte[] stdOutData, byte[] stdErrData) = await ReadProcessStreams(process, timeout, cancellationToken);
+					stats.ExecMs = GetMarkerTime(timer);
+
+					foreach (string line in Encoding.UTF8.GetString(stdOutData).Split('\n'))
+					{
+						_logger.LogInformation("stdout: {Line}", line);
+					}
+					foreach (string line in Encoding.UTF8.GetString(stdErrData).Split('\n'))
+					{
+						_logger.LogInformation("stderr: {Line}", line);
+					}
+
+					process.WaitForExit();
+					_logger.LogInformation("exit: {ExitCode}", process.ExitCode);
+
+					ComputeTaskResult result = new ComputeTaskResult(process.ExitCode);
+					result.StdOutHash = await _storageClient.WriteBlobFromMemoryAsync(namespaceId, stdOutData, cancellationToken);
+					result.StdErrHash = await _storageClient.WriteBlobFromMemoryAsync(namespaceId, stdErrData, cancellationToken);
+					stats.UploadLogMs = GetMarkerTime(timer);
+
+					FileReference[] outputFiles = ResolveOutputPaths(sandboxDir, task.OutputPaths.Select(x => x.ToString())).OrderBy(x => x.FullName, StringComparer.Ordinal).ToArray();
+					if (outputFiles.Length > 0)
+					{
+						foreach (FileReference outputFile in outputFiles)
+						{
+							_logger.LogInformation("output: {File}", outputFile.MakeRelativeTo(sandboxDir));
+						}
+						result.OutputHash = await PutOutput(namespaceId, sandboxDir, outputFiles);
+					}
+					stats.UploadOutputMs = GetMarkerTime(timer);
+
+					CbObject resultObject = CbSerializer.Serialize(result);
+					await _storageClient.SetRefAsync(namespaceId, outputBucketId, computeTaskMessage.TaskRefId, resultObject, cancellationToken);
+					stats.UploadRefMs = GetMarkerTime(timer);
+
+					return new ComputeTaskResultMessage(computeTaskMessage.TaskRefId, stats);
+				}
+			}
+		}
+
+		/// <summary>
+		/// Downloads files to the sandbox
+		/// </summary>
+		/// <param name="namespaceId">Namespace for fetching data</param>
+		/// <param name="inputDirectory">The directory spec</param>
+		/// <param name="outputDir">Output directory on disk</param>
+		/// <returns>Async task</returns>
+		internal async Task SetupSandboxAsync(NamespaceId namespaceId, DirectoryTree inputDirectory, DirectoryReference outputDir)
+		{
+			DirectoryReference.CreateDirectory(outputDir);
+
+			async Task DownloadFile(FileNode fileNode)
+			{
+				FileReference file = FileReference.Combine(outputDir, fileNode.Name.ToString());
+				_logger.LogInformation("Downloading {File} (digest: {Digest})", file, fileNode.Hash);
+				byte[] data = fileNode.IsCompressed
+					? await _storageClient.ReadCompressedBlobToMemoryAsync(namespaceId, fileNode.Hash)
+					: await _storageClient.ReadBlobToMemoryAsync(namespaceId, fileNode.Hash);
+				_logger.LogInformation("Writing {File} (digest: {Digest})", file, fileNode.Hash);
+				await FileReference.WriteAllBytesAsync(file, data);
+			}
+
+			async Task DownloadDir(DirectoryNode directoryNode)
+			{
+				DirectoryTree inputSubDirectory = await _storageClient.ReadBlobAsync<DirectoryTree>(namespaceId, directoryNode.Hash);
+				DirectoryReference outputSubDirectory = DirectoryReference.Combine(outputDir, directoryNode.Name.ToString());
+				await SetupSandboxAsync(namespaceId, inputSubDirectory, outputSubDirectory);
+			}
+
+			List<Task> tasks = new List<Task>();
+			tasks.AddRange(inputDirectory.Files.Select(x => Task.Run(() => DownloadFile(x))));
+			tasks.AddRange(inputDirectory.Directories.Select(x => Task.Run(() => DownloadDir(x))));
+			await Task.WhenAll(tasks);
+		}
+
+		async Task<IoHash> PutOutput(NamespaceId namespaceId, DirectoryReference baseDir, IEnumerable<FileReference> files)
+		{
+			List<FileReference> sortedFiles = files.OrderBy(x => x.FullName, StringComparer.Ordinal).ToList();
+			(_, IoHash hash) = await PutDirectoryTree(namespaceId, baseDir.FullName.Length, sortedFiles, 0, sortedFiles.Count);
+			return hash;
+		}
+
+		async Task<(DirectoryTree, IoHash)> PutDirectoryTree(NamespaceId namespaceId, int baseDirLen, List<FileReference> sortedFiles, int minIdx, int maxIdx)
+		{
+			List<Task<FileNode>> files = new List<Task<FileNode>>();
+			List<Task<DirectoryNode>> trees = new List<Task<DirectoryNode>>();
+
+			while (minIdx < maxIdx)
+			{
+				FileReference file = sortedFiles[minIdx];
+
+				int nextMinIdx = minIdx + 1;
+
+				int nextDirLen = file.FullName.IndexOf(Path.DirectorySeparatorChar, baseDirLen + 1);
+				if (nextDirLen == -1)
+				{
+					string name = file.FullName.Substring(baseDirLen + 1);
+					files.Add(CreateFileNode(namespaceId, name, file));
+				}
+				else
+				{
+					string name = file.FullName.Substring(baseDirLen + 1, nextDirLen - (baseDirLen + 1));
+					while (nextMinIdx < maxIdx)
+					{
+						string nextFile = sortedFiles[nextMinIdx].FullName;
+						if (nextFile.Length < nextDirLen || String.Compare(name, 0, nextFile, baseDirLen, name.Length, StringComparison.Ordinal) == 0)
+						{
+							break;
+						}
+						nextMinIdx++;
+					}
+					trees.Add(CreateDirectoryNode(namespaceId, name, nextDirLen, sortedFiles, minIdx, nextMinIdx));
+				}
+
+				minIdx = nextMinIdx;
+			}
+
+			DirectoryTree tree = new DirectoryTree();
+			tree.Files.AddRange(await Task.WhenAll(files));
+			tree.Directories.AddRange(await Task.WhenAll(trees));
+
+			IoHash hash = await _storageClient.WriteBlobAsync<DirectoryTree>(namespaceId, tree);
+			return (tree, hash);
+		}
+
+		private static void EnsureFileIsExecutable(string filePath)
+		{
+			if (RuntimeInformation.IsOSPlatform(OSPlatform.Linux))
+			{
+				ushort mode = 493; // 0755 octal Posix permission
+				FileUtils.SetFileMode_Linux(filePath, mode);
+			}
+		}
+
+		async Task<DirectoryNode> CreateDirectoryNode(NamespaceId namespaceId, string name, int baseDirLen, List<FileReference> sortedFiles, int minIdx, int maxIdx)
+		{
+			(_, IoHash hash) = await PutDirectoryTree(namespaceId, baseDirLen, sortedFiles, minIdx, maxIdx);
+			return new DirectoryNode(name, hash);
+		}
+
+		async Task<FileNode> CreateFileNode(NamespaceId namespaceId, string name, FileReference file)
+		{
+			byte[] data = await FileReference.ReadAllBytesAsync(file);
+			IoHash hash = await _storageClient.WriteBlobFromMemoryAsync(namespaceId, data);
+			// TODO: Figure out how file can be marked as compressed
+			return new FileNode(name, hash, data.Length, (int)FileReference.GetAttributes(file), false);
+		}
+
+		/// <summary>
+		/// Resolves a list of output paths into file references
+		///
+		/// The REAPI spec allows directories to be specified as an output path which require all sub dirs and files
+		/// to be resolved.
+		/// </summary>
+		/// <param name="sandboxDir">Base directory where execution is taking place</param>
+		/// <param name="outputPaths">List of output paths relative to SandboxDir</param>
+		/// <returns>List of resolved paths (incl expanded dirs)</returns>
+		internal static HashSet<FileReference> ResolveOutputPaths(DirectoryReference sandboxDir, IEnumerable<string> outputPaths)
+		{
+			HashSet<FileReference> files = new HashSet<FileReference>();
+			foreach (string outputPath in outputPaths)
+			{
+				DirectoryReference dirRef = DirectoryReference.Combine(sandboxDir, outputPath);
+				if (DirectoryReference.Exists(dirRef))
+				{
+					IEnumerable<FileReference> listedFiles = DirectoryReference.EnumerateFiles(dirRef, "*", SearchOption.AllDirectories);
+					foreach (FileReference listedFileRef in listedFiles)
+					{
+						files.Add(listedFileRef);
+					}
+				}
+				else
+				{
+					FileReference fileRef = FileReference.Combine(sandboxDir, outputPath);
+					if (FileReference.Exists(fileRef))
+					{
+						files.Add(fileRef);
+					}
+				}
+			}
+			return files;
+		}
 	}
 }
 
