Horde: Build fix (deleted namespace).

diff --git a/Engine/Source/Programs/Horde/Horde.Agent/Execution/JobExecutor.cs b/Engine/Source/Programs/Horde/Horde.Agent/Execution/JobExecutor.cs
--- a/Engine/Source/Programs/Horde/Horde.Agent/Execution/JobExecutor.cs
+++ b/Engine/Source/Programs/Horde/Horde.Agent/Execution/JobExecutor.cs
@@ -14,7 +14,6 @@ using System.Threading;
 using System.Threading.Tasks;
 using EpicGames.Core;
 using Grpc.Core;
-using Horde.Agent.Execution.Interfaces;
 using Horde.Agent.Parser;
 using Horde.Agent.Services;
 using Horde.Agent.Utility;
diff --git a/Engine/Source/Programs/Horde/Horde.Agent/Execution/LocalExecutor.cs b/Engine/Source/Programs/Horde/Horde.Agent/Execution/LocalExecutor.cs
--- a/Engine/Source/Programs/Horde/Horde.Agent/Execution/LocalExecutor.cs
+++ b/Engine/Source/Programs/Horde/Horde.Agent/Execution/LocalExecutor.cs
@@ -4,7 +4,6 @@ using System;
 using System.Threading;
 using System.Threading.Tasks;
 using EpicGames.Core;
-using Horde.Agent.Execution.Interfaces;
 using Horde.Agent.Services;
 using Horde.Agent.Utility;
 using HordeCommon.Rpc;
diff --git a/Engine/Source/Programs/Horde/Horde.Agent/Execution/PerforceExecutor.cs b/Engine/Source/Programs/Horde/Horde.Agent/Execution/PerforceExecutor.cs
--- a/Engine/Source/Programs/Horde/Horde.Agent/Execution/PerforceExecutor.cs
+++ b/Engine/Source/Programs/Horde/Horde.Agent/Execution/PerforceExecutor.cs
@@ -11,14 +11,12 @@ using System.Threading.Tasks;
 using EpicGames.Core;
 using EpicGames.Perforce;
 using EpicGames.Perforce.Managed;
-using Horde.Agent.Execution.Interfaces;
 using Horde.Agent.Services;
 using Horde.Agent.Utility;
 using HordeCommon.Rpc;
 using HordeCommon.Rpc.Messages;
 using HordeCommon.Rpc.Tasks;
 using Microsoft.Extensions.Logging;
-using Microsoft.Extensions.Options;
 using OpenTracing;
 using OpenTracing.Util;
 
diff --git a/Engine/Source/Programs/Horde/Horde.Agent/Leases/Handlers/JobHandler.cs b/Engine/Source/Programs/Horde/Horde.Agent/Leases/Handlers/JobHandler.cs
--- a/Engine/Source/Programs/Horde/Horde.Agent/Leases/Handlers/JobHandler.cs
+++ b/Engine/Source/Programs/Horde/Horde.Agent/Leases/Handlers/JobHandler.cs
@@ -8,7 +8,7 @@ using System.Threading;
 using System.Threading.Tasks;
 using EpicGames.Core;
 using Grpc.Core;
-using Horde.Agent.Execution.Interfaces;
+using Horde.Agent.Execution;
 using Horde.Agent.Parser;
 using Horde.Agent.Services;
 using Horde.Agent.Utility;
@@ -30,14 +30,14 @@ namespace Horde.Agent.Leases.Handlers
 		/// </summary>
 		internal TimeSpan _stepAbortPollInterval = TimeSpan.FromSeconds(5);
 
-		readonly IExecutorFactory _executorFactory;
+		readonly JobExecutorFactory _executorFactory;
 		readonly AgentSettings _settings;
 		readonly ILogger _defaultLogger;
 
 		/// <summary>
 		/// Constructor
 		/// </summary>
-		public JobHandler(IExecutorFactory executorFactory, IOptions<AgentSettings> settings, ILogger<JobHandler> defaultLogger)
+		public JobHandler(JobExecutorFactory executorFactory, IOptions<AgentSettings> settings, ILogger<JobHandler> defaultLogger)
 		{
 			_executorFactory = executorFactory;
 			_settings = settings.Value;
@@ -99,7 +99,7 @@ namespace Horde.Agent.Leases.Handlers
 			await session.TerminateProcessesAsync(batchLogger, cancellationToken);
 
 			// Create an executor for this job
-			IExecutor executor = _executorFactory.CreateExecutor(session, executeTask, batch);
+			JobExecutor executor = _executorFactory.CreateExecutor(session, executeTask, batch);
 
 			// Try to initialize the executor
 			batchLogger.LogInformation("Initializing...");
@@ -255,7 +255,7 @@ namespace Horde.Agent.Leases.Handlers
 		/// <param name="cancellationToken">Cancellation token to abort the batch</param>
 		/// <param name="stepCancellationToken">Cancellation token to abort only this individual step</param>
 		/// <returns>Async task</returns>
-		internal static async Task<(JobStepOutcome, JobStepState)> ExecuteStepAsync(IExecutor executor, BeginStepResponse step, ILogger stepLogger, CancellationToken cancellationToken, CancellationToken stepCancellationToken)
+		internal static async Task<(JobStepOutcome, JobStepState)> ExecuteStepAsync(JobExecutor executor, BeginStepResponse step, ILogger stepLogger, CancellationToken cancellationToken, CancellationToken stepCancellationToken)
 		{
 			using CancellationTokenSource combined = CancellationTokenSource.CreateLinkedTokenSource(cancellationToken, stepCancellationToken);
 			try
diff --git a/Engine/Source/Programs/Horde/Horde.Agent/Program.cs b/Engine/Source/Programs/Horde/Horde.Agent/Program.cs
--- a/Engine/Source/Programs/Horde/Horde.Agent/Program.cs
+++ b/Engine/Source/Programs/Horde/Horde.Agent/Program.cs
@@ -14,7 +14,6 @@ using Datadog.Trace.OpenTracing;
 using EpicGames.Core;
 using EpicGames.Horde.Storage;
 using Horde.Agent.Execution;
-using Horde.Agent.Execution.Interfaces;
 using Horde.Agent.Leases;
 using Horde.Agent.Leases.Handlers;
 using Horde.Agent.Services;
