Horde: Another missing file (removed namespaces)

diff --git a/Engine/Source/Programs/Horde/Horde.Agent/Utility/TempStorage.cs b/Engine/Source/Programs/Horde/Horde.Agent/Utility/TempStorage.cs
--- a/Engine/Source/Programs/Horde/Horde.Agent/Utility/TempStorage.cs
+++ b/Engine/Source/Programs/Horde/Horde.Agent/Utility/TempStorage.cs
@@ -345,13 +345,6 @@ namespace Horde.Storage.Utility
 		[XmlArrayItem("File")]
 		public TempStorageFile[] Files { get; set; }
 
-		/// <summary>
-		/// List of compressed archives containing the given files
-		/// </summary>
-		[XmlArray]
-		[XmlArrayItem("ZipFile")]
-		public TempStorageZipFile[] ZipFiles { get; set; }
-
 		/// <summary>
 		/// Construct a static Xml serializer to avoid throwing an exception searching for the reflection info at runtime
 		/// </summary>
@@ -363,7 +356,6 @@ namespace Horde.Storage.Utility
 		private TempStorageManifest()
 		{
 			Files = Array.Empty<TempStorageFile>();
-			ZipFiles = Array.Empty<TempStorageZipFile>();
 		}
 
 		/// <summary>
@@ -374,7 +366,6 @@ namespace Horde.Storage.Utility
 		public TempStorageManifest(FileInfo[] files, DirectoryReference rootDir)
 		{
 			Files = files.Select(x => new TempStorageFile(x, rootDir)).ToArray();
-			ZipFiles = Array.Empty<TempStorageZipFile>();
 		}
 
 		/// <summary>
