Horde: Missing files.

diff --git a/Engine/Source/Programs/Horde/Horde.Agent/Execution/TestExecutor.cs b/Engine/Source/Programs/Horde/Horde.Agent/Execution/TestExecutor.cs
--- a/Engine/Source/Programs/Horde/Horde.Agent/Execution/TestExecutor.cs
+++ b/Engine/Source/Programs/Horde/Horde.Agent/Execution/TestExecutor.cs
@@ -6,10 +6,7 @@ using System.Linq;
 using System.Reflection;
 using System.Threading;
 using System.Threading.Tasks;
-using Amazon.SimpleSystemsManagement.Model;
 using EpicGames.Core;
-using EpicGames.Perforce;
-using Horde.Agent.Execution.Interfaces;
 using Horde.Agent.Parser;
 using Horde.Agent.Services;
 using Horde.Agent.Utility;
