Horde: Fix tests.

diff --git a/Engine/Source/Programs/Horde/Horde.Agent.Tests/Helper.cs b/Engine/Source/Programs/Horde/Horde.Agent.Tests/Helper.cs
--- a/Engine/Source/Programs/Horde/Horde.Agent.Tests/Helper.cs
+++ b/Engine/Source/Programs/Horde/Horde.Agent.Tests/Helper.cs
@@ -7,7 +7,7 @@ using System.Threading.Tasks;
 using Google.Protobuf.WellKnownTypes;
 using Grpc.Core;
 using Grpc.Net.Client;
-using Horde.Agent.Execution.Interfaces;
+using Horde.Agent.Execution;
 using Horde.Agent.Services;
 using Horde.Agent.Utility;
 using HordeCommon;
@@ -202,44 +202,55 @@ namespace Horde.Agent.Tests
 		}
 	}
 
-	class SimpleTestExecutor : IExecutor
+	class SimpleTestExecutor : JobExecutor
 	{
 		private readonly Func<BeginStepResponse, ILogger, CancellationToken, Task<JobStepOutcome>> _func;
 
 		public SimpleTestExecutor(Func<BeginStepResponse, ILogger, CancellationToken, Task<JobStepOutcome>> func)
+			: base(null!, null!, null!, null!)
 		{
 			_func = func;
 		}
 
-		public Task InitializeAsync(ILogger logger, CancellationToken cancellationToken)
+		public override Task InitializeAsync(ILogger logger, CancellationToken cancellationToken)
 		{
 			logger.LogDebug("SimpleTestExecutor.InitializeAsync()");
 			return Task.CompletedTask;
 		}
 
-		public Task<JobStepOutcome> RunAsync(BeginStepResponse step, ILogger logger,
+		public override Task<JobStepOutcome> RunAsync(BeginStepResponse step, ILogger logger,
 			CancellationToken cancellationToken)
 		{
 			logger.LogDebug("SimpleTestExecutor.RunAsync(Step: {Step})", step);
 			return _func(step, logger, cancellationToken);
 		}
 
-		public Task FinalizeAsync(ILogger logger, CancellationToken cancellationToken)
+		public override Task FinalizeAsync(ILogger logger, CancellationToken cancellationToken)
 		{
 			logger.LogDebug("SimpleTestExecutor.FinalizeAsync()");
 			return Task.CompletedTask;
 		}
+
+		protected override Task<bool> SetupAsync(BeginStepResponse step, ILogger logger, CancellationToken cancellationToken)
+		{
+			throw new NotImplementedException();
+		}
+
+		protected override Task<bool> ExecuteAsync(BeginStepResponse step, ILogger logger, CancellationToken cancellationToken)
+		{
+			throw new NotImplementedException();
+		}
 	}
 
-	class SimpleTestExecutorFactory : IExecutorFactory
+	class SimpleTestExecutorFactory : JobExecutorFactory
 	{
-		readonly IExecutor _executor;
+		readonly JobExecutor _executor;
 
-		public SimpleTestExecutorFactory(IExecutor executor)
+		public SimpleTestExecutorFactory(JobExecutor executor)
 		{
 			_executor = executor;
 		}
 
-		public IExecutor CreateExecutor(ISession session, ExecuteJobTask executeJobTask, BeginBatchResponse beginBatchResponse) => _executor;
+		public override JobExecutor CreateExecutor(ISession session, ExecuteJobTask executeJobTask, BeginBatchResponse beginBatchResponse) => _executor;
 	}
 }
\ No newline at end of file
diff --git a/Engine/Source/Programs/Horde/Horde.Agent.Tests/WorkerServiceTest.cs b/Engine/Source/Programs/Horde/Horde.Agent.Tests/WorkerServiceTest.cs
--- a/Engine/Source/Programs/Horde/Horde.Agent.Tests/WorkerServiceTest.cs
+++ b/Engine/Source/Programs/Horde/Horde.Agent.Tests/WorkerServiceTest.cs
@@ -13,7 +13,7 @@ using EpicGames.Core;
 using Google.Protobuf.WellKnownTypes;
 using Grpc.Core;
 using Grpc.Net.Client;
-using Horde.Agent.Execution.Interfaces;
+using Horde.Agent.Execution;
 using Horde.Agent.Leases;
 using Horde.Agent.Leases.Handlers;
 using Horde.Agent.Services;
@@ -37,7 +37,7 @@ namespace Horde.Agent.Tests
 	{
 		private ServiceCollection _serviceCollection;
 
-		internal static IExecutor NullExecutor = new SimpleTestExecutor(async (step, logger, cancellationToken) =>
+		internal static JobExecutor NullExecutor = new SimpleTestExecutor(async (step, logger, cancellationToken) =>
 		{
 			await Task.Delay(1, cancellationToken);
 			return JobStepOutcome.Success;
@@ -75,7 +75,7 @@ namespace Horde.Agent.Tests
 				using CancellationTokenSource cancelSource = new CancellationTokenSource();
 				using CancellationTokenSource stepCancelSource = new CancellationTokenSource();
 
-				IExecutor executor = new SimpleTestExecutor(async (stepResponse, logger, cancelToken) =>
+				JobExecutor executor = new SimpleTestExecutor(async (stepResponse, logger, cancelToken) =>
 				{
 					cancelSource.CancelAfter(10);
 					await Task.Delay(5000, cancelToken);
@@ -90,7 +90,7 @@ namespace Horde.Agent.Tests
 				using CancellationTokenSource cancelSource = new CancellationTokenSource();
 				using CancellationTokenSource stepCancelSource = new CancellationTokenSource();
 
-				IExecutor executor = new SimpleTestExecutor(async (stepResponse, logger, cancelToken) =>
+				JobExecutor executor = new SimpleTestExecutor(async (stepResponse, logger, cancelToken) =>
 				{
 					stepCancelSource.CancelAfter(10);
 					await Task.Delay(5000, cancelToken);
@@ -130,13 +130,13 @@ namespace Horde.Agent.Tests
 			GetStepResponse step2Res = new GetStepResponse(JobStepOutcome.Unspecified, JobStepState.Unspecified, true);
 			client.GetStepResponses[step2Req] = step2Res;
 
-			IExecutor executor = new SimpleTestExecutor(async (step, logger, cancelToken) =>
+			JobExecutor executor = new SimpleTestExecutor(async (step, logger, cancelToken) =>
 			{
 				await Task.Delay(50, cancelToken);
 				return JobStepOutcome.Success;
 			});
 
-			_serviceCollection.AddSingleton<IExecutorFactory>(x => new SimpleTestExecutorFactory(executor));
+			_serviceCollection.AddSingleton<JobExecutorFactory>(x => new SimpleTestExecutorFactory(executor));
 			using ServiceProvider serviceProvider = _serviceCollection.BuildServiceProvider();
 
 			JobHandler jobHandler = serviceProvider.GetRequiredService<JobHandler>();
@@ -162,13 +162,13 @@ namespace Horde.Agent.Tests
 		[TestMethod]
 		public async Task PollForStepAbortFailureTest()
 		{
-			IExecutor executor = new SimpleTestExecutor(async (step, logger, cancelToken) =>
+			JobExecutor executor = new SimpleTestExecutor(async (step, logger, cancelToken) =>
 			{
 				await Task.Delay(50, cancelToken);
 				return JobStepOutcome.Success;
 			});
 
-			_serviceCollection.AddSingleton<IExecutorFactory>(x => new SimpleTestExecutorFactory(executor));
+			_serviceCollection.AddSingleton<JobExecutorFactory>(x => new SimpleTestExecutorFactory(executor));
 			using ServiceProvider serviceProvider = _serviceCollection.BuildServiceProvider();
 
 			JobHandler jobHandler = serviceProvider.GetRequiredService<JobHandler>();
@@ -200,13 +200,13 @@ namespace Horde.Agent.Tests
 		[TestMethod]
 		public async Task Shutdown()
 		{
-			IExecutor executor = new SimpleTestExecutor(async (step, logger, cancellationToken) =>
+			JobExecutor executor = new SimpleTestExecutor(async (step, logger, cancellationToken) =>
 			{
 				await Task.Delay(50, cancellationToken);
 				return JobStepOutcome.Success;
 			});
 
-			_serviceCollection.AddSingleton<IExecutorFactory>(x => new SimpleTestExecutorFactory(executor));
+			_serviceCollection.AddSingleton<JobExecutorFactory>(x => new SimpleTestExecutorFactory(executor));
 			using ServiceProvider serviceProvider = _serviceCollection.BuildServiceProvider();
 
 			using CancellationTokenSource cts = new ();
diff --git a/Engine/Source/Programs/Horde/Horde.Agent/Execution/JobExecutor.cs b/Engine/Source/Programs/Horde/Horde.Agent/Execution/JobExecutor.cs
--- a/Engine/Source/Programs/Horde/Horde.Agent/Execution/JobExecutor.cs
+++ b/Engine/Source/Programs/Horde/Horde.Agent/Execution/JobExecutor.cs
@@ -262,7 +262,7 @@ namespace Horde.Agent.Execution
 			}
 		}
 
-		public async Task<JobStepOutcome> RunAsync(BeginStepResponse step, ILogger logger, CancellationToken cancellationToken)
+		public virtual async Task<JobStepOutcome> RunAsync(BeginStepResponse step, ILogger logger, CancellationToken cancellationToken)
 		{
 			if (step.Name == "Setup Build")
 			{
