Horde: Add a unique index on the Slack collection (renaming it to SlackV2) to prevent consistency errors.

diff --git a/Engine/Source/Programs/Horde/Horde.Build/Notifications/Sinks/SlackNotificationSink.cs b/Engine/Source/Programs/Horde/Horde.Build/Notifications/Sinks/SlackNotificationSink.cs
--- a/Engine/Source/Programs/Horde/Horde.Build/Notifications/Sinks/SlackNotificationSink.cs
+++ b/Engine/Source/Programs/Horde/Horde.Build/Notifications/Sinks/SlackNotificationSink.cs
@@ -240,8 +240,8 @@ namespace Horde.Build.Notifications.Sinks
 			_externalIssueService = externalIssueService;
 			_environment = environment;
 			_settings = settings.Value;
-			_messageStates = mongoService.Database.GetCollection<MessageStateDocument>("Slack");
-			_slackUsers = mongoService.Database.GetCollection<SlackUserDocument>("Slack.UsersV2");
+			_messageStates = mongoService.GetCollection<MessageStateDocument>("SlackV2", keys => keys.Ascending(x => x.Recipient).Ascending(x => x.EventId), unique: true);
+			_slackUsers = mongoService.GetCollection<SlackUserDocument>("Slack.UsersV2");
 			_escalateIssues = new RedisSortedSet<int>(redisService.ConnectionPool, "slack/escalate");
 			_clock = clock;
 			_logger = logger;
