Horde: Add settings flag for activating AWS specific features

diff --git a/Engine/Source/Programs/Horde/Horde.Build.Tests/Fleet/AutoscaleServiceV2Test.cs b/Engine/Source/Programs/Horde/Horde.Build.Tests/Fleet/AutoscaleServiceV2Test.cs
--- a/Engine/Source/Programs/Horde/Horde.Build.Tests/Fleet/AutoscaleServiceV2Test.cs
+++ b/Engine/Source/Programs/Horde/Horde.Build.Tests/Fleet/AutoscaleServiceV2Test.cs
@@ -65,6 +65,21 @@ namespace Horde.Build.Tests.Fleet
 
 		public string Name { get; } = "PoolSizeStrategySpy";
 	}
+	
+	public class StubFleetManagerFactory : IFleetManagerFactory
+	{
+		private readonly IFleetManager _fleetManager;
+
+		public StubFleetManagerFactory(IFleetManager fleetManager)
+		{
+			_fleetManager = fleetManager;
+		}
+
+		public IFleetManager CreateFleetManager(FleetManagerType type, string config = "{}")
+		{
+			return _fleetManager;
+		}
+	}
 
 	[TestClass]
 	public class AutoscaleServiceV2Test : TestSetup
@@ -160,13 +175,10 @@ namespace Horde.Build.Tests.Fleet
 			IOptions<ServerSettings> serverSettingsOpt = ServiceProvider.GetRequiredService<IOptions<ServerSettings>>();
 			serverSettingsOpt.Value.FleetManagerV2 = FleetManagerType.AwsReuse;
 			
-			// Empty mocks to satisfy the constructor below
-			Mock<IAmazonEC2> mockEc2 = new ();
-			Mock<IAmazonAutoScaling> mockAwsAutoScaling = new ();
-
-			AutoscaleServiceV2 service = new (AgentCollection, GraphCollection, JobCollection, LeaseCollection, PoolCollection,
-				StreamService, _dogStatsD, mockEc2.Object, mockAwsAutoScaling.Object, Clock, Cache, serverSettingsOpt, loggerFactory,
-				_ => fleetManager);
+			AutoscaleServiceV2 service = new(
+				AgentCollection, GraphCollection, JobCollection, LeaseCollection, PoolCollection, StreamService, _dogStatsD,
+				new StubFleetManagerFactory(fleetManager), Clock, Cache, serverSettingsOpt, loggerFactory.CreateLogger<AutoscaleServiceV2>());
+				
 			return service;
 		}
 	}
diff --git a/Engine/Source/Programs/Horde/Horde.Build.Tests/TestSetup.cs b/Engine/Source/Programs/Horde/Horde.Build.Tests/TestSetup.cs
--- a/Engine/Source/Programs/Horde/Horde.Build.Tests/TestSetup.cs
+++ b/Engine/Source/Programs/Horde/Horde.Build.Tests/TestSetup.cs
@@ -153,6 +153,7 @@ namespace Horde.Build.Tests
 
 			settings.AdminClaimType = HordeClaimTypes.Role;
 			settings.AdminClaimValue = "app-horde-admins";
+			settings.WithAws = true;
 		}
 
 		protected override void ConfigureServices(IServiceCollection services)
@@ -206,6 +207,7 @@ namespace Horde.Build.Tests
 			// Empty mocked object to satisfy basic test runs
 			services.AddSingleton<IAmazonEC2>(sp => new Mock<IAmazonEC2>().Object);
 			services.AddSingleton<IAmazonAutoScaling>(sp => new Mock<IAmazonAutoScaling>().Object);
+			services.AddSingleton<IFleetManagerFactory, FleetManagerFactory>();
 
 			services.AddSingleton<AclService>();
 			services.AddSingleton<AgentService>();
diff --git a/Engine/Source/Programs/Horde/Horde.Build/Agents/Fleet/AutoscaleServiceV2.cs b/Engine/Source/Programs/Horde/Horde.Build/Agents/Fleet/AutoscaleServiceV2.cs
--- a/Engine/Source/Programs/Horde/Horde.Build/Agents/Fleet/AutoscaleServiceV2.cs
+++ b/Engine/Source/Programs/Horde/Horde.Build/Agents/Fleet/AutoscaleServiceV2.cs
@@ -47,8 +47,7 @@ namespace Horde.Build.Agents.Fleet
 		private readonly IPoolCollection _poolCollection;
 		private readonly StreamService _streamService;
 		private readonly IDogStatsd _dogStatsd;
-		private readonly IAmazonEC2 _ec2;
-		private readonly IAmazonAutoScaling _awsAutoScaling;
+		private readonly IFleetManagerFactory _fleetManagerFactory2;
 		private readonly IClock _clock;
 		private readonly IMemoryCache _cache;
 		private readonly ITicker _ticker;
@@ -57,11 +56,7 @@ namespace Horde.Build.Agents.Fleet
 		private readonly TimeSpan _defaultScaleInCooldown;
 		private readonly IOptions<ServerSettings> _settings;
 		private readonly ILogger<AutoscaleServiceV2> _logger;
-		private readonly ILoggerFactory _loggerFactory;
 		
-		/// <summary>Allow overriding the fleet manager during testing</summary>
-		private readonly Func<IPool, IFleetManager> _fleetManagerFactory;
-
 		/// <summary>
 		/// Constructor
 		/// </summary>
@@ -73,13 +68,11 @@ namespace Horde.Build.Agents.Fleet
 			IPoolCollection poolCollection,
 			StreamService streamService,
 			IDogStatsd dogStatsd,
-			IAmazonEC2 ec2,
-			IAmazonAutoScaling awsAutoScaling,
+			IFleetManagerFactory fleetManagerFactory,
 			IClock clock,
 			IMemoryCache cache,
 			IOptions<ServerSettings> settings,
-			ILoggerFactory loggerFactory,
-			Func<IPool, IFleetManager>? fleetManagerFactory = null)
+			ILogger<AutoscaleServiceV2> logger)
 		{
 			_agentCollection = agentCollection;
 			_graphCollection = graphCollection;
@@ -88,18 +81,16 @@ namespace Horde.Build.Agents.Fleet
 			_poolCollection = poolCollection;
 			_streamService = streamService;
 			_dogStatsd = dogStatsd;
-			_ec2 = ec2;
-			_awsAutoScaling = awsAutoScaling;
+			_fleetManagerFactory2 = fleetManagerFactory;
 			_clock = clock;
 			_cache = cache;
-			_logger = loggerFactory.CreateLogger<AutoscaleServiceV2>();
-			_loggerFactory = loggerFactory;
+			_logger = logger;
 			_ticker = clock.AddSharedTicker<AutoscaleServiceV2>(TimeSpan.FromMinutes(5.0), TickLeaderAsync, _logger);
 			_tickerHighFrequency = clock.AddSharedTicker("AutoscaleServiceV2.TickHighFrequency", TimeSpan.FromSeconds(30), TickHighFrequencyAsync, _logger);
 			_settings = settings;
 			_defaultScaleOutCooldown = TimeSpan.FromSeconds(settings.Value.AgentPoolScaleOutCooldownSeconds);
 			_defaultScaleInCooldown = TimeSpan.FromSeconds(settings.Value.AgentPoolScaleInCooldownSeconds);
-			_fleetManagerFactory = fleetManagerFactory ?? CreateFleetManager;
+
 		}
 
 		/// <inheritdoc/>
@@ -193,7 +184,7 @@ namespace Horde.Build.Agents.Fleet
 			int desiredAgentCount = poolSizeData.DesiredAgentCount.Value;
 			int deltaAgentCount = desiredAgentCount - currentAgentCount;
 
-			IFleetManager fleetManager = _fleetManagerFactory(poolSizeData.Pool);
+			IFleetManager fleetManager = CreateFleetManager(poolSizeData.Pool);
 
 			_logger.LogInformation("{PoolName,-48} Current={Current,4} Target={Target,4} Delta={Delta,4} Status={Status} FleetManager={FleetManager}",
 				pool.Name, currentAgentCount, desiredAgentCount, deltaAgentCount, poolSizeData.StatusMessage, fleetManager.GetType().Name);
@@ -282,60 +273,11 @@ namespace Horde.Build.Agents.Fleet
 			{
 				if (info.Condition == null || info.Condition.Evaluate(GetPropValues))
 				{
-					return CreateFleetManager(info.Type, info.Config);
+					return _fleetManagerFactory2.CreateFleetManager(info.Type, info.Config);
 				}
 			}
 
-			return GetDefaultFleetManager();
-		}
-		
-		/// <summary>
-		/// Instantiate the default fleet manager  />
-		/// </summary>
-		/// <returns>A fleet manager with parameters from server settings</returns>
-		/// <exception cref="ArgumentException">If fleet manager could not be instantiated</exception>
-		public IFleetManager GetDefaultFleetManager()
-		{
-			return CreateFleetManager(_settings.Value.FleetManagerV2, _settings.Value.FleetManagerV2Config ?? "{}");
-		}
-		
-		/// <summary>
-		/// Create a fleet manager
-		/// </summary>
-		/// <param name="type">Type of fleet manager</param>
-		/// <param name="config">Config as a serialized JSON string</param>
-		/// <returns>An instantiated fleet manager with parameters loaded from config</returns>
-		/// <exception cref="ArgumentException">If fleet manager could not be instantiated</exception>
-		private IFleetManager CreateFleetManager(FleetManagerType type, string config)
-		{
-			switch (type)
-			{
-				case FleetManagerType.NoOp:
-					return new NoOpFleetManager(_loggerFactory.CreateLogger<NoOpFleetManager>());
-				case FleetManagerType.Aws:
-					return new AwsFleetManager(_ec2, _agentCollection, DeserializeSettings<AwsFleetManagerSettings>(config), _loggerFactory.CreateLogger<AwsFleetManager>());
-				case FleetManagerType.AwsReuse:
-					return new AwsReuseFleetManager(_ec2, _agentCollection, DeserializeSettings<AwsReuseFleetManagerSettings>(config), _loggerFactory.CreateLogger<AwsReuseFleetManager>());
-				case FleetManagerType.AwsAsg:
-					return new AwsAsgFleetManager(_awsAutoScaling, DeserializeSettings<AwsAsgSettings>(config), _loggerFactory.CreateLogger<AwsAsgFleetManager>());	
-				default:
-					throw new ArgumentException("Unknown fleet manager type " + type);
-			}
-		}
-		
-		private static T DeserializeSettings<T>(string config)
-		{
-			if (String.IsNullOrEmpty(config)) { config = "{}"; }
-			try
-			{
-				T? settings = JsonSerializer.Deserialize<T>(config);
-				if (settings == null) throw new NullReferenceException($"Unable to deserialize");
-				return settings;
-			}
-			catch (ArgumentException e)
-			{
-				throw new ArgumentException($"Unable to deserialize {typeof(T)} config: '{config}'", e);
-			}
+			return _fleetManagerFactory2.CreateFleetManager(FleetManagerType.Default, "{}");
 		}
 
 		/// <summary>
diff --git a/Engine/Source/Programs/Horde/Horde.Build/Agents/Fleet/FleetManagerFactory.cs b/Engine/Source/Programs/Horde/Horde.Build/Agents/Fleet/FleetManagerFactory.cs
new file mode 100644
--- /dev/null
+++ b/Engine/Source/Programs/Horde/Horde.Build/Agents/Fleet/FleetManagerFactory.cs
@@ -0,0 +1,112 @@
+// Copyright Epic Games, Inc. All Rights Reserved.
+
+using System;
+using System.Text.Json;
+using Amazon.AutoScaling;
+using Amazon.EC2;
+using Horde.Build.Agents.Fleet.Providers;
+using Microsoft.Extensions.DependencyInjection;
+using Microsoft.Extensions.Logging;
+using Microsoft.Extensions.Options;
+
+namespace Horde.Build.Agents.Fleet;
+
+/// <summary>
+/// Interface for fleet manager factory
+/// </summary>
+public interface IFleetManagerFactory
+{
+	/// <summary>
+	/// Create a fleet manager
+	/// </summary>
+	/// <param name="type">Type of fleet manager</param>
+	/// <param name="config">Config as a serialized JSON string</param>
+	/// <returns>An instantiated fleet manager with parameters loaded from config</returns>
+	/// <exception cref="ArgumentException">If fleet manager could not be instantiated</exception>
+	public IFleetManager CreateFleetManager(FleetManagerType type, string config = "{}");
+}
+
+/// <summary>
+/// Factory creating instances of fleet managers
+/// </summary>
+public sealed class FleetManagerFactory : IFleetManagerFactory
+{
+	private readonly IAgentCollection _agentCollection;
+	private readonly IServiceProvider _provider;
+	private readonly IOptionsMonitor<ServerSettings> _settings;
+	private readonly ILoggerFactory _loggerFactory;
+	
+	private IAmazonAutoScaling? _awsAutoScaling;
+	private IAmazonEC2? _awsEc2;
+	
+	/// <summary>
+	/// Constructor
+	/// </summary>
+	public FleetManagerFactory(IAgentCollection agentCollection, IServiceProvider provider, IOptionsMonitor<ServerSettings> settings, ILoggerFactory loggerFactory)
+	{
+		_agentCollection = agentCollection;
+		_provider = provider;
+		_settings = settings;
+		_loggerFactory = loggerFactory;
+	}
+	
+	/// <inheritdoc/>
+	public IFleetManager CreateFleetManager(FleetManagerType type, string config)
+	{
+		return type switch
+		{
+			FleetManagerType.Default =>
+				CreateFleetManager(_settings.CurrentValue.FleetManagerV2, _settings.CurrentValue.FleetManagerV2Config ?? "{}"),
+			FleetManagerType.NoOp =>
+				new NoOpFleetManager(_loggerFactory.CreateLogger<NoOpFleetManager>()),
+			FleetManagerType.Aws =>
+				new AwsFleetManager(GetAwsEc2(type), _agentCollection, DeserializeSettings<AwsFleetManagerSettings>(config), _loggerFactory.CreateLogger<AwsFleetManager>()),
+			FleetManagerType.AwsReuse =>
+				new AwsReuseFleetManager(GetAwsEc2(type), _agentCollection, DeserializeSettings<AwsReuseFleetManagerSettings>(config), _loggerFactory.CreateLogger<AwsReuseFleetManager>()),
+			FleetManagerType.AwsAsg =>
+				new AwsAsgFleetManager(GetAwsAutoScaling(type), DeserializeSettings<AwsAsgSettings>(config), _loggerFactory.CreateLogger<AwsAsgFleetManager>()),
+			_ => throw new ArgumentException("Unknown fleet manager type " + type)
+		};
+	}
+	
+	private static T DeserializeSettings<T>(string config)
+	{
+		if (String.IsNullOrEmpty(config)) { config = "{}"; }
+		try
+		{
+			T? settings = JsonSerializer.Deserialize<T>(config);
+			if (settings == null) throw new NullReferenceException($"Unable to deserialize");
+			return settings;
+		}
+		catch (ArgumentException e)
+		{
+			throw new ArgumentException($"Unable to deserialize {typeof(T)} config: '{config}'", e);
+		}
+	}
+
+	private IAmazonEC2 GetAwsEc2(FleetManagerType type)
+	{
+		if (_awsEc2 != null) return _awsEc2;
+		
+		_awsEc2 = _provider.GetService<IAmazonEC2>();
+		if (_settings.CurrentValue.WithAws == false || _awsEc2 == null)
+		{
+			throw new ArgumentException($"Unable to create fleet manager {type} requiring AWS specific classes. Check that setting '{nameof(ServerSettings.WithAws)}' is enabled");
+		}
+
+		return _awsEc2;
+	}
+	
+	private IAmazonAutoScaling GetAwsAutoScaling(FleetManagerType type)
+	{
+		if (_awsAutoScaling != null) return _awsAutoScaling;
+		
+		_awsAutoScaling = _provider.GetService<IAmazonAutoScaling>();
+		if (_settings.CurrentValue.WithAws == false || _awsAutoScaling == null)
+		{
+			throw new ArgumentException($"Unable to create fleet manager {type} requiring AWS specific classes. Check that setting '{nameof(ServerSettings.WithAws)}' is enabled");
+		}
+
+		return _awsAutoScaling;
+	}
+}
diff --git a/Engine/Source/Programs/Horde/Horde.Build/Agents/Fleet/IFleetManager.cs b/Engine/Source/Programs/Horde/Horde.Build/Agents/Fleet/IFleetManager.cs
--- a/Engine/Source/Programs/Horde/Horde.Build/Agents/Fleet/IFleetManager.cs
+++ b/Engine/Source/Programs/Horde/Horde.Build/Agents/Fleet/IFleetManager.cs
@@ -13,6 +13,11 @@ namespace Horde.Build.Agents.Fleet
 	/// </summary>
 	public enum FleetManagerType
 	{
+		/// <summary>
+		/// Default fleet manager
+		/// </summary>
+		Default,
+		
 		/// <summary>
 		/// <see cref="NoOpFleetManager" />
 		/// </summary>
diff --git a/Engine/Source/Programs/Horde/Horde.Build/ServerSettings.cs b/Engine/Source/Programs/Horde/Horde.Build/ServerSettings.cs
--- a/Engine/Source/Programs/Horde/Horde.Build/ServerSettings.cs
+++ b/Engine/Source/Programs/Horde/Horde.Build/ServerSettings.cs
@@ -570,6 +570,11 @@ namespace Horde.Build
 		/// </summary>
 		public bool WithDatadog { get; set; }
 
+		/// <summary>
+		/// Whether to enable Amazon Web Services (AWS) specific features
+		/// </summary>
+		public bool WithAws { get; set; } = false;
+
 		/// <summary>
 		/// Path to the root config file
 		/// </summary>
diff --git a/Engine/Source/Programs/Horde/Horde.Build/Startup.cs b/Engine/Source/Programs/Horde/Horde.Build/Startup.cs
--- a/Engine/Source/Programs/Horde/Horde.Build/Startup.cs
+++ b/Engine/Source/Programs/Horde/Horde.Build/Startup.cs
@@ -370,19 +370,11 @@ namespace Horde.Build
 			services.AddSingleton(typeof(ISingletonDocument<>), typeof(SingletonDocument<>));
 
 			services.AddSingleton<AutoscaleServiceV2>();
+			services.AddSingleton<IFleetManagerFactory, FleetManagerFactory>();
 			
 			// Associate IFleetManager interface with the default implementation from config for convenience
 			// Though most fleet managers are created on a per-pool basis
-			services.AddSingleton<IFleetManager>(ctx =>
-			{
-				AutoscaleServiceV2? autoScale = ctx.GetService<AutoscaleServiceV2>();
-				if (autoScale == null)
-				{
-					throw new InvalidOperationException($"{nameof(AutoscaleServiceV2)} not registered!");
-				}
-
-				return autoScale.GetDefaultFleetManager();
-			});
+			services.AddSingleton<IFleetManager>(ctx => ctx.GetRequiredService<IFleetManagerFactory>().CreateFleetManager(FleetManagerType.Default));
 
 			services.AddSingleton<AclService>();
 			services.AddSingleton<AgentService>();			
@@ -448,17 +440,20 @@ namespace Horde.Build
 			services.AddSingleton(sp => CreateStorageBackend(sp, settings.ArtifactStorage).ForType<ArtifactCollection>());
 
 			services.AddHordeStorage(settings => configSection.GetSection("Storage").Bind(settings));
-			
-			AWSOptions awsOptions = Configuration.GetAWSOptions();
-			services.AddDefaultAWSOptions(awsOptions);
-			if (awsOptions.Region == null && Environment.GetEnvironmentVariable("AWS_REGION") == null)
+
+			if (settings.WithAws)
 			{
-				awsOptions.Region = RegionEndpoint.USEast1;
+				AWSOptions awsOptions = Configuration.GetAWSOptions();
+				services.AddDefaultAWSOptions(awsOptions);
+				if (awsOptions.Region == null && Environment.GetEnvironmentVariable("AWS_REGION") == null)
+				{
+					awsOptions.Region = RegionEndpoint.USEast1;
+				}
+				
+				services.AddAWSService<IAmazonCloudWatch>();
+				services.AddAWSService<IAmazonAutoScaling>();
+				services.AddAWSService<IAmazonEC2>();
 			}
-			
-			services.AddAWSService<IAmazonCloudWatch>();
-			services.AddAWSService<IAmazonAutoScaling>();
-			services.AddAWSService<IAmazonEC2>();
 
 			ConfigureLogStorage(services);
 
