Horde: Add logging to make it clear when an issue is created vs matched
#rnx

diff --git a/Engine/Source/Programs/Horde/Horde.Build/Issues/IssueService.cs b/Engine/Source/Programs/Horde/Horde.Build/Issues/IssueService.cs
--- a/Engine/Source/Programs/Horde/Horde.Build/Issues/IssueService.cs
+++ b/Engine/Source/Programs/Horde/Horde.Build/Issues/IssueService.cs
@@ -1219,6 +1219,11 @@ namespace Horde.Build.Issues
 			{
 				string summary = GetSummary(span.Fingerprint, span.FirstFailure.Severity);
 				issue = await _issueCollection.AddIssueAsync(summary);
+				_logger.LogInformation("Created issue {IssueId}", issue.Id);
+			}
+			else
+			{
+				_logger.LogInformation("Natched to issue {IssueId}", issue.Id);
 			}
 			return issue;
 		}
