Horde: Improve settings use in pool size strategies

* Move settings from IPool to strategy settings for LeaseUtilization
* Move hardcoded default values to settings for JobQueue
* Log settings used in the message field
* Properly initializes settings when created through old-style (fields from IPool)

diff --git a/Engine/Source/Programs/Horde/Horde.Build.Tests/Fleet/JobQueueStrategyTest.cs b/Engine/Source/Programs/Horde/Horde.Build.Tests/Fleet/JobQueueStrategyTest.cs
--- a/Engine/Source/Programs/Horde/Horde.Build.Tests/Fleet/JobQueueStrategyTest.cs
+++ b/Engine/Source/Programs/Horde/Horde.Build.Tests/Fleet/JobQueueStrategyTest.cs
@@ -28,7 +28,7 @@ namespace Horde.Build.Tests.Fleet
 		public async Task GetPoolQueueSizes()
 		{
 			(JobQueueStrategy strategy, PoolSizeData poolSizeData) = await SetUpJobsAsync(1, 5);
-			await Clock.AdvanceAsync(strategy.ReadyTimeThreshold + TimeSpan.FromSeconds(5));
+			await Clock.AdvanceAsync(TimeSpan.FromSeconds(strategy.Settings.ReadyTimeThresholdSec) + TimeSpan.FromSeconds(5));
 			Dictionary<PoolId, int> poolQueueSizes = await strategy.GetPoolQueueSizesAsync(Clock.UtcNow - TimeSpan.FromHours(2));
 			Assert.AreEqual(1, poolQueueSizes.Count);
 			Assert.AreEqual(5, poolQueueSizes[poolSizeData.Pool.Id]);
@@ -86,7 +86,7 @@ namespace Horde.Build.Tests.Fleet
 		{
 			(JobQueueStrategy strategy, PoolSizeData poolSizeData) = await SetUpJobsAsync(1, numBatchesReady, numAgents);
 			TimeSpan timeToWait = waitedBeyondThreshold
-				? strategy.ReadyTimeThreshold + TimeSpan.FromSeconds(5)
+				? TimeSpan.FromSeconds(strategy.Settings.ReadyTimeThresholdSec) + TimeSpan.FromSeconds(5)
 				: TimeSpan.FromSeconds(15);
 			
 			await Clock.AdvanceAsync(timeToWait);
diff --git a/Engine/Source/Programs/Horde/Horde.Build.Tests/Fleet/LeaseUtilizationStrategyTest.cs b/Engine/Source/Programs/Horde/Horde.Build.Tests/Fleet/LeaseUtilizationStrategyTest.cs
--- a/Engine/Source/Programs/Horde/Horde.Build.Tests/Fleet/LeaseUtilizationStrategyTest.cs
+++ b/Engine/Source/Programs/Horde/Horde.Build.Tests/Fleet/LeaseUtilizationStrategyTest.cs
@@ -97,7 +97,7 @@ namespace Horde.Build.Tests.Fleet
 
 		private async Task AssertPoolSizeAsync(IPool pool, int expectedNumAgents)
 		{
-			LeaseUtilizationStrategy strategy = new (AgentCollection, PoolCollection, LeaseCollection, Clock, Cache, new());
+			IPoolSizeStrategy strategy = AutoscaleServiceV2.CreatePoolSizeStrategy(pool);
 			List<PoolSizeData> output = await strategy.CalcDesiredPoolSizesAsync(new() { new(pool, _poolAgents, null) });
 			Assert.AreEqual(1, output.Count);
 			Assert.AreEqual(expectedNumAgents, output[0].DesiredAgentCount);
diff --git a/Engine/Source/Programs/Horde/Horde.Build/Agents/Fleet/AutoscaleServiceV2.cs b/Engine/Source/Programs/Horde/Horde.Build/Agents/Fleet/AutoscaleServiceV2.cs
--- a/Engine/Source/Programs/Horde/Horde.Build/Agents/Fleet/AutoscaleServiceV2.cs
+++ b/Engine/Source/Programs/Horde/Horde.Build/Agents/Fleet/AutoscaleServiceV2.cs
@@ -321,7 +321,10 @@ namespace Horde.Build.Agents.Fleet
 				case PoolSizeStrategy.JobQueue:
 					return new JobQueueStrategy(_jobCollection, _graphCollection, _streamService, _clock, _cache, pool.JobQueueSettings);
 				case PoolSizeStrategy.LeaseUtilization:
-					return new LeaseUtilizationStrategy(_agentCollection, _poolCollection, _leaseCollection, _clock, _cache, new());
+					LeaseUtilizationSettings luSettings = new();
+					if (pool.MinAgents != null) luSettings.MinAgents = pool.MinAgents.Value;
+					if (pool.NumReserveAgents != null) luSettings.NumReserveAgents = pool.NumReserveAgents.Value;
+					return new LeaseUtilizationStrategy(_agentCollection, _poolCollection, _leaseCollection, _clock, _cache, luSettings);
 				case PoolSizeStrategy.NoOp:
 					return new NoOpPoolSizeStrategy();
 				default:
diff --git a/Engine/Source/Programs/Horde/Horde.Build/Agents/Fleet/JobQueueStrategy.cs b/Engine/Source/Programs/Horde/Horde.Build/Agents/Fleet/JobQueueStrategy.cs
--- a/Engine/Source/Programs/Horde/Horde.Build/Agents/Fleet/JobQueueStrategy.cs
+++ b/Engine/Source/Programs/Horde/Horde.Build/Agents/Fleet/JobQueueStrategy.cs
@@ -3,6 +3,7 @@
 using System;
 using System.Collections.Generic;
 using System.Linq;
+using System.Text;
 using System.Text.Json.Serialization;
 using System.Threading.Tasks;
 using Horde.Build.Agents.Pools;
@@ -37,6 +38,19 @@ namespace Horde.Build.Agents.Fleet
 		/// </summary>
 		public double ScaleInFactor { get; set; } = 0.9;
 
+		/// <summary>
+		/// How far back in time to look for job batches (that potentially are in the queue)
+		/// </summary>
+		public int SamplePeriodMin { get; set; } = 60 * 24 * 5; // 5 days
+
+		/// <summary>
+		/// Time spent in ready state before considered truly waiting for an agent
+		///
+		/// A job batch can be in ready state before getting picked up and executed.
+		/// This threshold will help ensure only batches that have been waiting longer than this value will be considered.
+		/// </summary>
+		public int ReadyTimeThresholdSec { get; set; } = 45;
+
 		/// <summary>
 		/// Constructor used for JSON serialization
 		/// </summary>
@@ -55,6 +69,17 @@ namespace Horde.Build.Agents.Fleet
 			ScaleOutFactor = scaleOutFactor.GetValueOrDefault(ScaleOutFactor);
 			ScaleInFactor = scaleInFactor.GetValueOrDefault(ScaleInFactor);
 		}
+		
+		/// <inheritdoc />
+		public override string ToString()
+		{
+			StringBuilder sb = new (50);
+			sb.AppendFormat("{0}={1} ", nameof(ScaleOutFactor), ScaleOutFactor);
+			sb.AppendFormat("{0}={1} ", nameof(ScaleInFactor), ScaleInFactor);
+			sb.AppendFormat("{0}={1} ", nameof(SamplePeriodMin), SamplePeriodMin);
+			sb.AppendFormat("{0}={1} ", nameof(ReadyTimeThresholdSec), ReadyTimeThresholdSec);
+			return sb.ToString();
+		}
 	}
 	
 	/// <summary>
@@ -74,21 +99,6 @@ namespace Horde.Build.Agents.Fleet
 		private readonly IClock _clock;
 		private readonly IMemoryCache _cache;
 		
-		/// <summary>
-		/// How far back in time to look for job batches (that potentially are in the queue)
-		/// </summary>
-		private readonly TimeSpan _samplePeriod = TimeSpan.FromDays(5);
-		
-		/// <summary>
-		/// Time spent in ready state before considered truly waiting for an agent
-		///
-		/// A job batch can be in ready state before getting picked up and executed.
-		/// This threshold will help ensure only batches that have been waiting longer than this value will be considered.
-		/// </summary>
-		internal readonly TimeSpan ReadyTimeThreshold = TimeSpan.FromSeconds(45.0);
-
-		private readonly JobQueueSettings _defaultPoolSettings = new();
-
 		/// <summary>
 		/// Constructor
 		/// </summary>
@@ -98,16 +108,14 @@ namespace Horde.Build.Agents.Fleet
 		/// <param name="clock"></param>
 		/// <param name="cache"></param>
 		/// <param name="settings"></param>
-		/// <param name="samplePeriod">Time period for each sample</param>
-		public JobQueueStrategy(IJobCollection jobs, IGraphCollection graphs, StreamService streamService, IClock clock, IMemoryCache cache, JobQueueSettings? settings = null, TimeSpan? samplePeriod = null)
+		public JobQueueStrategy(IJobCollection jobs, IGraphCollection graphs, StreamService streamService, IClock clock, IMemoryCache cache, JobQueueSettings? settings = null)
 		{
 			_jobs = jobs;
 			_graphs = graphs;
 			_streamService = streamService;
 			_clock = clock;
 			_cache = cache;
-			_samplePeriod = samplePeriod ?? _samplePeriod;
-			Settings = settings ?? _defaultPoolSettings;
+			Settings = settings ?? new JobQueueSettings();
 		}
 
 		/// <inheritdoc/>
@@ -136,7 +144,7 @@ namespace Horde.Build.Agents.Fleet
 				{
 					continue;
 				}
-				if (waitTime.Value < ReadyTimeThreshold)
+				if (waitTime.Value < TimeSpan.FromSeconds(Settings.ReadyTimeThresholdSec))
 				{
 					continue;
 				}
@@ -177,11 +185,10 @@ namespace Horde.Build.Agents.Fleet
 		/// <inheritdoc/>
 		public async Task<List<PoolSizeData>> CalcDesiredPoolSizesAsync(List<PoolSizeData> pools)
 		{
-			DateTimeOffset minCreateTime = _clock.UtcNow - _samplePeriod;
-			Dictionary<PoolId, int> poolQueueSizes;// = await GetPoolQueueSizesAsync(minCreateTime);
+			DateTimeOffset minCreateTime = _clock.UtcNow - TimeSpan.FromMinutes(Settings.SamplePeriodMin);
 
 			// Cache pool queue sizes for a short while for faster runs when many pools are scaled
-			if (!_cache.TryGetValue(CacheKey, out poolQueueSizes))
+			if (!_cache.TryGetValue(CacheKey, out Dictionary<PoolId, int> poolQueueSizes))
 			{
 				// Pool sizes haven't been cached, update them (might happen from multiple tasks but that is fine)
 				poolQueueSizes = await GetPoolQueueSizesAsync(minCreateTime);
@@ -190,18 +197,17 @@ namespace Horde.Build.Agents.Fleet
 
 			return pools.Select(current =>
 			{
-				JobQueueSettings settings = current.Pool.JobQueueSettings ?? _defaultPoolSettings;
 				poolQueueSizes.TryGetValue(current.Pool.Id, out int queueSize);
 				if (queueSize > 0)
 				{
-					int additionalAgentCount = (int)Math.Ceiling(queueSize * settings.ScaleOutFactor);
+					int additionalAgentCount = (int)Math.Ceiling(queueSize * Settings.ScaleOutFactor);
 					int desiredAgentCount = current.Agents.Count + additionalAgentCount;
-					return new PoolSizeData(current.Pool, current.Agents, desiredAgentCount, $"QueueSize={queueSize}");
+					return new PoolSizeData(current.Pool, current.Agents, desiredAgentCount, $"QueueSize={queueSize} " + Settings);
 				}
 				else
 				{
-					int desiredAgentCount = (int)(current.Agents.Count * settings.ScaleInFactor);
-					return new PoolSizeData(current.Pool, current.Agents, desiredAgentCount, "Empty job queue");
+					int desiredAgentCount = (int)(current.Agents.Count * Settings.ScaleInFactor);
+					return new PoolSizeData(current.Pool, current.Agents, desiredAgentCount, "QueueSize={queueSize} " + Settings);
 				}
 			}).ToList();
 		}
diff --git a/Engine/Source/Programs/Horde/Horde.Build/Agents/Fleet/LeaseUtilizationStrategy.cs b/Engine/Source/Programs/Horde/Horde.Build/Agents/Fleet/LeaseUtilizationStrategy.cs
--- a/Engine/Source/Programs/Horde/Horde.Build/Agents/Fleet/LeaseUtilizationStrategy.cs
+++ b/Engine/Source/Programs/Horde/Horde.Build/Agents/Fleet/LeaseUtilizationStrategy.cs
@@ -37,7 +37,29 @@ namespace Horde.Build.Agents.Fleet
 		/// <summary>
 		/// Min number of samples for a valid result
 		/// </summary>
-		public int NumSamplesForResult { get; set;  } = 9;
+		public int NumSamplesForResult { get; set; } = 9;
+
+		/// <summary>
+		/// The minimum number of agents to keep in the pool
+		/// </summary>
+		public int MinAgents { get; set; } = 1;
+
+		/// <summary>
+		/// The minimum number of idle agents to hold in reserve
+		/// </summary>
+		public int NumReserveAgents { get; set; } = 5;
+
+		/// <inheritdoc />
+		public override string ToString()
+		{
+			StringBuilder sb = new (150);
+			sb.AppendFormat("{0}={1} ", nameof(SampleTimeSec), SampleTimeSec);
+			sb.AppendFormat("{0}={1} ", nameof(NumSamples), NumSamples);
+			sb.AppendFormat("{0}={1} ", nameof(NumSamplesForResult), NumSamplesForResult);
+			sb.AppendFormat("{0}={1} ", nameof(MinAgents), MinAgents);
+			sb.AppendFormat("{0}={1} ", nameof(NumReserveAgents), NumReserveAgents);
+			return sb.ToString();
+		}
 	}
 	
 	/// <summary>
@@ -229,15 +251,13 @@ namespace Horde.Build.Agents.Fleet
 				PoolSizeData? poolSize = pools.Find(x => x.Pool.Id == pool.Id);
 				if (poolSize != null)
 				{
-					int minAgents = pool.MinAgents ?? 1;
-					int numReserveAgents = pool.NumReserveAgents ?? 5;
 					double utilization = poolData.Samples.Select(x => x._jobWork).OrderByDescending(x => x).Skip(Settings.NumSamples - Settings.NumSamplesForResult).First();
 				
 					// Number of agents in use over the sampling period. Can never be greater than number of agents available in pool.
 					int numAgentsUtilized = (int)utilization;
 					
 					// Include reserve agent count to ensure pool always can grow
-					int desiredAgentCount = Math.Max(numAgentsUtilized + numReserveAgents, minAgents);
+					int desiredAgentCount = Math.Max(numAgentsUtilized + Settings.NumReserveAgents, Settings.MinAgents);
 					
 					StringBuilder sb = new();
 					sb.AppendFormat("Jobs=[{0}] ", GetDensityMap(poolData.Samples.Select(x => x._jobWork / Math.Max(1, poolData.Agents.Count))));
@@ -246,6 +266,7 @@ namespace Horde.Build.Agents.Fleet
 					sb.AppendFormat("Max=[{0,5:0.0}] ", poolData.Samples.Max(x => x._jobWork));
 					sb.AppendFormat("Avg=[{0,5:0.0}] ", utilization);
 					sb.AppendFormat("Pct=[{0,5:0.0}] ", poolData.Samples.Sum(x => x._jobWork) / Settings.NumSamples);
+					sb.Append(Settings);
 
 					result.Add(new(pool, poolSize.Agents, desiredAgentCount, sb.ToString()));
 				}
