Horde: Add a queue for processing issue notifications, and acquire a lock before each issue is updated. This should prevent races between different servers trying to update/post about the same message.

diff --git a/Engine/Source/Programs/Horde/Horde.Build/Notifications/Sinks/SlackNotificationSink.cs b/Engine/Source/Programs/Horde/Horde.Build/Notifications/Sinks/SlackNotificationSink.cs
--- a/Engine/Source/Programs/Horde/Horde.Build/Notifications/Sinks/SlackNotificationSink.cs
+++ b/Engine/Source/Programs/Horde/Horde.Build/Notifications/Sinks/SlackNotificationSink.cs
@@ -15,6 +15,7 @@ using System.Threading;
 using System.Threading.Tasks;
 using EpicGames.Core;
 using EpicGames.Redis;
+using EpicGames.Redis.Utility;
 using EpicGames.Slack;
 using EpicGames.Slack.Blocks;
 using EpicGames.Slack.Elements;
@@ -37,6 +38,7 @@ using Microsoft.Extensions.Options;
 using MongoDB.Bson;
 using MongoDB.Bson.Serialization.Attributes;
 using MongoDB.Driver;
+using StackExchange.Redis;
 
 namespace Horde.Build.Notifications.Sinks
 {
@@ -205,6 +207,7 @@ namespace Horde.Build.Notifications.Sinks
 			}
 		}
 
+		readonly RedisService _redisService;
 		readonly IssueService _issueService;
 		readonly IUserCollection _userCollection;
 		readonly ILogFileService _logFileService;
@@ -220,6 +223,10 @@ namespace Horde.Build.Notifications.Sinks
 		readonly IClock _clock;
 		readonly ILogger _logger;
 
+		readonly ITicker _issueQueueTicker;
+		readonly RedisList<int> _redisIssueQueue;
+		readonly string _redisIssueLockPrefix;
+
 		readonly HttpClient _httpClient;
 		readonly SlackClient _slackClient;
 
@@ -233,6 +240,7 @@ namespace Horde.Build.Notifications.Sinks
 		/// </summary>
 		public SlackNotificationSink(MongoService mongoService, RedisService redisService, IssueService issueService, IUserCollection userCollection, ILogFileService logFileService, StreamService streamService, IExternalIssueService externalIssueService, IWebHostEnvironment environment, IOptions<ServerSettings> settings, IClock clock, ILogger<SlackNotificationSink> logger)
 		{
+			_redisService = redisService;
 			_issueService = issueService;
 			_userCollection = userCollection;
 			_logFileService = logFileService;
@@ -246,6 +254,10 @@ namespace Horde.Build.Notifications.Sinks
 			_clock = clock;
 			_logger = logger;
 
+			_issueQueueTicker = clock.AddSharedTicker("slack-issues", TimeSpan.FromMinutes(1.0), ProcessIssueQueueAsync, _logger);
+			_redisIssueQueue = new RedisList<int>(redisService.ConnectionPool, "slack/issue-queue");
+			_redisIssueLockPrefix = "slack/issues/";
+
 			_httpClient = new HttpClient();
 			_httpClient.DefaultRequestHeaders.Add("Authorization", $"Bearer {_settings.SlackToken ?? ""}");
 			_slackClient = new SlackClient(_httpClient, _logger);
@@ -268,6 +280,18 @@ namespace Horde.Build.Notifications.Sinks
 			_httpClient.Dispose();
 		}
 
+		/// <inheritdoc/>
+		public override async Task StartAsync(CancellationToken cancellationToken)
+		{
+			await _issueQueueTicker.StartAsync();
+		}
+
+		/// <inheritdoc/>
+		public override async Task StopAsync(CancellationToken cancellationToken)
+		{
+			await _issueQueueTicker.StopAsync();
+		}
+
 		#region Avatars
 
 		/// <inheritdoc/>
@@ -655,10 +679,62 @@ namespace Horde.Build.Notifications.Sinks
 				return;
 			}
 
-			using IDisposable scope = _logger.BeginScope("Slack notifications for issue {IssueId}", issue.Id);
-			_logger.LogInformation("Updating Slack notifications for issue {IssueId}", issue.Id);
+			// Otherwise add it to the redis queue, and attempt to process the queue immediately.
+			await _redisIssueQueue.RightPushAsync(issue.Id);
+			await ProcessIssueQueueAsync(CancellationToken.None);
+		}
+
+		/// <summary>
+		/// Processes the issue queue
+		/// </summary>
+		/// <param name="cancellationToken">Cancellation token for the operation</param>
+		async ValueTask ProcessIssueQueueAsync(CancellationToken cancellationToken)
+		{
+			HashSet<int> testedIssueIds = new HashSet<int>();
+
+			// Execute loop number of times based on the length of the queue at the start. This should bound the number of iterations while allowing us
+			// to re-queue items that we're unable to process now, without having to track whether we've reached the end of the list in its original state.
+			long count = await _redisIssueQueue.LengthAsync();
+			for (; count > 0; count--)
+			{
+				int issueId = await _redisIssueQueue.LeftPopAsync();
+				if (!testedIssueIds.Add(issueId) || !await TryUpdateIssueAsync(issueId))
+				{
+					await _redisIssueQueue.RightPushAsync(issueId);
+				}
+				cancellationToken.ThrowIfCancellationRequested();
+			}
+		}
+
+		async ValueTask<bool> TryUpdateIssueAsync(int issueId)
+		{
+			using (RedisLock issueLock = new RedisLock(_redisService.GetDatabase(), $"{_redisIssueLockPrefix}/{issueId}"))
+			{
+				if (await issueLock.AcquireAsync(TimeSpan.FromSeconds(30.0)))
+				{
+					await NotifyIssueUpdatedInternalAsync(issueId);
+					return true;
+				}
+				else
+				{
+					_logger.LogDebug("Unable to aquire lock for updating issue {IssueId}.", issueId);
+					return false;
+				}
+			}
+		}
+
+		async Task NotifyIssueUpdatedInternalAsync(int issueId)
+		{
+			using IDisposable scope = _logger.BeginScope("Slack notifications for issue {IssueId}", issueId);
+			_logger.LogInformation("Updating Slack notifications for issue {IssueId}", issueId);
+
+			IIssueDetails? details = await _issueService.GetIssueDetailsAsync(issueId);
+			if (details == null)
+			{
+				return;
+			}
 
-			IIssueDetails details = await _issueService.GetIssueDetailsAsync(issue);
+			IIssue issue = details.Issue;
 
 			bool notifyOwner = true;
 			bool notifySuspects = issue.Promoted || details.Spans.Any(x => x.LastFailure.Annotations.NotifySubmitters ?? false);
