Horde: Add a specific batch error code for a shelved change no longer existing.

diff --git a/Engine/Source/Programs/Horde/Horde.Build/Jobs/JobMessages.cs b/Engine/Source/Programs/Horde/Horde.Build/Jobs/JobMessages.cs
--- a/Engine/Source/Programs/Horde/Horde.Build/Jobs/JobMessages.cs
+++ b/Engine/Source/Programs/Horde/Horde.Build/Jobs/JobMessages.cs
@@ -721,6 +721,11 @@ namespace Horde.Build.Jobs
 		/// An error ocurred while executing the lease. Cannot be retried.
 		/// </summary>
 		ExecutionError = 10,
+
+		/// <summary>
+		/// The change that the job is running against is invalid
+		/// </summary>
+		UnknownShelf = 11,
 	}
 
 	/// <summary>
diff --git a/Engine/Source/Programs/Horde/Horde.Build/Jobs/JobService.cs b/Engine/Source/Programs/Horde/Horde.Build/Jobs/JobService.cs
--- a/Engine/Source/Programs/Horde/Horde.Build/Jobs/JobService.cs
+++ b/Engine/Source/Programs/Horde/Horde.Build/Jobs/JobService.cs
@@ -829,6 +829,8 @@ namespace Horde.Build.Jobs
 
 			using IDisposable scope = _logger.BeginScope("UpdateBatchAsync({JobId})", job.Id);
 
+			JobStepBatchError? error = null;
+
 			bool checkForBadAgent = true;
 			for (; ; )
 			{
@@ -853,26 +855,24 @@ namespace Horde.Build.Jobs
 				{
 					if (batch.Steps.Any(x => x.State == JobStepState.Waiting || x.State == JobStepState.Ready || x.State == JobStepState.Running))
 					{
-						// Mark the batch as incomplete
-						newError = JobStepBatchError.Incomplete;
+						// Check if the job is valid. If not, we will fail with a specific error code for it.
+						error ??= await CheckJobAsync(job) ?? JobStepBatchError.Incomplete;
+						newError = error.Value;
 
 						// Find the agent and set the conform flag
 						if (checkForBadAgent)
 						{
-							if (await IsValidJobAsync(job))
+							for (; ; )
 							{
-								for (; ; )
+								IAgent? agent = await _agents.GetAsync(batch.AgentId.Value);
+								if (agent == null || agent.RequestConform)
 								{
-									IAgent? agent = await _agents.GetAsync(batch.AgentId.Value);
-									if (agent == null || agent.RequestConform)
-									{
-										break;
-									}
-									if (await _agents.TryUpdateSettingsAsync(agent, bRequestConform: true) != null)
-									{
-										_logger.LogError("Agent {AgentId} did not complete lease; marking for conform", agent.Id);
-										break;
-									}
+									break;
+								}
+								if (await _agents.TryUpdateSettingsAsync(agent, bRequestConform: true) != null)
+								{
+									_logger.LogError("Agent {AgentId} did not complete lease; marking for conform", agent.Id);
+									break;
 								}
 							}
 							checkForBadAgent = false;
@@ -898,7 +898,7 @@ namespace Horde.Build.Jobs
 			}
 		}
 
-		async Task<bool> IsValidJobAsync(IJob job)
+		async Task<JobStepBatchError?> CheckJobAsync(IJob job)
 		{
 			try
 			{
@@ -908,22 +908,22 @@ namespace Horde.Build.Jobs
 					if (stream == null)
 					{
 						_logger.LogWarning("Job {JobId} is no longer valid - stream {StreamId} does not exist.", job.Id, job.StreamId);
-						return false;
+						return JobStepBatchError.UnknownStream;
 					}
 
 					(CheckShelfResult result, _) = await _perforceService.CheckShelfAsync(stream, job.PreflightChange);
 					if (result != CheckShelfResult.Ok)
 					{
 						_logger.LogWarning("Job {JobId} is no longer valid - check shelf returned {Result}", result);
-						return false;
+						return JobStepBatchError.UnknownShelf;
 					}
 				}
-				return true;
+				return null;
 			}
 			catch(Exception ex)
 			{
 				_logger.LogWarning(ex, "Job {JobId} is no longer valid.", job.Id);
-				return false;
+				return JobStepBatchError.ExecutionError;
 			}
 		}
 
