Horde: Encode temp storage metadata in bundle nodes.

diff --git a/Engine/Source/Programs/Horde/Horde.Agent/Execution/JobExecutor.cs b/Engine/Source/Programs/Horde/Horde.Agent/Execution/JobExecutor.cs
--- a/Engine/Source/Programs/Horde/Horde.Agent/Execution/JobExecutor.cs
+++ b/Engine/Source/Programs/Horde/Horde.Agent/Execution/JobExecutor.cs
@@ -13,6 +13,8 @@ using System.Text.Json.Serialization;
 using System.Threading;
 using System.Threading.Tasks;
 using EpicGames.Core;
+using EpicGames.Horde.Storage;
+using EpicGames.Horde.Storage.Backends;
 using Grpc.Core;
 using Horde.Agent.Parser;
 using Horde.Agent.Services;
@@ -21,6 +23,8 @@ using Horde.Storage.Utility;
 using HordeCommon;
 using HordeCommon.Rpc;
 using HordeCommon.Rpc.Tasks;
+using Microsoft.CodeAnalysis.Rename;
+using Microsoft.Extensions.Caching.Memory;
 using Microsoft.Extensions.Logging;
 using OpenTracing;
 using OpenTracing.Util;
@@ -641,14 +645,16 @@ namespace Horde.Agent.Execution
 
 		private async Task<bool> ExecuteWithTempStorageAsync(BeginStepResponse step, DirectoryReference workspaceDir, DirectoryReference sharedStorageDir, string arguments, ILogger logger, CancellationToken cancellationToken)
 		{
-			TempStorage storage = new TempStorage(workspaceDir, DirectoryReference.Combine(workspaceDir, "Engine", "Saved", "BuildGraph"), sharedStorageDir);
+			// Create the storage client
+			using MemoryCache cache = new MemoryCache(new MemoryCacheOptions { });
+			FileStorageClient storage = new FileStorageClient(DirectoryReference.Combine(sharedStorageDir, "bundles"), cache, logger);
 			logger.LogInformation("Using Horde-managed shared storage via {SharedStorageDir}", sharedStorageDir);
 
 			// Create the mapping of tag names to file sets
 			Dictionary<string, HashSet<FileReference>> tagNameToFileSet = new Dictionary<string, HashSet<FileReference>>();
 
 			// Read all the input tags for this node, and build a list of referenced input storage blocks
-			HashSet<TempStorageBlock> inputStorageBlocks = new HashSet<TempStorageBlock>();
+			HashSet<TempStorageBlockRef> inputStorageBlocks = new HashSet<TempStorageBlockRef>();
 			foreach (string input in step.Inputs)
 			{
 				int slashIdx = input.IndexOf('/', StringComparison.Ordinal);
@@ -661,40 +667,34 @@ namespace Horde.Agent.Execution
 				string nodeName = input.Substring(0, slashIdx);
 				string tagName = input.Substring(slashIdx + 1);
 
-				TempStorageFileList? fileList = storage.ReadFileList(nodeName, tagName, logger);
-				if (fileList == null)
-				{
-					logger.LogError("Unable to read file list for {Node}/{Tag} from shared storage ({SharedStorageDir})", nodeName, tagName, sharedStorageDir);
-					return false;
-				}
-
+				TempStorageTagManifest fileList = await TempStorage.RetrieveTagAsync(storage, nodeName, tagName, workspaceDir, logger, cancellationToken);
 				tagNameToFileSet[tagName] = fileList.ToFileSet(workspaceDir);
 				inputStorageBlocks.UnionWith(fileList.Blocks);
 			}
 
 			// Read the manifests for all the input storage blocks
-			Dictionary<TempStorageBlock, TempStorageManifest> inputManifests = new Dictionary<TempStorageBlock, TempStorageManifest>();
+			Dictionary<TempStorageBlockRef, TempStorageBlockManifest> inputManifests = new Dictionary<TempStorageBlockRef, TempStorageBlockManifest>();
 			using (IScope scope = GlobalTracer.Instance.BuildSpan("TempStorage").WithTag("resource", "read").StartActive())
 			{
 				scope.Span.SetTag("blocks", inputStorageBlocks.Count);
-				foreach (TempStorageBlock inputStorageBlock in inputStorageBlocks)
+				foreach (TempStorageBlockRef inputStorageBlock in inputStorageBlocks)
 				{
-					TempStorageManifest manifest = await storage.RetrieveAsync(inputStorageBlock.NodeName, inputStorageBlock.OutputName, logger, cancellationToken);
+					TempStorageBlockManifest manifest = await TempStorage.RetrieveBlockAsync(storage, inputStorageBlock.NodeName, inputStorageBlock.OutputName, workspaceDir, logger, cancellationToken);
 					inputManifests[inputStorageBlock] = manifest;
 				}
 				scope.Span.SetTag("size", inputManifests.Sum(x => x.Value.GetTotalSize()));
 			}
 
 			// Read all the input storage blocks, keeping track of which block each file came from
-			Dictionary<FileReference, TempStorageBlock> fileToStorageBlock = new Dictionary<FileReference, TempStorageBlock>();
-			foreach (KeyValuePair<TempStorageBlock, TempStorageManifest> pair in inputManifests)
+			Dictionary<FileReference, TempStorageBlockRef> fileToStorageBlock = new Dictionary<FileReference, TempStorageBlockRef>();
+			foreach (KeyValuePair<TempStorageBlockRef, TempStorageBlockManifest> pair in inputManifests)
 			{
-				TempStorageBlock inputStorageBlock = pair.Key;
+				TempStorageBlockRef inputStorageBlock = pair.Key;
 				foreach (FileReference file in pair.Value.Files.Select(x => x.ToFileReference(workspaceDir)))
 				{
 					logger.LogInformation("Reading input block: {File}", file);
 
-					TempStorageBlock? currentStorageBlock;
+					TempStorageBlockRef? currentStorageBlock;
 					if (fileToStorageBlock.TryGetValue(file, out currentStorageBlock) && !TempStorage.IsDuplicateBuildProduct(file))
 					{
 						logger.LogError("File '{File}' was produced by {InputBlock} and {CurrentBlock}", file, inputStorageBlock.ToString(), currentStorageBlock.ToString());
@@ -712,12 +712,10 @@ namespace Horde.Agent.Execution
 			// Read all the output manifests
 			foreach (string tagName in step.OutputNames)
 			{
-				TempStorageFileList? fileList = storage.ReadFileList(step.Name, tagName, logger);
-				if (fileList == null)
-				{
-					logger.LogError("Unable to read local file list for output {Tag}", tagName);
-					return false;
-				}
+				FileReference tagFileListLocation = TempStorage.GetTagManifestLocation(workspaceDir, step.Name, tagName);
+				logger.LogInformation("Reading local file list from {File}", tagFileListLocation.FullName);
+
+				TempStorageTagManifest fileList = TempStorageTagManifest.Load(tagFileListLocation);
 				tagNameToFileSet[tagName] = fileList.ToFileSet(workspaceDir);
 			}
 
@@ -809,36 +807,47 @@ namespace Horde.Agent.Execution
 			// Write all the storage blocks, and update the mapping from file to storage block
 			using (GlobalTracer.Instance.BuildSpan("TempStorage").WithTag("resource", "Write").StartActive())
 			{
+				RefName refName = new RefName(RefName.Sanitize(step.Name));
+
+				TreeOptions treeOptions = new TreeOptions();
+				TreeWriter treeWriter = new TreeWriter(storage, treeOptions, refName.Text);
+
+				TempStorageNode outputNode = new TempStorageNode();
+
+				// Create all the output blocks
 				foreach (KeyValuePair<string, HashSet<FileReference>> pair in outputStorageBlockToFiles)
 				{
-					TempStorageBlock outputBlock = new TempStorageBlock(step.Name, pair.Key);
+					TempStorageBlockRef outputBlock = new TempStorageBlockRef(step.Name, pair.Key);
 					foreach (FileReference file in pair.Value)
 					{
 						fileToStorageBlock.Add(file, outputBlock);
 					}
 					if (pair.Value.Any(x => referencedOutputFiles.Contains(x)))
 					{
-						await storage.ArchiveAsync(step.Name, pair.Key, pair.Value.ToArray(), logger, cancellationToken);
+						outputNode.Blocks[pair.Key] = await TempStorage.ArchiveBlockAsync(treeWriter, pair.Key, workspaceDir, pair.Value.ToArray(), logger, cancellationToken);
 					}
 				}
 
-				// Publish all the output tags
+				// Create all the output tags
 				foreach (string outputName in step.OutputNames)
 				{
 					HashSet<FileReference> files = tagNameToFileSet[outputName];
 
-					HashSet<TempStorageBlock> storageBlocks = new HashSet<TempStorageBlock>();
+					HashSet<TempStorageBlockRef> storageBlocks = new HashSet<TempStorageBlockRef>();
 					foreach (FileReference file in files)
 					{
-						TempStorageBlock? storageBlock;
+						TempStorageBlockRef? storageBlock;
 						if (fileToStorageBlock.TryGetValue(file, out storageBlock))
 						{
 							storageBlocks.Add(storageBlock);
 						}
 					}
 
-					storage.WriteFileList(step.Name, outputName, files, storageBlocks.ToArray(), logger);
+					outputNode.Tags[outputName] = await TempStorage.ArchiveTagAsync(treeWriter, outputName, workspaceDir, files, storageBlocks.ToArray(), logger, cancellationToken);
 				}
+
+				// Write the final node
+				await treeWriter.WriteRefAsync(refName, outputNode, cancellationToken);
 			}
 
 			return true;
diff --git a/Engine/Source/Programs/Horde/Horde.Agent/Utility/TempStorage.cs b/Engine/Source/Programs/Horde/Horde.Agent/Utility/TempStorage.cs
--- a/Engine/Source/Programs/Horde/Horde.Agent/Utility/TempStorage.cs
+++ b/Engine/Source/Programs/Horde/Horde.Agent/Utility/TempStorage.cs
@@ -2,9 +2,7 @@
 
 using EpicGames.Core;
 using EpicGames.Horde.Storage;
-using EpicGames.Horde.Storage.Backends;
 using EpicGames.Horde.Storage.Nodes;
-using Microsoft.Extensions.Caching.Memory;
 using Microsoft.Extensions.Logging;
 using System;
 using System.Collections.Generic;
@@ -12,7 +10,6 @@ using System.Diagnostics;
 using System.Diagnostics.CodeAnalysis;
 using System.IO;
 using System.Linq;
-using System.Text;
 using System.Threading;
 using System.Threading.Tasks;
 using System.Xml;
@@ -39,7 +36,7 @@ namespace Horde.Storage.Utility
 	/// <summary>
 	/// Stores the name of a temp storage block
 	/// </summary>
-	public class TempStorageBlock
+	public class TempStorageBlockRef
 	{
 		/// <summary>
 		/// Name of the node
@@ -56,7 +53,7 @@ namespace Horde.Storage.Utility
 		/// <summary>
 		/// Default constructor, for XML serialization.
 		/// </summary>
-		private TempStorageBlock()
+		private TempStorageBlockRef()
 		{
 			NodeName = String.Empty;
 			OutputName = String.Empty;
@@ -67,7 +64,7 @@ namespace Horde.Storage.Utility
 		/// </summary>
 		/// <param name="nodeName">Name of the node</param>
 		/// <param name="outputName">Name of the node's output</param>
-		public TempStorageBlock(string nodeName, string outputName)
+		public TempStorageBlockRef(string nodeName, string outputName)
 		{
 			NodeName = nodeName;
 			OutputName = outputName;
@@ -78,7 +75,7 @@ namespace Horde.Storage.Utility
 		/// </summary>
 		/// <param name="other">The object to compare against</param>
 		/// <returns>True if the blocks are equivalent</returns>
-		public override bool Equals(object? other) => other is TempStorageBlock otherBlock && NodeName.Equals(otherBlock.NodeName, StringComparison.OrdinalIgnoreCase) && OutputName.Equals(otherBlock.OutputName, StringComparison.OrdinalIgnoreCase);
+		public override bool Equals(object? other) => other is TempStorageBlockRef otherBlock && NodeName.Equals(otherBlock.NodeName, StringComparison.OrdinalIgnoreCase) && OutputName.Equals(otherBlock.OutputName, StringComparison.OrdinalIgnoreCase);
 
 		/// <summary>
 		/// Returns a hash code for this block name
@@ -131,6 +128,17 @@ namespace Horde.Storage.Utility
 			RelativePath = String.Empty;
 		}
 
+		/// <summary>
+		/// Constructor
+		/// </summary>
+		public TempStorageFile(string relativePath, long lastWriteTimeUtcTicks, long length, string? digest)
+		{
+			RelativePath = relativePath;
+			LastWriteTimeUtcTicks = lastWriteTimeUtcTicks;
+			Length = length;
+			Digest = digest;
+		}
+
 		/// <summary>
 		/// Constructor
 		/// </summary>
@@ -296,47 +304,10 @@ namespace Horde.Storage.Utility
 		}
 	}
 
-	/// <summary>
-	/// Information about a single file in temp storage
-	/// </summary>
-	[DebuggerDisplay("{Name}")]
-	public class TempStorageZipFile
-	{
-		/// <summary>
-		/// Name of this file, including extension
-		/// </summary>
-		[XmlAttribute]
-		public string Name { get; set; }
-
-		/// <summary>
-		/// Length of the file in bytes
-		/// </summary>
-		[XmlAttribute]
-		public long Length { get; set; }
-
-		/// <summary>
-		/// Default constructor, for XML serialization
-		/// </summary>
-		private TempStorageZipFile()
-		{
-			Name = String.Empty;
-		}
-
-		/// <summary>
-		/// Constructor
-		/// </summary>
-		/// <param name="info">FileInfo for the zip file</param>
-		public TempStorageZipFile(FileInfo info)
-		{
-			Name = info.Name;
-			Length = info.Length;
-		}
-	}
-
 	/// <summary>
 	/// A manifest storing information about build products for a node's output
 	/// </summary>
-	public class TempStorageManifest
+	public class TempStorageBlockManifest
 	{
 		/// <summary>
 		/// List of output files
@@ -348,22 +319,30 @@ namespace Horde.Storage.Utility
 		/// <summary>
 		/// Construct a static Xml serializer to avoid throwing an exception searching for the reflection info at runtime
 		/// </summary>
-		static readonly XmlSerializer s_serializer = XmlSerializer.FromTypes(new Type[]{ typeof(TempStorageManifest) })[0];
+		static readonly XmlSerializer s_serializer = XmlSerializer.FromTypes(new Type[]{ typeof(TempStorageBlockManifest) })[0];
 
 		/// <summary>
 		/// Construct an empty temp storage manifest
 		/// </summary>
-		private TempStorageManifest()
+		private TempStorageBlockManifest()
 		{
 			Files = Array.Empty<TempStorageFile>();
 		}
 
+		/// <summary>
+		/// Construct a temp storage manifest
+		/// </summary>
+		public TempStorageBlockManifest(TempStorageFile[] files)
+		{
+			Files = files;
+		}
+
 		/// <summary>
 		/// Creates a manifest from a flat list of files (in many folders) and a BaseFolder from which they are rooted.
 		/// </summary>
 		/// <param name="files">List of full file paths</param>
 		/// <param name="rootDir">Root folder for all the files. All files must be relative to this RootDir.</param>
-		public TempStorageManifest(FileInfo[] files, DirectoryReference rootDir)
+		public TempStorageBlockManifest(FileInfo[] files, DirectoryReference rootDir)
 		{
 			Files = files.Select(x => new TempStorageFile(x, rootDir)).ToArray();
 		}
@@ -386,14 +365,14 @@ namespace Horde.Storage.Utility
 		/// Load a manifest from disk
 		/// </summary>
 		/// <param name="file">File to load</param>
-		public static TempStorageManifest Load(FileReference file)
+		public static TempStorageBlockManifest Load(FileReference file)
 		{
 			using (StreamReader reader = new(file.FullName))
 			{
 				XmlReaderSettings settings = new XmlReaderSettings();
 				using (XmlReader xmlReader = XmlReader.Create(reader, settings))
 				{
-					return (TempStorageManifest)s_serializer.Deserialize(xmlReader)!;
+					return (TempStorageBlockManifest)s_serializer.Deserialize(xmlReader)!;
 				}
 			}
 		}
@@ -418,7 +397,7 @@ namespace Horde.Storage.Utility
 	/// <summary>
 	/// Stores the contents of a tagged file set
 	/// </summary>
-	public class TempStorageFileList
+	public class TempStorageTagManifest
 	{
 		/// <summary>
 		/// List of files that are in this tag set, relative to the root directory
@@ -439,21 +418,34 @@ namespace Horde.Storage.Utility
 		/// </summary>
 		[XmlArray]
 		[XmlArrayItem("Block")]
-		public TempStorageBlock[] Blocks { get; set; }
+		public TempStorageBlockRef[] Blocks { get; set; }
 
 		/// <summary>
 		/// Construct a static Xml serializer to avoid throwing an exception searching for the reflection info at runtime
 		/// </summary>
-		static readonly XmlSerializer s_serializer = XmlSerializer.FromTypes(new Type[]{ typeof(TempStorageFileList) })[0];
+		static readonly XmlSerializer s_serializer = XmlSerializer.FromTypes(new Type[]{ typeof(TempStorageTagManifest) })[0];
 
 		/// <summary>
 		/// Construct an empty file list for deserialization
 		/// </summary>
-		private TempStorageFileList()
+		private TempStorageTagManifest()
 		{
 			LocalFiles = Array.Empty<string>();
 			ExternalFiles = Array.Empty<string>();
-			Blocks = Array.Empty<TempStorageBlock>();
+			Blocks = Array.Empty<TempStorageBlockRef>();
+		}
+
+		/// <summary>
+		/// 
+		/// </summary>
+		/// <param name="localFiles"></param>
+		/// <param name="externalFiles"></param>
+		/// <param name="blocks"></param>
+		public TempStorageTagManifest(string[] localFiles, string[] externalFiles, TempStorageBlockRef[] blocks)
+		{
+			LocalFiles = localFiles;
+			ExternalFiles = externalFiles;
+			Blocks = blocks;
 		}
 
 		/// <summary>
@@ -462,7 +454,7 @@ namespace Horde.Storage.Utility
 		/// <param name="files">List of full file paths</param>
 		/// <param name="rootDir">Root folder for all the files. All files must be relative to this RootDir.</param>
 		/// <param name="blocks">Referenced storage blocks required for these files</param>
-		public TempStorageFileList(IEnumerable<FileReference> files, DirectoryReference rootDir, IEnumerable<TempStorageBlock> blocks)
+		public TempStorageTagManifest(IEnumerable<FileReference> files, DirectoryReference rootDir, IEnumerable<TempStorageBlockRef> blocks)
 		{
 			List<string> newLocalFiles = new List<string>();
 			List<string> newExternalFiles = new List<string>();
@@ -487,14 +479,14 @@ namespace Horde.Storage.Utility
 		/// Load this list of files from disk
 		/// </summary>
 		/// <param name="file">File to load</param>
-		public static TempStorageFileList Load(FileReference file)
+		public static TempStorageTagManifest Load(FileReference file)
 		{
 			using(StreamReader reader = new StreamReader(file.FullName))
 			{
 				XmlReaderSettings settings = new XmlReaderSettings();
 				using (XmlReader xmlReader = XmlReader.Create(reader, settings))
 				{
-					return (TempStorageFileList)s_serializer.Deserialize(xmlReader)!;
+					return (TempStorageTagManifest)s_serializer.Deserialize(xmlReader)!;
 				}
 			}
 		}
@@ -529,6 +521,134 @@ namespace Horde.Storage.Utility
 		}
 	}
 
+	/// <summary>
+	/// Node representing a temp storage block
+	/// </summary>
+	class TempStorageBlockNode : TreeNode
+	{
+		public TempStorageBlockManifest Manifest { get; set; }
+		public TreeNodeRef<DirectoryNode> Contents { get; }
+
+		public TempStorageBlockNode(TempStorageBlockManifest manifest, TreeNodeRef<DirectoryNode> contents)
+		{
+			Manifest = manifest;
+			Contents = contents;
+		}
+
+		public TempStorageBlockNode(ITreeNodeReader reader)
+			: base(reader)
+		{
+			TempStorageFile[] files = reader.ReadVariableLengthArray(() => ReadTempStorageFile(reader));
+			Manifest = new TempStorageBlockManifest(files);
+
+			Contents = reader.ReadRef<DirectoryNode>();
+		}
+
+		static TempStorageFile ReadTempStorageFile(ITreeNodeReader reader)
+		{
+			string relativePath = reader.ReadString();
+			long lastWriteTimeUtcTicks = (long)reader.ReadUnsignedVarInt();
+			long length = (long)reader.ReadUnsignedVarInt();
+			string digest = reader.ReadString();
+			return new TempStorageFile(relativePath, lastWriteTimeUtcTicks, length, String.IsNullOrEmpty(digest) ? null : digest);
+		}
+
+		static void WriteTempStorageFile(TempStorageFile file, ITreeNodeWriter writer)
+		{
+			writer.WriteString(file.RelativePath);
+			writer.WriteUnsignedVarInt((ulong)file.LastWriteTimeUtcTicks);
+			writer.WriteUnsignedVarInt((ulong)file.Length);
+			writer.WriteString(file.Digest ?? String.Empty);
+		}
+
+		public override IEnumerable<TreeNodeRef> EnumerateRefs()
+		{
+			yield return Contents;
+		}
+
+		public override void Serialize(ITreeNodeWriter writer)
+		{
+			writer.WriteVariableLengthArray(Manifest.Files, file => WriteTempStorageFile(file, writer));
+			writer.WriteRef(Contents);
+		}
+	}
+
+	/// <summary>
+	/// Node representing a temp storage tag
+	/// </summary>
+	class TempStorageTagNode : TreeNode
+	{
+		public TempStorageTagManifest FileList { get; }
+
+		public TempStorageTagNode(TempStorageTagManifest fileList)
+		{
+			FileList = fileList;
+		}
+
+		public TempStorageTagNode(ITreeNodeReader reader)
+		{
+			string[] localFiles = reader.ReadVariableLengthArray(() => reader.ReadString());
+			string[] externalFiles = reader.ReadVariableLengthArray(() => reader.ReadString());
+			TempStorageBlockRef[] blocks = reader.ReadVariableLengthArray(() => ReadTempStorageBlock(reader));
+			FileList = new TempStorageTagManifest(localFiles, externalFiles, blocks);
+		}
+
+		static TempStorageBlockRef ReadTempStorageBlock(ITreeNodeReader reader)
+		{
+			string nodeName = reader.ReadString();
+			string blockName = reader.ReadString();
+			return new TempStorageBlockRef(nodeName, blockName);
+		}
+
+		public override void Serialize(ITreeNodeWriter writer)
+		{
+			throw new NotImplementedException();
+		}
+
+		public override IEnumerable<TreeNodeRef> EnumerateRefs()
+		{
+			throw new NotImplementedException();
+		}
+	}
+
+	/// <summary>
+	/// Node representing all the outputs from a node
+	/// </summary>
+	class TempStorageNode : TreeNode
+	{
+		public Dictionary<string, TreeNodeRef<TempStorageTagNode>> Tags { get; } = new Dictionary<string, TreeNodeRef<TempStorageTagNode>>(StringComparer.OrdinalIgnoreCase);
+		public Dictionary<string, TreeNodeRef<TempStorageBlockNode>> Blocks { get; } = new Dictionary<string, TreeNodeRef<TempStorageBlockNode>>(StringComparer.OrdinalIgnoreCase);
+
+		public TempStorageNode()
+		{
+		}
+
+		public TempStorageNode(ITreeNodeReader reader)
+			: base(reader)
+		{
+			reader.ReadDictionary(Tags, () => reader.ReadString(), () => reader.ReadRef<TempStorageTagNode>());
+			reader.ReadDictionary(Blocks, () => reader.ReadString(), () => reader.ReadRef<TempStorageBlockNode>());
+		}
+
+		public override IEnumerable<TreeNodeRef> EnumerateRefs()
+		{
+			foreach (TreeNodeRef<TempStorageTagNode> tagNode in Tags.Values)
+			{
+				yield return tagNode;
+			}
+			foreach (TreeNodeRef<TempStorageBlockNode> blockNode in Blocks.Values)
+			{
+				yield return blockNode;
+			}
+		}
+
+		public override void Serialize(ITreeNodeWriter writer)
+		{
+			writer.WriteDictionary(Tags, k => writer.WriteString(k), v => writer.WriteRef(v));
+			writer.WriteDictionary(Blocks, k => writer.WriteString(k), v => writer.WriteRef(v));
+		}
+	}
+
 	/// <summary>
 	/// Tracks the state of the current build job using the filesystem, allowing jobs to be restarted after a failure or expanded to include larger targets, and 
 	/// providing a proxy for different machines executing parts of the build in parallel to transfer build products and share state as part of a build system.
@@ -539,190 +659,108 @@ namespace Horde.Storage.Utility
 	/// The local temp storage directory contains the same information, with the exception of the archived build products. Metadata is still kept to detect modified 
 	/// build products between runs. If data is not present in local temp storage, it's retrieved from shared temp storage and cached in local storage.
 	/// </summary>
-	class TempStorage
+	static class TempStorage
 	{
 		/// <summary>
-		/// Root directory for this branch.
-		/// </summary>
-		readonly DirectoryReference _rootDir;
-
-		/// <summary>
-		/// The local temp storage directory (typically somewhere under /Engine/Saved directory).
-		/// </summary>
-		readonly DirectoryReference _localDir;
-
-		/// <summary>
-		/// The shared temp storage directory; typically a network location. May be null.
-		/// </summary>
-		readonly DirectoryReference _sharedDir;
-
-		/// <summary>
-		/// Constructor
+		/// Gets the ref name for a particular node
 		/// </summary>
-		/// <param name="rootDir">Root directory for this branch</param>
-		/// <param name="localDir">The local temp storage directory.</param>
-		/// <param name="sharedDir">The shared temp storage directory. May be null.</param>
-		public TempStorage(DirectoryReference rootDir, DirectoryReference localDir, DirectoryReference sharedDir)
+		/// <param name="nodeName"></param>
+		/// <returns></returns>
+		static RefName GetRefNameForNode(string nodeName)
 		{
-			_rootDir = rootDir;
-			_localDir = localDir;
-			_sharedDir = sharedDir;
+			return new RefName(RefName.Sanitize(nodeName));
 		}
 
 		/// <summary>
 		/// Reads a set of tagged files from disk
 		/// </summary>
+		/// <param name="store"></param>
 		/// <param name="nodeName">Name of the node which produced the tag set</param>
 		/// <param name="tagName">Name of the tag, with a '#' prefix</param>
+		/// <param name="localDir">The local directory</param>
 		/// <param name="logger">Logger for output</param>
+		/// <param name="cancellationToken"></param>
 		/// <returns>The set of files</returns>
-		public TempStorageFileList? ReadFileList(string nodeName, string tagName, ILogger logger)
+		public static async Task<TempStorageTagManifest> RetrieveTagAsync(IStorageClient store, string nodeName, string tagName, DirectoryReference localDir, ILogger logger, CancellationToken cancellationToken)
 		{
-			TempStorageFileList? fileList;
+			TempStorageTagManifest fileList;
 
 			// Try to read the tag set from the local directory
-			FileReference localFileListLocation = GetTaggedFileListLocation(_localDir, nodeName, tagName);
+			FileReference localFileListLocation = GetTagManifestLocation(localDir, nodeName, tagName);
 			if(FileReference.Exists(localFileListLocation))
 			{
 				logger.LogInformation("Reading local file list from {File}", localFileListLocation.FullName);
-				fileList = TempStorageFileList.Load(localFileListLocation);
+				fileList = TempStorageTagManifest.Load(localFileListLocation);
 			}
 			else
 			{
-				// Check we have shared storage
-				if(_sharedDir == null)
-				{
-					throw new TempStorageException($"Missing local file list - {localFileListLocation.FullName}");
-				}
+				RefName refName = GetRefNameForNode(nodeName);
+				logger.LogInformation("Reading tag {NodeName}:{TagName} from temp storage (ref: {RefName})", nodeName, tagName, refName);
 
-				// Make sure the manifest exists.
-				FileReference sharedFileListLocation = GetTaggedFileListLocation(_sharedDir, nodeName, tagName);
-				if (!FileReference.Exists(sharedFileListLocation))
-				{
-					throw new TempStorageException($"Missing local or shared file list - {sharedFileListLocation.FullName}");
-				}
+				TempStorageNode node = await store.ReadNodeAsync<TempStorageNode>(refName, cancellationToken: cancellationToken);
 
-				try
+				TreeNodeRef<TempStorageTagNode>? tagNodeRef;
+				if (!node.Tags.TryGetValue(tagName, out tagNodeRef))
 				{
-					// Read the shared manifest
-					logger.LogInformation("Copying shared tag set from {Source} to {Target}", sharedFileListLocation.FullName, localFileListLocation.FullName);
-					fileList = TempStorageFileList.Load(sharedFileListLocation);
-				}
-				catch
-				{
-					throw new TempStorageException($"Local or shared file list {sharedFileListLocation.FullName} was found but failed to be read");
+					throw new TempStorageException($"Missing tag {tagName} from node {nodeName}");
 				}
 
+				TempStorageTagNode tagNode = await tagNodeRef.ExpandAsync(cancellationToken);
+				fileList = tagNode.FileList;
+
 				// Save the manifest locally
 				DirectoryReference.CreateDirectory(localFileListLocation.Directory);
-				fileList?.Save(localFileListLocation);
+				fileList.Save(localFileListLocation);
 			}
 			return fileList;
 		}
 
 		/// <summary>
-		/// Writes a list of tagged files to disk
+		/// Saves a tag to temp storage
 		/// </summary>
-		/// <param name="nodeName">Name of the node which produced the tag set</param>
-		/// <param name="tagName">Name of the tag, with a '#' prefix</param>
-		/// <param name="files">List of files in this set</param>
-		/// <param name="blocks">List of referenced storage blocks</param>
+		/// <param name="writer">Writer for output</param>
+		/// <param name="tagName">Name of the output tag</param>
+		/// <param name="rootDir">Root directory for the build products</param>
+		/// <param name="files">Files in the tag</param>
+		/// <param name="blocks">Blocks containing files in the tag</param>
 		/// <param name="logger">Logger for output</param>
-		/// <returns>The set of files</returns>
-		public void WriteFileList(string nodeName, string tagName, IEnumerable<FileReference> files, IEnumerable<TempStorageBlock> blocks, ILogger logger)
-		{
-			// Create the file list
-			TempStorageFileList fileList = new TempStorageFileList(files, _rootDir, blocks);
-
-			// Save the set of files to the local and shared locations
-			FileReference sharedFileListLocation = GetTaggedFileListLocation(_sharedDir, nodeName, tagName);
-			try
-			{
-				logger.LogInformation("Saving file list to {File}", sharedFileListLocation.FullName);
-				DirectoryReference.CreateDirectory(sharedFileListLocation.Directory);
-				fileList.Save(sharedFileListLocation);
-			}
-			catch (Exception ex)
-			{
-				throw new TempStorageException($"Failed to save file list {sharedFileListLocation}, exception: {ex}");
-			}
-		}
-
-		static DirectoryNode FindOrAddDirNode(DirectoryReference dirRef, Dictionary<DirectoryReference, DirectoryNode> dirRefToNode)
-		{
-			DirectoryNode? dirNode;
-			if (!dirRefToNode.TryGetValue(dirRef, out dirNode))
-			{
-				DirectoryNode parentNode = FindOrAddDirNode(dirRef.ParentDirectory!, dirRefToNode);
-				dirNode = parentNode.AddDirectory(dirRef.GetDirectoryName());
-				dirRefToNode.Add(dirRef, dirNode);
-			}
-			return dirNode;
-		}
-
-		/// <summary>
-		/// Gets a ref name from the given node and block name
-		/// </summary>
-		static RefName GetRefName(string nodeName, string blockName)
+		/// <param name="cancellationToken">Cancellation token for the operation</param>
+		/// <returns></returns>
+		public static async Task<TreeNodeRef<TempStorageTagNode>> ArchiveTagAsync(TreeWriter writer, string tagName, DirectoryReference rootDir, IEnumerable<FileReference> files, IEnumerable<TempStorageBlockRef> blocks, ILogger logger, CancellationToken cancellationToken)
 		{
-			string inputName = $"{nodeName}/{blockName}".TrimEnd('/');
-
-			StringBuilder result = new StringBuilder();
-			for (int idx = 0; idx < inputName.Length; idx++)
-			{
-				if (inputName[idx] >= 'A' && inputName[idx] <= 'Z')
-				{
-					result.Append((char)(inputName[idx] + 'a' - 'A'));
-				}
-				else if ((inputName[idx] >= 'a' && inputName[idx] <= 'z') || (inputName[idx] >= '0' && inputName[idx] <= '9') || inputName[idx] == '+' || inputName[idx] == '/')
-				{
-					result.Append(inputName[idx]);
-				}
-				else if (result.Length > 0 && result[result.Length - 1] != '-')
-				{
-					result.Append('-');
-				}
-			}
+			logger.LogInformation("Creating output tag {TagName}", tagName);
 
-			if (result.Length > 0 && result[result.Length - 1] == '-')
-			{
-				result.Remove(result.Length - 1, 1);
-			}
+			TempStorageTagManifest fileList = new TempStorageTagManifest(files, rootDir, blocks);
+			TreeNodeRef<TempStorageTagNode> tagNode = new TreeNodeRef<TempStorageTagNode>(new TempStorageTagNode(fileList));
+			await writer.WriteAsync(tagNode, cancellationToken);
 
-			return new RefName(result.ToString());
+			return tagNode;
 		}
 
-		DirectoryReference GetBundleDir() => DirectoryReference.Combine(_sharedDir!, "data");
-
 		/// <summary>
 		/// Saves the given files (that should be rooted at the branch root) to a shared temp storage manifest with the given temp storage node and game.
 		/// </summary>
-		/// <param name="nodeName">The node which created the storage block</param>
-		/// <param name="blockName">Name of the block to retrieve. May be null or empty.</param>
+		/// <param name="writer">Writer for output</param>
+		/// <param name="blockName">Name of the output block</param>
+		/// <param name="rootDir">Root directory for the build products</param>
 		/// <param name="buildProducts">Array of build products to be archived</param>
 		/// <param name="logger">Logger for output</param>
-		/// <param name="cancellationToken"></param>
+		/// <param name="cancellationToken">Cancellation token for the operation</param>
 		/// <returns>The created manifest instance (which has already been saved to disk).</returns>
-		public async Task ArchiveAsync(string nodeName, string blockName, FileReference[] buildProducts, ILogger logger, CancellationToken cancellationToken)
+		public static async Task<TreeNodeRef<TempStorageBlockNode>> ArchiveBlockAsync(TreeWriter writer, string blockName, DirectoryReference rootDir, FileReference[] buildProducts, ILogger logger, CancellationToken cancellationToken)
 		{
+			logger.LogInformation("Creating output block {BlockName}", blockName);
+
 			// Create a manifest for the given build products
 			FileInfo[] files = buildProducts.Select(x => new FileInfo(x.FullName)).ToArray();
 
-			// Compress the files and copy to shared storage if necessary
-			DirectoryReference bundleDir = GetBundleDir();
-
-			// Create the storage client
-			using MemoryCache cache = new MemoryCache(new MemoryCacheOptions { });
-			FileStorageClient store = new FileStorageClient(bundleDir, cache, logger);
-
 			// Create the directory tree
 			DirectoryNode root = new DirectoryNode(DirectoryFlags.None);
 
 			Dictionary<DirectoryReference, DirectoryNode> dirRefToNode = new Dictionary<DirectoryReference, DirectoryNode>();
-			dirRefToNode.Add(_rootDir, root);
+			dirRefToNode.Add(rootDir, root);
 
 			List<(DirectoryNode, FileInfo)> filesToAdd = new List<(DirectoryNode, FileInfo)>();
-
 			foreach (FileInfo file in files)
 			{
 				DirectoryReference dirRef = new FileReference(file).Directory;
@@ -730,88 +768,88 @@ namespace Horde.Storage.Utility
 				filesToAdd.Add((dirNode, file));
 			}
 
-			// Add all the files and flush the ref
-			RefName refName = GetRefName(nodeName, blockName);
-
-			TreeOptions treeOptions = new TreeOptions();
-			TreeWriter writer = new TreeWriter(store, treeOptions, prefix: refName.Text);
-
 			ChunkingOptions chunkingOptions = new ChunkingOptions();
 			await DirectoryNode.CopyFromDirectoryAsync(filesToAdd, chunkingOptions, writer, cancellationToken);
 
-			await writer.WriteRefAsync(refName, root, cancellationToken);
-
-			// Save the shared manifest
-			logger.LogInformation("Written {RefName} to {NodeDir}", refName, bundleDir);
-
-			// Create the shared directory for this node
-			FileReference sharedManifestFile = GetManifestLocation(_sharedDir!, nodeName, blockName);
-			logger.LogInformation("Saving shared manifest to {File}", sharedManifestFile.FullName);
+			// Create the block node
+			TempStorageBlockManifest manifest = new TempStorageBlockManifest(files, rootDir);
+			TempStorageBlockNode node = new TempStorageBlockNode(manifest, new TreeNodeRef<DirectoryNode>(root));
+			return new TreeNodeRef<TempStorageBlockNode>(node);
+		}
 
-			DirectoryReference.CreateDirectory(sharedManifestFile.Directory);
-			TempStorageManifest manifest = new TempStorageManifest(files, _rootDir);
-			manifest.Save(sharedManifestFile);
+		/// <summary>
+		/// Finds or adds a node for the given directory, using a cached lookup for existing nodes
+		/// </summary>
+		/// <param name="dirRef"></param>
+		/// <param name="dirRefToNode"></param>
+		/// <returns></returns>
+		static DirectoryNode FindOrAddDirNode(DirectoryReference dirRef, Dictionary<DirectoryReference, DirectoryNode> dirRefToNode)
+		{
+			DirectoryNode? dirNode;
+			if (!dirRefToNode.TryGetValue(dirRef, out dirNode))
+			{
+				DirectoryNode parentNode = FindOrAddDirNode(dirRef.ParentDirectory!, dirRefToNode);
+				dirNode = parentNode.AddDirectory(dirRef.GetDirectoryName());
+				dirRefToNode.Add(dirRef, dirNode);
+			}
+			return dirNode;
 		}
 
 		/// <summary>
 		/// Retrieve an output of the given node. Fetches and decompresses the files from shared storage if necessary, or validates the local files.
 		/// </summary>
+		/// <param name="store">Store to read data from</param>
 		/// <param name="nodeName">The node which created the storage block</param>
 		/// <param name="blockName">Name of the block to retrieve.</param>
+		/// <param name="rootDir">Local directory for extracting data to</param>
 		/// <param name="logger">Logger for output</param>
 		/// <param name="cancellationToken"></param>
 		/// <returns>Manifest of the files retrieved</returns>
-		public async Task<TempStorageManifest> RetrieveAsync(string nodeName, string blockName, ILogger logger, CancellationToken cancellationToken)
+		public static async Task<TempStorageBlockManifest> RetrieveBlockAsync(IStorageClient store, string nodeName, string blockName, DirectoryReference rootDir, ILogger logger, CancellationToken cancellationToken)
 		{
 			// Get the path to the local manifest
-			FileReference localManifestFile = GetManifestLocation(_localDir, nodeName, blockName);
+			FileReference localManifestFile = GetBlockManifestLocation(rootDir, nodeName, blockName);
 			bool local = FileReference.Exists(localManifestFile);
 
 			// Read the manifest, either from local storage or shared storage
-			TempStorageManifest? manifest;
+			TempStorageBlockManifest? manifest;
 			if(local)
 			{
-				logger.LogInformation("Reading shared manifest from {File}", localManifestFile.FullName);
-				manifest = TempStorageManifest.Load(localManifestFile);
+				logger.LogInformation("Reading block manifest from {File}", localManifestFile.FullName);
+				manifest = TempStorageBlockManifest.Load(localManifestFile);
 			}
 			else
 			{
-				// Get the shared directory for this node
-				FileReference sharedManifestFile = GetManifestLocation(_sharedDir, nodeName, blockName);
+				// Read the shared manifest
+				RefName refName = GetRefNameForNode(nodeName);
+				logger.LogInformation("Reading block {NodeName}:{BlockName} from temp storage (ref: {RefName}, local: {LocalFile})", nodeName, blockName, refName, localManifestFile);
 
-				// Make sure the manifest exists
-				if(!FileReference.Exists(sharedManifestFile))
+				TempStorageNode node = await store.ReadNodeAsync<TempStorageNode>(refName, cancellationToken: cancellationToken);
+				
+				TreeNodeRef<TempStorageBlockNode>? blockNodeRef;
+				if (!node.Blocks.TryGetValue(blockName, out blockNodeRef))
 				{
-					throw new TempStorageException($"Missing local or shared manifest for node - {sharedManifestFile.FullName}");
+					throw new TempStorageException($"Missing block {blockName} from node {nodeName}");
 				}
 
-				// Read the shared manifest
-				logger.LogInformation("Copying shared manifest from {Source} to {Target}", sharedManifestFile.FullName, localManifestFile.FullName);
-				manifest = TempStorageManifest.Load(sharedManifestFile);
+				TempStorageBlockNode blockNode = await blockNodeRef.ExpandAsync(cancellationToken);
+				manifest = blockNode.Manifest;
 
 				// Delete all the existing files. They may be read-only.
 				foreach (TempStorageFile ManifestFile in manifest.Files)
 				{
-					FileReference File = ManifestFile.ToFileReference(_rootDir);
+					FileReference File = ManifestFile.ToFileReference(rootDir);
 					FileUtils.ForceDeleteFile(File);
 				}
 
-				// Compress the files and copy to shared storage if necessary
-				DirectoryReference bundleDir = GetBundleDir();
-
-				// Create the storage client
-				using MemoryCache cache = new MemoryCache(new MemoryCacheOptions { });
-				FileStorageClient store = new FileStorageClient(bundleDir, cache, logger);
-
 				// Add all the files and flush the ref
-				RefName refName = GetRefName(nodeName, blockName);
-				DirectoryNode directory = await store.ReadNodeAsync<DirectoryNode>(refName, cancellationToken: cancellationToken);
-				await directory.CopyToDirectoryAsync(_rootDir.ToDirectoryInfo(), logger, cancellationToken);
+				DirectoryNode rootDirNode = await blockNode.Contents.ExpandAsync(cancellationToken);
+				await rootDirNode.CopyToDirectoryAsync(rootDir.ToDirectoryInfo(), logger, cancellationToken);
 
 				// Update the timestamps to match the manifest.
 				foreach (TempStorageFile ManifestFile in manifest.Files)
 				{
-					FileReference File = ManifestFile.ToFileReference(_rootDir);
+					FileReference File = ManifestFile.ToFileReference(rootDir);
 					System.IO.File.SetLastWriteTimeUtc(File.FullName, new DateTime(ManifestFile.LastWriteTimeUtcTicks, DateTimeKind.Utc));
 				}
 
@@ -824,7 +862,7 @@ namespace Horde.Storage.Utility
 			bool bAllMatch = true;
 			foreach(TempStorageFile File in manifest.Files)
 			{
-				bAllMatch &= File.Compare(_rootDir, logger);
+				bAllMatch &= File.Compare(rootDir, logger);
 			}
 			if(!bAllMatch)
 			{
@@ -839,7 +877,7 @@ namespace Horde.Storage.Utility
 		/// <param name="baseDir">A local or shared temp storage root directory.</param>
 		/// <param name="nodeName">Name of the node to get the file for</param>
 		/// <param name="blockName">Name of the output block to get the manifest for</param>
-		static FileReference GetManifestLocation(DirectoryReference baseDir, string nodeName, string blockName)
+		public static FileReference GetBlockManifestLocation(DirectoryReference baseDir, string nodeName, string? blockName)
 		{
 			return FileReference.Combine(baseDir, nodeName, String.IsNullOrEmpty(blockName)? "Manifest.xml" : String.Format("Manifest-{0}.xml", blockName));
 		}
@@ -850,7 +888,7 @@ namespace Horde.Storage.Utility
 		/// <param name="baseDir">A local or shared temp storage root directory.</param>
 		/// <param name="nodeName">Name of the node to get the file for</param>
 		/// <param name="tagName">Name of the tag to get the manifest for</param>
-		static FileReference GetTaggedFileListLocation(DirectoryReference baseDir, string nodeName, string tagName)
+		public static FileReference GetTagManifestLocation(DirectoryReference baseDir, string nodeName, string tagName)
 		{
 			Debug.Assert(tagName.StartsWith("#", StringComparison.Ordinal));
 			return FileReference.Combine(baseDir, nodeName, String.Format("Tag-{0}.xml", tagName.Substring(1)));
diff --git a/Engine/Source/Programs/Shared/EpicGames.Core/MemoryReader.cs b/Engine/Source/Programs/Shared/EpicGames.Core/MemoryReader.cs
--- a/Engine/Source/Programs/Shared/EpicGames.Core/MemoryReader.cs
+++ b/Engine/Source/Programs/Shared/EpicGames.Core/MemoryReader.cs
@@ -300,5 +300,38 @@ namespace EpicGames.Core
 			}
 			return array;
 		}
+
+		/// <summary>
+		/// Reads a dictionary from the writer
+		/// </summary>
+		/// <param name="reader">Reader to serialize from</param>
+		/// <param name="dictionary">The dictionary to read</param>
+		/// <param name="readKey">Delegate to write an individual key</param>
+		/// <param name="readValue">Delegate to write an individual value</param>
+		public static void ReadDictionary<TKey, TValue>(this IMemoryReader reader, Dictionary<TKey, TValue> dictionary, Func<TKey> readKey, Func<TValue> readValue) where TKey : notnull
+		{
+			int count = (int)reader.ReadUnsignedVarInt();
+			dictionary.EnsureCapacity(count);
+
+			for (int idx = 0; idx < count; idx++)
+			{
+				TKey key = readKey();
+				TValue value = readValue();
+				dictionary.Add(key, value);
+			}
+		}
+
+		/// <summary>
+		/// Reads a dictionary from the writer
+		/// </summary>
+		/// <param name="reader">Reader to serialize from</param>
+		/// <param name="readKey">Delegate to write an individual key</param>
+		/// <param name="readValue">Delegate to write an individual value</param>
+		public static Dictionary<TKey, TValue> ReadDictionary<TKey, TValue>(this IMemoryReader reader, Func<TKey> readKey, Func<TValue> readValue, IEqualityComparer<TKey>? comparer = null) where TKey : notnull
+		{
+			Dictionary<TKey, TValue> dictionary = new Dictionary<TKey, TValue>(comparer);
+			ReadDictionary(reader, dictionary, readKey, readValue);
+			return dictionary;
+		}
 	}
 }
diff --git a/Engine/Source/Programs/Shared/EpicGames.Core/MemoryWriter.cs b/Engine/Source/Programs/Shared/EpicGames.Core/MemoryWriter.cs
--- a/Engine/Source/Programs/Shared/EpicGames.Core/MemoryWriter.cs
+++ b/Engine/Source/Programs/Shared/EpicGames.Core/MemoryWriter.cs
@@ -416,6 +416,23 @@ namespace EpicGames.Core
 			}
 		}
 
+		/// <summary>
+		/// Writes a dictionary to the writer
+		/// </summary>
+		/// <param name="writer">Writer to serialize to</param>
+		/// <param name="dictionary">The dictionary to write</param>
+		/// <param name="writeKey">Delegate to write an individual key</param>
+		/// <param name="writeValue">Delegate to write an individual value</param>
+		public static void WriteDictionary<TKey, TValue>(this IMemoryWriter writer, IReadOnlyDictionary<TKey, TValue> dictionary, Action<TKey> writeKey, Action<TValue> writeValue)
+		{
+			writer.WriteUnsignedVarInt(dictionary.Count);
+			foreach (KeyValuePair<TKey, TValue> kvp in dictionary)
+			{
+				writeKey(kvp.Key);
+				writeValue(kvp.Value);
+			}
+		}
+
 		/// <summary>
 		/// Read a string from memory
 		/// </summary>
diff --git a/Engine/Source/Programs/Shared/EpicGames.Horde/Storage/BlobId.cs b/Engine/Source/Programs/Shared/EpicGames.Horde/Storage/BlobId.cs
--- a/Engine/Source/Programs/Shared/EpicGames.Horde/Storage/BlobId.cs
+++ b/Engine/Source/Programs/Shared/EpicGames.Horde/Storage/BlobId.cs
@@ -24,7 +24,7 @@ namespace EpicGames.Horde.Storage
 		/// <summary>
 		/// Dummy enum to allow invoking the constructor which takes a sanitized full path
 		/// </summary>
-		public enum Sanitize
+		public enum Validate
 		{
 			/// <summary>
 			/// Dummy value
@@ -56,11 +56,11 @@ namespace EpicGames.Horde.Storage
 		/// Constructor
 		/// </summary>
 		/// <param name="inner"></param>
-		/// <param name="sanitize"></param>
-		public BlobId(Utf8String inner, Sanitize sanitize)
+		/// <param name="validate"></param>
+		public BlobId(Utf8String inner, Validate validate)
 		{
 			Inner = inner;
-			_ = sanitize;
+			_ = validate;
 		}
 
 		static readonly Utf8String s_process;
@@ -108,7 +108,7 @@ namespace EpicGames.Horde.Storage
 			span = span.Slice(1);
 
 			GenerateUniqueId(span);
-			return new BlobId(new Utf8String(buffer), Sanitize.None);
+			return new BlobId(new Utf8String(buffer), Validate.None);
 		}
 
 		/// <summary>
@@ -173,6 +173,44 @@ namespace EpicGames.Horde.Storage
 			}
 		}
 
+		/// <summary>
+		/// Sanitize the given string to make a valid blob id
+		/// </summary>
+		/// <param name="name"></param>
+		/// <returns></returns>
+		public static Utf8String Sanitize(Utf8String name)
+		{
+			byte[] output = new byte[name.Length];
+
+			int outputIdx = 0;
+			for (int idx = 0; idx < name.Length; idx++)
+			{
+				if (name[idx] >= 'A' && name[idx] <= 'Z')
+				{
+					output[outputIdx++] = (byte)(name[idx] + 'a' - 'A');
+				}
+				else if ((name[idx] >= 'a' && name[idx] <= 'z') || (name[idx] >= '0' && name[idx] <= '9') || name[idx] == '+')
+				{
+					output[outputIdx++] = name[idx];
+				}
+				else if (name.Length > 0 && name[name.Length - 1] != '-')
+				{
+					output[outputIdx++] = (byte)'-';
+				}
+				else if (name[idx] == '/' && outputIdx > 0)
+				{
+					output[outputIdx++] = (byte)'/';
+				}
+			}
+
+			while (outputIdx > 0 && (output[outputIdx - 1] == '-' || output[outputIdx - 1] == '/'))
+			{
+				outputIdx--;
+			}
+
+			return new Utf8String(output.AsMemory(0, outputIdx));
+		}
+
 		/// <summary>
 		/// Test if a given character is valid in a store id
 		/// </summary>
diff --git a/Engine/Source/Programs/Shared/EpicGames.Horde/Storage/BlobLocator.cs b/Engine/Source/Programs/Shared/EpicGames.Horde/Storage/BlobLocator.cs
--- a/Engine/Source/Programs/Shared/EpicGames.Horde/Storage/BlobLocator.cs
+++ b/Engine/Source/Programs/Shared/EpicGames.Horde/Storage/BlobLocator.cs
@@ -63,7 +63,7 @@ namespace EpicGames.Horde.Storage
 			get
 			{
 				int colonIdx = Inner.LastIndexOf((byte)':');
-				return new BlobId(Inner.Substring(colonIdx + 1), BlobId.Sanitize.None);
+				return new BlobId(Inner.Substring(colonIdx + 1), BlobId.Validate.None);
 			}
 		}
 
diff --git a/Engine/Source/Programs/Shared/EpicGames.Horde/Storage/RefName.cs b/Engine/Source/Programs/Shared/EpicGames.Horde/Storage/RefName.cs
--- a/Engine/Source/Programs/Shared/EpicGames.Horde/Storage/RefName.cs
+++ b/Engine/Source/Programs/Shared/EpicGames.Horde/Storage/RefName.cs
@@ -39,6 +39,12 @@ namespace EpicGames.Horde.Storage
 			BlobId.ValidateArgument(nameof(text), text);
 		}
 
+		/// <summary>
+		/// Sanitize the given name
+		/// </summary>
+		/// <param name="name">Name to sanitize</param>
+		public static Utf8String Sanitize(Utf8String name) => BlobId.Sanitize(name);
+
 		/// <inheritdoc/>
 		public override bool Equals(object? obj) => obj is RefName refId && Equals(refId);
 
