Horde: Fix incorrect local manifest directories for transfer to/from temp storage.

diff --git a/Engine/Source/Programs/Horde/Horde.Agent/Execution/JobExecutor.cs b/Engine/Source/Programs/Horde/Horde.Agent/Execution/JobExecutor.cs
--- a/Engine/Source/Programs/Horde/Horde.Agent/Execution/JobExecutor.cs
+++ b/Engine/Source/Programs/Horde/Horde.Agent/Execution/JobExecutor.cs
@@ -645,6 +645,8 @@ namespace Horde.Agent.Execution
 
 		private async Task<bool> ExecuteWithTempStorageAsync(BeginStepResponse step, DirectoryReference workspaceDir, DirectoryReference sharedStorageDir, string arguments, ILogger logger, CancellationToken cancellationToken)
 		{
+			DirectoryReference manifestDir = DirectoryReference.Combine(workspaceDir, "Engine", "Saved", "BuildGraph");
+
 			// Create the storage client
 			using MemoryCache cache = new MemoryCache(new MemoryCacheOptions { });
 			FileStorageClient storage = new FileStorageClient(DirectoryReference.Combine(sharedStorageDir, "bundles"), cache, logger);
@@ -679,7 +681,7 @@ namespace Horde.Agent.Execution
 				scope.Span.SetTag("blocks", inputStorageBlocks.Count);
 				foreach (TempStorageBlockRef inputStorageBlock in inputStorageBlocks)
 				{
-					TempStorageBlockManifest manifest = await TempStorage.RetrieveBlockAsync(storage, inputStorageBlock.NodeName, inputStorageBlock.OutputName, workspaceDir, logger, cancellationToken);
+					TempStorageBlockManifest manifest = await TempStorage.RetrieveBlockAsync(storage, inputStorageBlock.NodeName, inputStorageBlock.OutputName, workspaceDir, manifestDir, logger, cancellationToken);
 					inputManifests[inputStorageBlock] = manifest;
 				}
 				scope.Span.SetTag("size", inputManifests.Sum(x => x.Value.GetTotalSize()));
diff --git a/Engine/Source/Programs/Horde/Horde.Agent/Utility/TempStorage.cs b/Engine/Source/Programs/Horde/Horde.Agent/Utility/TempStorage.cs
--- a/Engine/Source/Programs/Horde/Horde.Agent/Utility/TempStorage.cs
+++ b/Engine/Source/Programs/Horde/Horde.Agent/Utility/TempStorage.cs
@@ -677,16 +677,16 @@ namespace Horde.Storage.Utility
 		/// <param name="store"></param>
 		/// <param name="nodeName">Name of the node which produced the tag set</param>
 		/// <param name="tagName">Name of the tag, with a '#' prefix</param>
-		/// <param name="localDir">The local directory</param>
+		/// <param name="manifestDir">The local directory containing manifests</param>
 		/// <param name="logger">Logger for output</param>
 		/// <param name="cancellationToken"></param>
 		/// <returns>The set of files</returns>
-		public static async Task<TempStorageTagManifest> RetrieveTagAsync(IStorageClient store, string nodeName, string tagName, DirectoryReference localDir, ILogger logger, CancellationToken cancellationToken)
+		public static async Task<TempStorageTagManifest> RetrieveTagAsync(IStorageClient store, string nodeName, string tagName, DirectoryReference manifestDir, ILogger logger, CancellationToken cancellationToken)
 		{
 			TempStorageTagManifest fileList;
 
 			// Try to read the tag set from the local directory
-			FileReference localFileListLocation = GetTagManifestLocation(localDir, nodeName, tagName);
+			FileReference localFileListLocation = GetTagManifestLocation(manifestDir, nodeName, tagName);
 			if(FileReference.Exists(localFileListLocation))
 			{
 				logger.LogInformation("Reading local file list from {File}", localFileListLocation.FullName);
@@ -802,13 +802,14 @@ namespace Horde.Storage.Utility
 		/// <param name="nodeName">The node which created the storage block</param>
 		/// <param name="blockName">Name of the block to retrieve.</param>
 		/// <param name="rootDir">Local directory for extracting data to</param>
+		/// <param name="manifestDir">Local directory containing manifests</param>
 		/// <param name="logger">Logger for output</param>
 		/// <param name="cancellationToken"></param>
 		/// <returns>Manifest of the files retrieved</returns>
-		public static async Task<TempStorageBlockManifest> RetrieveBlockAsync(IStorageClient store, string nodeName, string blockName, DirectoryReference rootDir, ILogger logger, CancellationToken cancellationToken)
+		public static async Task<TempStorageBlockManifest> RetrieveBlockAsync(IStorageClient store, string nodeName, string blockName, DirectoryReference rootDir, DirectoryReference manifestDir, ILogger logger, CancellationToken cancellationToken)
 		{
 			// Get the path to the local manifest
-			FileReference localManifestFile = GetBlockManifestLocation(rootDir, nodeName, blockName);
+			FileReference localManifestFile = GetBlockManifestLocation(manifestDir, nodeName, blockName);
 			bool local = FileReference.Exists(localManifestFile);
 
 			// Read the manifest, either from local storage or shared storage
