Horde: Fix another temp storage local manifest path.

diff --git a/Engine/Source/Programs/Horde/Horde.Agent/Execution/JobExecutor.cs b/Engine/Source/Programs/Horde/Horde.Agent/Execution/JobExecutor.cs
--- a/Engine/Source/Programs/Horde/Horde.Agent/Execution/JobExecutor.cs
+++ b/Engine/Source/Programs/Horde/Horde.Agent/Execution/JobExecutor.cs
@@ -714,7 +714,7 @@ namespace Horde.Agent.Execution
 			// Read all the output manifests
 			foreach (string tagName in step.OutputNames)
 			{
-				FileReference tagFileListLocation = TempStorage.GetTagManifestLocation(workspaceDir, step.Name, tagName);
+				FileReference tagFileListLocation = TempStorage.GetTagManifestLocation(manifestDir, step.Name, tagName);
 				logger.LogInformation("Reading local file list from {File}", tagFileListLocation.FullName);
 
 				TempStorageTagManifest fileList = TempStorageTagManifest.Load(tagFileListLocation);
