Horde: Change XML root element names for temp storage manifests to match BuildGraph.

diff --git a/Engine/Source/Programs/Horde/Horde.Agent/Utility/TempStorage.cs b/Engine/Source/Programs/Horde/Horde.Agent/Utility/TempStorage.cs
--- a/Engine/Source/Programs/Horde/Horde.Agent/Utility/TempStorage.cs
+++ b/Engine/Source/Programs/Horde/Horde.Agent/Utility/TempStorage.cs
@@ -307,6 +307,7 @@ namespace Horde.Storage.Utility
 	/// <summary>
 	/// A manifest storing information about build products for a node's output
 	/// </summary>
+	[XmlRoot(ElementName = "TempStorageManifest")]
 	public class TempStorageBlockManifest
 	{
 		/// <summary>
@@ -397,6 +398,7 @@ namespace Horde.Storage.Utility
 	/// <summary>
 	/// Stores the contents of a tagged file set
 	/// </summary>
+	[XmlRoot(ElementName = "TempStorageFileList")]
 	public class TempStorageTagManifest
 	{
 		/// <summary>
