Horde: Implement missing temp storage serialization methods.

diff --git a/Engine/Source/Programs/Horde/Horde.Agent/Utility/TempStorage.cs b/Engine/Source/Programs/Horde/Horde.Agent/Utility/TempStorage.cs
--- a/Engine/Source/Programs/Horde/Horde.Agent/Utility/TempStorage.cs
+++ b/Engine/Source/Programs/Horde/Horde.Agent/Utility/TempStorage.cs
@@ -598,19 +598,24 @@ namespace Horde.Storage.Utility
 		static TempStorageBlockRef ReadTempStorageBlock(ITreeNodeReader reader)
 		{
 			string nodeName = reader.ReadString();
-			string blockName = reader.ReadString();
-			return new TempStorageBlockRef(nodeName, blockName);
+			string outputName = reader.ReadString();
+			return new TempStorageBlockRef(nodeName, outputName);
 		}
 
-		public override void Serialize(ITreeNodeWriter writer)
+		static void WriteTempStorageBlock(ITreeNodeWriter writer, TempStorageBlockRef blockRef)
 		{
-			throw new NotImplementedException();
+			writer.WriteString(blockRef.NodeName);
+			writer.WriteString(blockRef.OutputName);
 		}
 
-		public override IEnumerable<TreeNodeRef> EnumerateRefs()
+		public override void Serialize(ITreeNodeWriter writer)
 		{
-			throw new NotImplementedException();
+			writer.WriteVariableLengthArray(FileList.LocalFiles, x => writer.WriteString(x));
+			writer.WriteVariableLengthArray(FileList.ExternalFiles, x => writer.WriteString(x));
+			writer.WriteVariableLengthArray(FileList.Blocks, x => WriteTempStorageBlock(writer, x));
 		}
+
+		public override IEnumerable<TreeNodeRef> EnumerateRefs() => Enumerable.Empty<TreeNodeRef>();
 	}
 
 	/// <summary>
