Horde: Fix output path for temp storage tag manifest.

diff --git a/Engine/Source/Programs/Horde/Horde.Agent/Execution/JobExecutor.cs b/Engine/Source/Programs/Horde/Horde.Agent/Execution/JobExecutor.cs
--- a/Engine/Source/Programs/Horde/Horde.Agent/Execution/JobExecutor.cs
+++ b/Engine/Source/Programs/Horde/Horde.Agent/Execution/JobExecutor.cs
@@ -669,7 +669,7 @@ namespace Horde.Agent.Execution
 				string nodeName = input.Substring(0, slashIdx);
 				string tagName = input.Substring(slashIdx + 1);
 
-				TempStorageTagManifest fileList = await TempStorage.RetrieveTagAsync(storage, nodeName, tagName, workspaceDir, logger, cancellationToken);
+				TempStorageTagManifest fileList = await TempStorage.RetrieveTagAsync(storage, nodeName, tagName, manifestDir, logger, cancellationToken);
 				tagNameToFileSet[tagName] = fileList.ToFileSet(workspaceDir);
 				inputStorageBlocks.UnionWith(fileList.Blocks);
 			}
diff --git a/Engine/Source/Programs/Horde/Horde.Agent/Utility/TempStorage.cs b/Engine/Source/Programs/Horde/Horde.Agent/Utility/TempStorage.cs
--- a/Engine/Source/Programs/Horde/Horde.Agent/Utility/TempStorage.cs
+++ b/Engine/Source/Programs/Horde/Horde.Agent/Utility/TempStorage.cs
@@ -702,7 +702,7 @@ namespace Horde.Storage.Utility
 			else
 			{
 				RefName refName = GetRefNameForNode(nodeName);
-				logger.LogInformation("Reading tag {NodeName}:{TagName} from temp storage (ref: {RefName})", nodeName, tagName, refName);
+				logger.LogInformation("Reading tag {NodeName}:{TagName} from temp storage (ref: {RefName}, localFile: {LocalFile})", nodeName, tagName, refName, localFileListLocation);
 
 				TempStorageNode node = await store.ReadNodeAsync<TempStorageNode>(refName, cancellationToken: cancellationToken);
 
