Horde: Improve legibility of temp storage log events.

diff --git a/Engine/Source/Programs/Horde/Horde.Agent/Utility/TempStorage.cs b/Engine/Source/Programs/Horde/Horde.Agent/Utility/TempStorage.cs
--- a/Engine/Source/Programs/Horde/Horde.Agent/Utility/TempStorage.cs
+++ b/Engine/Source/Programs/Horde/Horde.Agent/Utility/TempStorage.cs
@@ -702,7 +702,7 @@ namespace Horde.Storage.Utility
 			else
 			{
 				RefName refName = GetRefNameForNode(nodeName);
-				logger.LogInformation("Reading tag {NodeName}:{TagName} from temp storage (ref: {RefName}, localFile: {LocalFile})", nodeName, tagName, refName, localFileListLocation);
+				logger.LogInformation("Reading node \"{NodeName}\" tag \"{TagName}\" from temp storage (ref: {RefName}, localFile: {LocalFile})", nodeName, tagName, refName, localFileListLocation);
 
 				TempStorageNode node = await store.ReadNodeAsync<TempStorageNode>(refName, cancellationToken: cancellationToken);
 
@@ -830,7 +830,7 @@ namespace Horde.Storage.Utility
 			{
 				// Read the shared manifest
 				RefName refName = GetRefNameForNode(nodeName);
-				logger.LogInformation("Reading block {NodeName}:{BlockName} from temp storage (ref: {RefName}, local: {LocalFile})", nodeName, blockName, refName, localManifestFile);
+				logger.LogInformation("Reading node \"{NodeName}\" block \"{BlockName}\" from temp storage (ref: {RefName}, local: {LocalFile})", nodeName, blockName, refName, localManifestFile);
 
 				TempStorageNode node = await store.ReadNodeAsync<TempStorageNode>(refName, cancellationToken: cancellationToken);
 				
