Horde: Allow commit tag definitions to derive from existing commit tag definitions, using the "Base" property.

diff --git a/Engine/Source/Programs/Horde/Horde.Agent/Utility/TempStorage.cs b/Engine/Source/Programs/Horde/Horde.Agent/Utility/TempStorage.cs
--- a/Engine/Source/Programs/Horde/Horde.Agent/Utility/TempStorage.cs
+++ b/Engine/Source/Programs/Horde/Horde.Agent/Utility/TempStorage.cs
@@ -735,7 +735,7 @@ namespace Horde.Storage.Utility
 		/// <returns></returns>
 		public static async Task<TreeNodeRef<TempStorageTagNode>> ArchiveTagAsync(TreeWriter writer, string tagName, DirectoryReference rootDir, IEnumerable<FileReference> files, IEnumerable<TempStorageBlockRef> blocks, ILogger logger, CancellationToken cancellationToken)
 		{
-			logger.LogInformation("Creating output tag {TagName}", tagName);
+			logger.LogInformation("Creating output tag \"{TagName}\"", tagName);
 
 			TempStorageTagManifest fileList = new TempStorageTagManifest(files, rootDir, blocks);
 			TreeNodeRef<TempStorageTagNode> tagNode = new TreeNodeRef<TempStorageTagNode>(new TempStorageTagNode(fileList));
@@ -756,7 +756,7 @@ namespace Horde.Storage.Utility
 		/// <returns>The created manifest instance (which has already been saved to disk).</returns>
 		public static async Task<TreeNodeRef<TempStorageBlockNode>> ArchiveBlockAsync(TreeWriter writer, string blockName, DirectoryReference rootDir, FileReference[] buildProducts, ILogger logger, CancellationToken cancellationToken)
 		{
-			logger.LogInformation("Creating output block {BlockName}", blockName);
+			logger.LogInformation("Creating output block \"{BlockName}\"", blockName);
 
 			// Create a manifest for the given build products
 			FileInfo[] files = buildProducts.Select(x => new FileInfo(x.FullName)).ToArray();
@@ -837,7 +837,7 @@ namespace Horde.Storage.Utility
 				TreeNodeRef<TempStorageBlockNode>? blockNodeRef;
 				if (!node.Blocks.TryGetValue(blockName, out blockNodeRef))
 				{
-					throw new TempStorageException($"Missing block {blockName} from node {nodeName}");
+					throw new TempStorageException($"Missing block \"{blockName}\" from node \"{nodeName}\"");
 				}
 
 				TempStorageBlockNode blockNode = await blockNodeRef.ExpandAsync(cancellationToken);
diff --git a/Engine/Source/Programs/Horde/Horde.Build.Tests/CommitTagTests.cs b/Engine/Source/Programs/Horde/Horde.Build.Tests/CommitTagTests.cs
new file mode 100644
--- /dev/null
+++ b/Engine/Source/Programs/Horde/Horde.Build.Tests/CommitTagTests.cs
@@ -0,0 +1,73 @@
+// Copyright Epic Games, Inc. All Rights Reserved.
+
+using System.Collections.Generic;
+using EpicGames.Core;
+using Horde.Build.Perforce;
+using Horde.Build.Streams;
+using Microsoft.VisualStudio.TestTools.UnitTesting;
+
+namespace Horde.Build.Tests
+{
+	[TestClass]
+	public class CommitTagTests
+	{
+		[TestMethod]
+		public void TestCode()
+		{
+			StreamConfig config = new StreamConfig();
+
+			FileFilter? filter;
+			Assert.IsTrue(config.TryGetCommitTagFilter(CommitTag.Code, out filter));
+
+			Assert.IsTrue(filter.Matches("foo.cpp"));
+			Assert.IsTrue(filter.Matches("foo.h"));
+			Assert.IsTrue(filter.Matches("foo.cs"));
+			Assert.IsFalse(filter.Matches("foo.uasset"));
+			Assert.IsFalse(filter.Matches("foo.png"));
+		}
+
+		[TestMethod]
+		public void TestContent()
+		{
+			StreamConfig config = new StreamConfig();
+
+			FileFilter? filter;
+			Assert.IsTrue(config.TryGetCommitTagFilter(CommitTag.Content, out filter));
+
+			Assert.IsFalse(filter.Matches("foo.cpp"));
+			Assert.IsFalse(filter.Matches("foo.h"));
+			Assert.IsFalse(filter.Matches("foo.cs"));
+			Assert.IsTrue(filter.Matches("foo.uasset"));
+			Assert.IsTrue(filter.Matches("foo.png"));
+		}
+
+		[TestMethod]
+		public void TestCustom()
+		{
+			StreamConfig config = new StreamConfig();
+			config.CommitTags.Add(new CommitTagConfig { Name = new CommitTag("test"), Filter = new List<string> { "*TEST*", "-*TEST2*" } });
+
+			FileFilter? filter;
+			Assert.IsTrue(config.TryGetCommitTagFilter(new CommitTag("test"), out filter));
+
+			Assert.IsFalse(filter.Matches("somefile"));
+			Assert.IsTrue(filter.Matches("sometestfile"));
+			Assert.IsFalse(filter.Matches("sometest2file"));
+		}
+
+		[TestMethod]
+		public void TestInheritancee()
+		{
+			StreamConfig config = new StreamConfig();
+			config.CommitTags.Add(new CommitTagConfig { Name = new CommitTag("test"), Base = CommitTag.Code, Filter = new List<string> { "*TEST*", "-*TEST2*" } });
+
+			FileFilter? filter;
+			Assert.IsTrue(config.TryGetCommitTagFilter(new CommitTag("test"), out filter));
+
+			Assert.IsTrue(filter.Matches("somefile.cpp"));
+			Assert.IsFalse(filter.Matches("somefile.uasset"));
+			Assert.IsTrue(filter.Matches("sometestfile.uasset"));
+			Assert.IsFalse(filter.Matches("sometest2file.uasset"));
+		}
+	}
+}
diff --git a/Engine/Source/Programs/Horde/Horde.Build.Tests/Stubs/Services/PerforceServiceStub.cs b/Engine/Source/Programs/Horde/Horde.Build.Tests/Stubs/Services/PerforceServiceStub.cs
--- a/Engine/Source/Programs/Horde/Horde.Build.Tests/Stubs/Services/PerforceServiceStub.cs
+++ b/Engine/Source/Programs/Horde/Horde.Build.Tests/Stubs/Services/PerforceServiceStub.cs
@@ -6,11 +6,13 @@ using System.Linq;
 using System.Runtime.CompilerServices;
 using System.Threading;
 using System.Threading.Tasks;
+using EpicGames.Core;
 using EpicGames.Perforce;
 using Horde.Build.Perforce;
 using Horde.Build.Streams;
 using Horde.Build.Users;
 using Horde.Build.Utilities;
+using Microsoft.VisualStudio.TestTools.UnitTesting;
 
 namespace Horde.Build.Tests.Stubs.Services
 {
@@ -60,12 +62,16 @@ namespace Horde.Build.Tests.Stubs.Services
 
 			public ValueTask<IReadOnlyList<CommitTag>> GetTagsAsync(CancellationToken cancellationToken)
 			{
+				StreamConfig config = new StreamConfig();
+				Assert.IsTrue(config.TryGetCommitTagFilter(CommitTag.Code, out FileFilter? codeFilter));
+				Assert.IsTrue(config.TryGetCommitTagFilter(CommitTag.Content, out FileFilter? contentFilter));
+
 				List<CommitTag> tags = new List<CommitTag>();
-				if (CommitTagConfig.CodeDefault.CreateFileFilter().ApplyTo(Files).Any())
+				if (codeFilter.ApplyTo(Files).Any())
 				{
 					tags.Add(CommitTag.Code);
 				}
-				if (CommitTagConfig.ContentDefault.CreateFileFilter().ApplyTo(Files).Any())
+				if (contentFilter.ApplyTo(Files).Any())
 				{
 					tags.Add(CommitTag.Content);
 				}
diff --git a/Engine/Source/Programs/Horde/Horde.Build/Perforce/CommitTag.cs b/Engine/Source/Programs/Horde/Horde.Build/Perforce/CommitTag.cs
--- a/Engine/Source/Programs/Horde/Horde.Build/Perforce/CommitTag.cs
+++ b/Engine/Source/Programs/Horde/Horde.Build/Perforce/CommitTag.cs
@@ -57,6 +57,11 @@ namespace Horde.Build.Perforce
 			StringId.ValidateArgument(text, nameof(text));
 		}
 
+		/// <summary>
+		/// Check if the current tag is empty
+		/// </summary>
+		public bool IsEmpty() => Text == null;
+
 		/// <inheritdoc/>
 		public override bool Equals(object? obj) => obj is CommitTag id && Equals(id);
 
diff --git a/Engine/Source/Programs/Horde/Horde.Build/Perforce/PerforceService.cs b/Engine/Source/Programs/Horde/Horde.Build/Perforce/PerforceService.cs
--- a/Engine/Source/Programs/Horde/Horde.Build/Perforce/PerforceService.cs
+++ b/Engine/Source/Programs/Horde/Horde.Build/Perforce/PerforceService.cs
@@ -186,7 +186,7 @@ namespace Horde.Build.Perforce
 				IReadOnlyList<string> files = await GetFilesAsync(cancellationToken);
 				foreach (CommitTagConfig tagConfig in _stream.Config.GetAllCommitTags())
 				{
-					if (tagConfig.CreateFileFilter().ApplyTo(files).Any())
+					if (_stream.Config.TryGetCommitTagFilter(tagConfig.Name, out FileFilter? filter) && filter.ApplyTo(files).Any())
 					{
 						commitTags.Add(tagConfig.Name);
 					}
@@ -801,8 +801,8 @@ namespace Horde.Build.Perforce
 				List<CommitTag> tags = new List<CommitTag>();
 				foreach (CommitTagConfig tagConfig in stream.Config.GetAllCommitTags())
 				{
-					FileFilter filter = tagConfig.CreateFileFilter();
-					if (mappedFiles.Any(x => filter.Matches(x)))
+					FileFilter? filter;
+					if (stream.Config.TryGetCommitTagFilter(tagConfig.Name, out filter) && mappedFiles.Any(x => filter.Matches(x)))
 					{
 						tags.Add(tagConfig.Name);
 					}
@@ -1190,9 +1190,10 @@ namespace Horde.Build.Perforce
 						List<FileFilter> filters = new List<FileFilter>();
 						foreach (CommitTag tag in tags)
 						{
-							if (_stream.Config.TryGetCommitTag(tag, out CommitTagConfig? commitTagConfig))
+							FileFilter? filter;
+							if (_stream.Config.TryGetCommitTagFilter(tag, out filter))
 							{
-								filters.Add(commitTagConfig.CreateFileFilter());
+								filters.Add(filter);
 							}
 						}
 
diff --git a/Engine/Source/Programs/Horde/Horde.Build/Perforce/PerforceServiceCache.cs b/Engine/Source/Programs/Horde/Horde.Build/Perforce/PerforceServiceCache.cs
--- a/Engine/Source/Programs/Horde/Horde.Build/Perforce/PerforceServiceCache.cs
+++ b/Engine/Source/Programs/Horde/Horde.Build/Perforce/PerforceServiceCache.cs
@@ -52,10 +52,10 @@ namespace Horde.Build.Perforce
 			public CommitTag Name { get; }
 			public FileFilter Filter { get; }
 
-			public CommitTagInfo(CommitTagConfig config)
+			public CommitTagInfo(CommitTag name, FileFilter filter)
 			{
-				Name = config.Name;
-				Filter = config.CreateFileFilter();
+				Name = name;
+				Filter = filter;
 			}
 		}
 
@@ -63,13 +63,20 @@ namespace Horde.Build.Perforce
 		{
 			public IStream Stream { get; set; }
 			public PerforceViewMap View { get; }
-			public List<CommitTagInfo> CommitTags { get; set; }
+			public List<CommitTagInfo> CommitTags { get; set; } = new List<CommitTagInfo>();
 
 			public StreamInfo(IStream stream, PerforceViewMap view)
 			{
 				Stream = stream;
 				View = view;
-				CommitTags = stream.Config.GetAllCommitTags().Select(x => new CommitTagInfo(x)).ToList();
+
+				foreach (CommitTagConfig commitTagConfig in stream.Config.GetAllCommitTags())
+				{
+					if (stream.Config.TryGetCommitTagFilter(commitTagConfig.Name, out FileFilter? filter))
+					{
+						CommitTags.Add(new CommitTagInfo(commitTagConfig.Name, filter));
+					}
+				}
 			}
 		}
 
diff --git a/Engine/Source/Programs/Horde/Horde.Build/Streams/StreamConfig.cs b/Engine/Source/Programs/Horde/Horde.Build/Streams/StreamConfig.cs
--- a/Engine/Source/Programs/Horde/Horde.Build/Streams/StreamConfig.cs
+++ b/Engine/Source/Programs/Horde/Horde.Build/Streams/StreamConfig.cs
@@ -236,6 +236,54 @@ namespace Horde.Build.Streams
 			}
 		}
 
+		/// <summary>
+		/// Constructs a <see cref="FileFilter"/> from the rules in this configuration object
+		/// </summary>
+		/// <returns>Filter object</returns>
+		public bool TryGetCommitTagFilter(CommitTag tag, [NotNullWhen(true)] out FileFilter? filter)
+		{
+			FileFilter newFilter = new FileFilter(FileFilterType.Exclude);
+			if (TryGetCommitTagFilter(tag, newFilter, new HashSet<CommitTag>()))
+			{
+				filter = newFilter;
+				return true;
+			}
+			else
+			{
+				filter = null;
+				return false;
+			}
+		}
+
+		/// <summary>
+		/// Find the filter for a given tag, recursively
+		/// </summary>
+		bool TryGetCommitTagFilter(CommitTag tag, FileFilter filter, HashSet<CommitTag> visitedTags)
+		{
+			// Check we don't have a recursive definition for the tag
+			if (!visitedTags.Add(tag))
+			{
+				return false;
+			}
+
+			// Get the tag configuration
+			CommitTagConfig? config;
+			if (!TryGetCommitTag(tag, out config))
+			{
+				return false;
+			}
+
+			// Add rules from the base tag
+			if (!config.Base.IsEmpty() && !TryGetCommitTagFilter(config.Base, filter, visitedTags))
+			{
+				return false;
+			}
+
+			// Add rules from this tag
+			filter.AddRules(config.Filter);
+			return true;
+		}
+
 		/// <summary>
 		/// Tries to find a template with the given id
 		/// </summary>
@@ -854,6 +902,11 @@ namespace Horde.Build.Streams
 		[Required]
 		public CommitTag Name { get; set; }
 
+		/// <summary>
+		/// Base tag to copy settings from
+		/// </summary>
+		public CommitTag Base { get; set; }
+
 		/// <summary>
 		/// List of files to be included in this filter
 		/// </summary>
@@ -868,16 +921,5 @@ namespace Horde.Build.Streams
 		/// Default config for content filters
 		/// </summary>
 		public static CommitTagConfig ContentDefault { get; } = new CommitTagConfig { Name = CommitTag.Content, Filter = CommitTag.ContentFilter.ToList() };
-
-		/// <summary>
-		/// Constructs a <see cref="FileFilter"/> from the rules in this configuration object
-		/// </summary>
-		/// <returns>Filter object</returns>
-		public FileFilter CreateFileFilter()
-		{
-			FileFilter filter = new FileFilter(FileFilterType.Exclude);
-			filter.AddRules(Filter);
-			return filter;
-		}
 	}
 }
