Horde: Disable P4 metadata fetching during downtime windows.

diff --git a/Engine/Source/Programs/Horde/Horde.Build/Perforce/PerforceServiceCache.cs b/Engine/Source/Programs/Horde/Horde.Build/Perforce/PerforceServiceCache.cs
--- a/Engine/Source/Programs/Horde/Horde.Build/Perforce/PerforceServiceCache.cs
+++ b/Engine/Source/Programs/Horde/Horde.Build/Perforce/PerforceServiceCache.cs
@@ -161,6 +161,7 @@ namespace Horde.Build.Perforce
 		}
 
 		readonly MongoService _mongoService;
+		readonly DowntimeService _downtimeService;
 		readonly ServerSettings _settings;
 		readonly IMongoCollection<CachedCommitDoc> _commits;
 		readonly IStreamCollection _streamCollection;
@@ -170,10 +171,11 @@ namespace Horde.Build.Perforce
 		/// <summary>
 		/// Constructor
 		/// </summary>
-		public PerforceServiceCache(PerforceLoadBalancer loadBalancer, MongoService mongoService, GlobalsService globalsService, IUserCollection userCollection, IStreamCollection streamCollection, IClock clock, IOptions<ServerSettings> settings, ILogger<PerforceService> logger)
+		public PerforceServiceCache(PerforceLoadBalancer loadBalancer, MongoService mongoService, DowntimeService downtimeService, GlobalsService globalsService, IUserCollection userCollection, IStreamCollection streamCollection, IClock clock, IOptions<ServerSettings> settings, ILogger<PerforceService> logger)
 			: base(loadBalancer, globalsService, userCollection, settings, logger)
 		{
 			_mongoService = mongoService;
+			_downtimeService = downtimeService;
 			_settings = settings.Value;
 
 			List<MongoIndex<CachedCommitDoc>> indexes = new List<MongoIndex<CachedCommitDoc>>();
@@ -222,6 +224,12 @@ namespace Horde.Build.Perforce
 		/// <returns></returns>
 		async ValueTask UpdateCommitsAsync(CancellationToken cancellationToken)
 		{
+			// Don't do any updates during downtime; we might just create a bunch of P4 errors.
+			if (_downtimeService.IsDowntimeActive)
+			{
+				return;
+			}
+
 			CacheState state = await _mongoService.GetSingletonAsync<CacheState>();
 
 			// Get the current list of streams and their views
