Horde: Fix tests.

diff --git a/Engine/Source/Programs/Horde/Horde.Build/Perforce/PerforceServiceCache.cs b/Engine/Source/Programs/Horde/Horde.Build/Perforce/PerforceServiceCache.cs
--- a/Engine/Source/Programs/Horde/Horde.Build/Perforce/PerforceServiceCache.cs
+++ b/Engine/Source/Programs/Horde/Horde.Build/Perforce/PerforceServiceCache.cs
@@ -161,7 +161,7 @@ namespace Horde.Build.Perforce
 		}
 
 		readonly MongoService _mongoService;
-		readonly DowntimeService _downtimeService;
+		readonly IDowntimeService _downtimeService;
 		readonly ServerSettings _settings;
 		readonly IMongoCollection<CachedCommitDoc> _commits;
 		readonly IStreamCollection _streamCollection;
@@ -171,7 +171,7 @@ namespace Horde.Build.Perforce
 		/// <summary>
 		/// Constructor
 		/// </summary>
-		public PerforceServiceCache(PerforceLoadBalancer loadBalancer, MongoService mongoService, DowntimeService downtimeService, GlobalsService globalsService, IUserCollection userCollection, IStreamCollection streamCollection, IClock clock, IOptions<ServerSettings> settings, ILogger<PerforceService> logger)
+		public PerforceServiceCache(PerforceLoadBalancer loadBalancer, MongoService mongoService, IDowntimeService downtimeService, GlobalsService globalsService, IUserCollection userCollection, IStreamCollection streamCollection, IClock clock, IOptions<ServerSettings> settings, ILogger<PerforceService> logger)
 			: base(loadBalancer, globalsService, userCollection, settings, logger)
 		{
 			_mongoService = mongoService;
