Horde: Fix tests.

diff --git a/Engine/Source/Programs/Horde/Horde.Agent.Tests/Helper.cs b/Engine/Source/Programs/Horde/Horde.Agent.Tests/Helper.cs
--- a/Engine/Source/Programs/Horde/Horde.Agent.Tests/Helper.cs
+++ b/Engine/Source/Programs/Horde/Horde.Agent.Tests/Helper.cs
@@ -207,7 +207,7 @@ namespace Horde.Agent.Tests
 		private readonly Func<BeginStepResponse, ILogger, CancellationToken, Task<JobStepOutcome>> _func;
 
 		public SimpleTestExecutor(Func<BeginStepResponse, ILogger, CancellationToken, Task<JobStepOutcome>> func)
-			: base(null!, null!, null!, null!)
+			: base(null!, null!, null!, null!, null!)
 		{
 			_func = func;
 		}
