Horde: Allow specifying prefixes for blobs uploaded using HTTP storage client, via the 'prefix' form parameter.

diff --git a/Engine/Source/Programs/Horde/Horde.Build/Storage/StorageController.cs b/Engine/Source/Programs/Horde/Horde.Build/Storage/StorageController.cs
--- a/Engine/Source/Programs/Horde/Horde.Build/Storage/StorageController.cs
+++ b/Engine/Source/Programs/Horde/Horde.Build/Storage/StorageController.cs
@@ -94,10 +94,11 @@ namespace Horde.Build.Storage
 		/// </summary>
 		/// <param name="namespaceId">Namespace to fetch from</param>
 		/// <param name="file">Data to be uploaded</param>
+		/// <param name="prefix">Prefix for the uploaded file</param>
 		/// <param name="cancellationToken">Cancellation token for the operation</param>
 		[HttpPost]
 		[Route("/api/v1/storage/{namespaceId}/blobs")]
-		public async Task<ActionResult<WriteBlobResponse>> WriteBlobAsync(NamespaceId namespaceId, IFormFile file, CancellationToken cancellationToken)
+		public async Task<ActionResult<WriteBlobResponse>> WriteBlobAsync(NamespaceId namespaceId, IFormFile file, [FromForm] string? prefix = default, CancellationToken cancellationToken = default)
 		{
 			if (!await _storageService.AuthorizeAsync(namespaceId, User, AclAction.WriteBlobs, null, cancellationToken))
 			{
@@ -107,7 +108,7 @@ namespace Horde.Build.Storage
 			IStorageClient client = await _storageService.GetClientAsync(namespaceId, cancellationToken);
 			using (Stream stream = file.OpenReadStream())
 			{
-				BlobLocator locator = await client.WriteBlobAsync(stream, cancellationToken: cancellationToken);
+				BlobLocator locator = await client.WriteBlobAsync(stream, prefix: (prefix == null)? Utf8String.Empty : new Utf8String(prefix), cancellationToken: cancellationToken);
 				return new WriteBlobResponse { Locator = locator };
 			}
 		}
diff --git a/Engine/Source/Programs/Shared/EpicGames.Horde/Storage/Backends/HttpStorageClient.cs b/Engine/Source/Programs/Shared/EpicGames.Horde/Storage/Backends/HttpStorageClient.cs
--- a/Engine/Source/Programs/Shared/EpicGames.Horde/Storage/Backends/HttpStorageClient.cs
+++ b/Engine/Source/Programs/Shared/EpicGames.Horde/Storage/Backends/HttpStorageClient.cs
@@ -72,9 +72,11 @@ namespace EpicGames.Horde.Storage.Backends
 			using (HttpRequestMessage request = new HttpRequestMessage(HttpMethod.Post, $"api/v1/storage/{_namespaceId}/blobs"))
 			{
 				using StreamContent streamContent = new StreamContent(stream);
+				using StringContent stringContent = new StringContent(prefix.ToString());
 
 				MultipartFormDataContent form = new MultipartFormDataContent();
 				form.Add(streamContent, "file", "filename");
+				form.Add(stringContent, "prefix");
 
 				request.Content = form;
 				using (HttpResponseMessage response = await _httpClient.SendAsync(request, cancellationToken))
