Horde: Include the stream, changelist number, and job id in uploaded temp storage artifacts.

diff --git a/Engine/Source/Programs/Horde/Horde.Agent/Execution/JobExecutor.cs b/Engine/Source/Programs/Horde/Horde.Agent/Execution/JobExecutor.cs
--- a/Engine/Source/Programs/Horde/Horde.Agent/Execution/JobExecutor.cs
+++ b/Engine/Source/Programs/Horde/Horde.Agent/Execution/JobExecutor.cs
@@ -821,7 +821,7 @@ namespace Horde.Agent.Execution
 			// Write all the storage blocks, and update the mapping from file to storage block
 			using (GlobalTracer.Instance.BuildSpan("TempStorage").WithTag("resource", "Write").StartActive())
 			{
-				RefName refName = new RefName(RefName.Sanitize(step.Name));
+				RefName refName = new RefName(RefName.Sanitize($"{_job.StreamId}/{_job.Change}-{_jobId}/{step.Name}"));
 
 				TreeOptions treeOptions = new TreeOptions();
 				TreeWriter treeWriter = new TreeWriter(storage, treeOptions, refName.Text);
