Horde: Rename AutoscaleServiceV2 to FleetService

Also remove unused AutoscaleService v1.

diff --git a/Engine/Source/Programs/Horde/Horde.Build.Tests/Fleet/AutoscaleServiceV2Test.cs b/Engine/Source/Programs/Horde/Horde.Build.Tests/Fleet/AutoscaleServiceV2Test.cs
--- a/Engine/Source/Programs/Horde/Horde.Build.Tests/Fleet/AutoscaleServiceV2Test.cs
+++ b/Engine/Source/Programs/Horde/Horde.Build.Tests/Fleet/AutoscaleServiceV2Test.cs
@@ -100,7 +100,7 @@ namespace Horde.Build.Tests.Fleet
 		[TestMethod]
 		public async Task OnlyEnabledAgentsAreAutoScaled()
 		{
-			using AutoscaleServiceV2 service = GetAutoscaleService(_fleetManagerSpy);
+			using FleetService service = GetAutoscaleService(_fleetManagerSpy);
 			IPool pool = await PoolService.CreatePoolAsync("testPool", null, true, 0, 0, sizeStrategy: PoolSizeStrategy.LeaseUtilization);
 			await CreateAgentAsync(pool, true);
 			await CreateAgentAsync(pool, true);
@@ -115,7 +115,7 @@ namespace Horde.Build.Tests.Fleet
 		[TestMethod]
 		public async Task ExceptionsAreSwallowedAndLogged()
 		{
-			using AutoscaleServiceV2 service = GetAutoscaleService(_fleetManagerSpy);
+			using FleetService service = GetAutoscaleService(_fleetManagerSpy);
 			IPool pool = await PoolService.CreatePoolAsync("testPool", null, true, 0, 0, sizeStrategy: PoolSizeStrategy.LeaseUtilization);
 			List<PoolSizeData> poolSizeDatas = new()
 			{
@@ -130,7 +130,7 @@ namespace Horde.Build.Tests.Fleet
 		[TestMethod]
 		public async Task ScaleOutCooldown()
 		{
-			using AutoscaleServiceV2 service = GetAutoscaleService(_fleetManagerSpy);
+			using FleetService service = GetAutoscaleService(_fleetManagerSpy);
 			IPool pool = await PoolService.CreatePoolAsync("testPool", null, true, 0, 0, sizeStrategy: PoolSizeStrategy.NoOp);
 
 			// First scale-out will succeed
@@ -150,7 +150,7 @@ namespace Horde.Build.Tests.Fleet
 		[TestMethod]
 		public async Task ScaleInCooldown()
 		{
-			using AutoscaleServiceV2 service = GetAutoscaleService(_fleetManagerSpy);
+			using FleetService service = GetAutoscaleService(_fleetManagerSpy);
 			IPool pool = await PoolService.CreatePoolAsync("testPool", null, true, 0, 0, sizeStrategy: PoolSizeStrategy.NoOp);
 			IAgent agent1 = await CreateAgentAsync(pool);
 			IAgent agent2 = await CreateAgentAsync(pool);
@@ -169,15 +169,15 @@ namespace Horde.Build.Tests.Fleet
 			Assert.AreEqual(2, _fleetManagerSpy.ShrinkPoolAsyncCallCount);
 		}
 
-		private AutoscaleServiceV2 GetAutoscaleService(IFleetManager fleetManager)
+		private FleetService GetAutoscaleService(IFleetManager fleetManager)
 		{
 			ILoggerFactory loggerFactory = ServiceProvider.GetRequiredService<ILoggerFactory>();
 			IOptions<ServerSettings> serverSettingsOpt = ServiceProvider.GetRequiredService<IOptions<ServerSettings>>();
 			serverSettingsOpt.Value.FleetManagerV2 = FleetManagerType.AwsReuse;
 			
-			AutoscaleServiceV2 service = new(
+			FleetService service = new(
 				AgentCollection, GraphCollection, JobCollection, LeaseCollection, PoolCollection, StreamService, _dogStatsD,
-				new StubFleetManagerFactory(fleetManager), Clock, Cache, serverSettingsOpt, loggerFactory.CreateLogger<AutoscaleServiceV2>());
+				new StubFleetManagerFactory(fleetManager), Clock, Cache, serverSettingsOpt, loggerFactory.CreateLogger<FleetService>());
 				
 			return service;
 		}
@@ -192,10 +192,10 @@ namespace Horde.Build.Tests.Fleet
 		public async Task CreateJobQueueFromLegacySettings()
 		{
 			IPool pool1 = await PoolService.CreatePoolAsync("test1", sizeStrategy: PoolSizeStrategy.JobQueue);
-			Assert.AreEqual(typeof(JobQueueStrategy), AutoscaleServiceV2.CreatePoolSizeStrategy(pool1).GetType());
+			Assert.AreEqual(typeof(JobQueueStrategy), FleetService.CreatePoolSizeStrategy(pool1).GetType());
 			
 			IPool pool2 = await PoolService.CreatePoolAsync("test2", sizeStrategy: PoolSizeStrategy.JobQueue, jobQueueSettings: new JobQueueSettings(22, 33));
-			IPoolSizeStrategy s = AutoscaleServiceV2.CreatePoolSizeStrategy(pool2);
+			IPoolSizeStrategy s = FleetService.CreatePoolSizeStrategy(pool2);
 			Assert.AreEqual(typeof(JobQueueStrategy), s.GetType());
 			Assert.AreEqual(22.0, ((JobQueueStrategy)s).Settings.ScaleOutFactor);
 			Assert.AreEqual(33.0, ((JobQueueStrategy)s).Settings.ScaleInFactor);
@@ -205,7 +205,7 @@ namespace Horde.Build.Tests.Fleet
 		public async Task CreateLeaseUtilizationFromLegacySettings()
 		{
 			IPool pool = await PoolService.CreatePoolAsync("test1", sizeStrategy: PoolSizeStrategy.LeaseUtilization);
-			Assert.AreEqual(typeof(LeaseUtilizationStrategy), AutoscaleServiceV2.CreatePoolSizeStrategy(pool).GetType());
+			Assert.AreEqual(typeof(LeaseUtilizationStrategy), FleetService.CreatePoolSizeStrategy(pool).GetType());
 		}
 
 		[TestMethod]
@@ -288,7 +288,7 @@ namespace Horde.Build.Tests.Fleet
 		{
 			IPool pool = await PoolService.CreatePoolAsync("testPool-" + s_poolCount++, null, true, 0, 0, sizeStrategy: PoolSizeStrategy.NoOp);
 			await PoolService.UpdatePoolAsync(pool, newSizeStrategies: infos.ToList());
-			return AutoscaleServiceV2.CreatePoolSizeStrategy(pool);
+			return FleetService.CreatePoolSizeStrategy(pool);
 		}
 	}
 
@@ -363,7 +363,7 @@ namespace Horde.Build.Tests.Fleet
 		private async Task<IFleetManager> CreateFleetManager(params FleetManagerInfo[] infos)
 		{
 			IPool pool = await PoolService.CreatePoolAsync("testPool-" + s_poolCount++, null, true, 0, 0, fleetManagers: infos.ToList(), sizeStrategy: PoolSizeStrategy.NoOp);
-			return AutoscaleServiceV2.CreateFleetManager(pool);
+			return FleetService.CreateFleetManager(pool);
 		}
 	}
 }
\ No newline at end of file
diff --git a/Engine/Source/Programs/Horde/Horde.Build.Tests/Fleet/LeaseUtilizationStrategyTest.cs b/Engine/Source/Programs/Horde/Horde.Build.Tests/Fleet/LeaseUtilizationStrategyTest.cs
--- a/Engine/Source/Programs/Horde/Horde.Build.Tests/Fleet/LeaseUtilizationStrategyTest.cs
+++ b/Engine/Source/Programs/Horde/Horde.Build.Tests/Fleet/LeaseUtilizationStrategyTest.cs
@@ -97,7 +97,7 @@ namespace Horde.Build.Tests.Fleet
 
 		private async Task AssertPoolSizeAsync(IPool pool, int expectedNumAgents)
 		{
-			IPoolSizeStrategy strategy = AutoscaleServiceV2.CreatePoolSizeStrategy(pool);
+			IPoolSizeStrategy strategy = FleetService.CreatePoolSizeStrategy(pool);
 			List<PoolSizeData> output = await strategy.CalcDesiredPoolSizesAsync(new() { new(pool, _poolAgents, null) });
 			Assert.AreEqual(1, output.Count);
 			Assert.AreEqual(expectedNumAgents, output[0].DesiredAgentCount);
diff --git a/Engine/Source/Programs/Horde/Horde.Build.Tests/TestSetup.cs b/Engine/Source/Programs/Horde/Horde.Build.Tests/TestSetup.cs
--- a/Engine/Source/Programs/Horde/Horde.Build.Tests/TestSetup.cs
+++ b/Engine/Source/Programs/Horde/Horde.Build.Tests/TestSetup.cs
@@ -98,7 +98,7 @@ namespace Horde.Build.Tests
 		public IDeviceCollection DeviceCollection => ServiceProvider.GetRequiredService<IDeviceCollection>();
 
 		public AclService AclService => ServiceProvider.GetRequiredService<AclService>();
-		public AutoscaleServiceV2 AutoscaleServiceV2 => ServiceProvider.GetRequiredService<AutoscaleServiceV2>();
+		public FleetService FleetService => ServiceProvider.GetRequiredService<FleetService>();
 		public AgentSoftwareService AgentSoftwareService => ServiceProvider.GetRequiredService<AgentSoftwareService>();
 		public AgentService AgentService => ServiceProvider.GetRequiredService<AgentService>();
 		public ICommitService CommitService => ServiceProvider.GetRequiredService<ICommitService>();
@@ -212,7 +212,7 @@ namespace Horde.Build.Tests
 			services.AddSingleton<AclService>();
 			services.AddSingleton<AgentService>();
 			services.AddSingleton<AgentSoftwareService>();
-			services.AddSingleton<AutoscaleServiceV2>();
+			services.AddSingleton<FleetService>();
 			services.AddSingleton<ConsistencyService>();
 			services.AddSingleton<ConfigCollection>();
 			services.AddSingleton<ComputeService>();
diff --git a/Engine/Source/Programs/Horde/Horde.Build/Agents/Fleet/AutoscaleService.cs b/Engine/Source/Programs/Horde/Horde.Build/Agents/Fleet/AutoscaleService.cs
deleted file mode 100644
diff --git a/Engine/Source/Programs/Horde/Horde.Build/Agents/Fleet/AutoscaleServiceV2.cs b/Engine/Source/Programs/Horde/Horde.Build/Agents/Fleet/AutoscaleServiceV2.cs
deleted file mode 100644
diff --git a/Engine/Source/Programs/Horde/Horde.Build/Agents/Fleet/FleetService.cs b/Engine/Source/Programs/Horde/Horde.Build/Agents/Fleet/FleetService.cs
new file mode 100644
--- /dev/null
+++ b/Engine/Source/Programs/Horde/Horde.Build/Agents/Fleet/FleetService.cs
@@ -0,0 +1,340 @@
+// Copyright Epic Games, Inc. All Rights Reserved.
+
+using System;
+using System.Collections.Concurrent;
+using System.Collections.Generic;
+using System.Diagnostics;
+using System.Linq;
+using System.Text.Json;
+using System.Threading;
+using System.Threading.Tasks;
+using Horde.Build.Agents.Leases;
+using Horde.Build.Agents.Pools;
+using Horde.Build.Jobs;
+using Horde.Build.Jobs.Graphs;
+using Horde.Build.Streams;
+using Horde.Build.Utilities;
+using HordeCommon;
+using Microsoft.Extensions.Caching.Memory;
+using Microsoft.Extensions.Hosting;
+using Microsoft.Extensions.Logging;
+using Microsoft.Extensions.Options;
+using OpenTracing;
+using OpenTracing.Util;
+using StatsdClient;
+
+namespace Horde.Build.Agents.Fleet
+{
+	using PoolId = StringId<IPool>;
+
+	/// <summary>
+	/// Service for managing the autoscaling of agent pools
+	/// </summary>
+	public sealed class FleetService : IHostedService, IDisposable
+	{
+		/// <summary>
+		/// Max number of auto-scaling calculations to be done concurrently (sizing calculations and fleet manager calls)
+		/// </summary>
+		private const int MaxParallelTasks = 10;
+		
+		private readonly IAgentCollection _agentCollection;
+		private readonly IGraphCollection _graphCollection;
+		private readonly IJobCollection _jobCollection;
+		private readonly ILeaseCollection _leaseCollection;
+		private readonly IPoolCollection _poolCollection;
+		private readonly StreamService _streamService;
+		private readonly IDogStatsd _dogStatsd;
+		private readonly IFleetManagerFactory _fleetManagerFactory2;
+		private readonly IClock _clock;
+		private readonly IMemoryCache _cache;
+		private readonly ITicker _ticker;
+		private readonly ITicker _tickerHighFrequency;
+		private readonly TimeSpan _defaultScaleOutCooldown;
+		private readonly TimeSpan _defaultScaleInCooldown;
+		private readonly IOptions<ServerSettings> _settings;
+		private readonly ILogger<FleetService> _logger;
+		
+		/// <summary>
+		/// Constructor
+		/// </summary>
+		public FleetService(
+			IAgentCollection agentCollection,
+			IGraphCollection graphCollection,
+			IJobCollection jobCollection,
+			ILeaseCollection leaseCollection,
+			IPoolCollection poolCollection,
+			StreamService streamService,
+			IDogStatsd dogStatsd,
+			IFleetManagerFactory fleetManagerFactory,
+			IClock clock,
+			IMemoryCache cache,
+			IOptions<ServerSettings> settings,
+			ILogger<FleetService> logger)
+		{
+			_agentCollection = agentCollection;
+			_graphCollection = graphCollection;
+			_jobCollection = jobCollection;
+			_leaseCollection = leaseCollection;
+			_poolCollection = poolCollection;
+			_streamService = streamService;
+			_dogStatsd = dogStatsd;
+			_fleetManagerFactory2 = fleetManagerFactory;
+			_clock = clock;
+			_cache = cache;
+			_logger = logger;
+			_ticker = clock.AddSharedTicker<FleetService>(TimeSpan.FromMinutes(5.0), TickLeaderAsync, _logger);
+			_tickerHighFrequency = clock.AddSharedTicker("FleetService.TickHighFrequency", TimeSpan.FromSeconds(30), TickHighFrequencyAsync, _logger);
+			_settings = settings;
+			_defaultScaleOutCooldown = TimeSpan.FromSeconds(settings.Value.AgentPoolScaleOutCooldownSeconds);
+			_defaultScaleInCooldown = TimeSpan.FromSeconds(settings.Value.AgentPoolScaleInCooldownSeconds);
+
+		}
+
+		/// <inheritdoc/>
+		public Task StartAsync(CancellationToken cancellationToken)
+		{
+			return Task.WhenAll(_ticker.StartAsync(), _tickerHighFrequency.StartAsync());
+		}
+
+		/// <inheritdoc/>
+		public Task StopAsync(CancellationToken cancellationToken)
+		{
+			return Task.WhenAll(_ticker.StopAsync(), _tickerHighFrequency.StopAsync());
+		}
+
+		/// <inheritdoc/>
+		public void Dispose()
+		{
+			_ticker.Dispose();
+			_tickerHighFrequency.Dispose();
+		}
+
+		internal async ValueTask TickLeaderAsync(CancellationToken stoppingToken)
+		{
+			_logger.LogInformation("Autoscaling pools...");
+			Stopwatch stopwatch = Stopwatch.StartNew();
+			using IScope _ = GlobalTracer.Instance.BuildSpan("AutoscaleService.TickAsync").StartActive();
+			
+			List<PoolSizeData> input = await GetPoolSizeDataAsync();
+			List<PoolSizeData> output = await CalculatePoolSizesAsync(input, stoppingToken);
+			await ResizePoolsAsync(output, stoppingToken);
+
+			stopwatch.Stop();
+			_logger.LogInformation("Autoscaling pools took {ElapsedTime} ms", stopwatch.ElapsedMilliseconds);
+		}
+
+		internal async ValueTask TickHighFrequencyAsync(CancellationToken stoppingToken)
+		{
+			_logger.LogInformation("Autoscaling pools (high frequency)...");
+			Stopwatch stopwatch = Stopwatch.StartNew();
+			using IScope _ = GlobalTracer.Instance.BuildSpan("AutoscaleService.TickHighFrequency").StartActive();
+			
+			// TODO: Re-enable high frequency scaling (only used for experimental scaling of remote execution agents at the moment)
+			await Task.Delay(0, stoppingToken);
+
+			stopwatch.Stop();
+			_logger.LogInformation("Autoscaling pools (high frequency) took {ElapsedTime} ms", stopwatch.ElapsedMilliseconds);
+		}
+
+		internal async Task<List<PoolSizeData>> CalculatePoolSizesAsync(List<PoolSizeData> poolSizeDatas, CancellationToken cancellationToken)
+		{
+			using IScope _ = GlobalTracer.Instance.BuildSpan("AutoscaleService.CalculatePoolSizesAsync").StartActive();
+			ConcurrentQueue<PoolSizeData> results = new ();
+			ParallelOptions options = new () { MaxDegreeOfParallelism = MaxParallelTasks, CancellationToken = cancellationToken };
+
+			await Parallel.ForEachAsync(poolSizeDatas, options, async (input, innerCt) =>
+			{
+				try
+				{
+					IPoolSizeStrategy sizeStrategy = CreatePoolSizeStrategy(input.Pool);
+					PoolSizeData output = (await sizeStrategy.CalcDesiredPoolSizesAsync(new List<PoolSizeData> { input }))[0];
+					results.Enqueue(output);
+				}
+				catch (Exception e)
+				{
+					_logger.LogError(e, "Failed calculating pool size for pool {PoolId}", input.Pool?.Id);
+				}
+			});
+
+			return results.ToList();
+		}
+
+		internal async Task<List<PoolSizeData>> GetPoolSizeDataAsync()
+		{
+			List<IAgent> agents = await _agentCollection.FindAsync(status: AgentStatus.Ok, enabled: true);
+			List<IAgent> GetAgentsInPool(PoolId poolId) => agents.FindAll(a => a.GetPools().Any(p => p == poolId));
+			List<IPool> pools = await _poolCollection.GetAsync();
+
+			return pools.Select(pool => new PoolSizeData(pool, GetAgentsInPool(pool.Id), null)).ToList();
+		}
+
+		internal async Task ResizePoolAsync(PoolSizeData poolSizeData, CancellationToken cancellationToken)
+		{
+			IPool pool = poolSizeData.Pool;
+
+			if (!pool.EnableAutoscaling || poolSizeData.DesiredAgentCount == null)
+			{
+				return;
+			}
+
+			int currentAgentCount = poolSizeData.Agents.Count;
+			int desiredAgentCount = poolSizeData.DesiredAgentCount.Value;
+			int deltaAgentCount = desiredAgentCount - currentAgentCount;
+
+			IFleetManager fleetManager = CreateFleetManager(poolSizeData.Pool);
+
+			_logger.LogInformation("{PoolName,-48} Current={Current,4} Target={Target,4} Delta={Delta,4} Status={Status} FleetManager={FleetManager}",
+				pool.Name, currentAgentCount, desiredAgentCount, deltaAgentCount, poolSizeData.StatusMessage, fleetManager.GetType().Name);
+			
+			try
+			{
+				using IScope scope = GlobalTracer.Instance.BuildSpan("ScalingPool").StartActive();
+				scope.Span.SetTag("poolName", pool.Name);
+				scope.Span.SetTag("current", currentAgentCount);
+				scope.Span.SetTag("desired", desiredAgentCount);
+				scope.Span.SetTag("delta", deltaAgentCount);
+
+				if (deltaAgentCount > 0)
+				{
+					TimeSpan scaleOutCooldown = pool.ScaleOutCooldown ?? _defaultScaleOutCooldown;
+					bool isCoolingDown = pool.LastScaleUpTime != null && pool.LastScaleUpTime + scaleOutCooldown > _clock.UtcNow;
+					scope.Span.SetTag("isCoolingDown", isCoolingDown);
+					if (!isCoolingDown)
+					{
+						await fleetManager.ExpandPoolAsync(pool, poolSizeData.Agents, deltaAgentCount, cancellationToken);
+						await _poolCollection.TryUpdateAsync(pool, lastScaleUpTime: _clock.UtcNow);
+					}
+					else
+					{
+						TimeSpan? cooldownTimeLeft = pool.LastScaleUpTime + _defaultScaleOutCooldown - _clock.UtcNow;
+						_logger.LogDebug("Cannot scale out {PoolName}, it's cooling down for another {TimeLeft} secs", pool.Name, cooldownTimeLeft?.TotalSeconds);
+					}
+				}
+
+				if (deltaAgentCount < 0)
+				{
+					TimeSpan scaleInCooldown = pool.ScaleInCooldown ?? _defaultScaleInCooldown;
+					bool isCoolingDown = pool.LastScaleDownTime != null && pool.LastScaleDownTime + scaleInCooldown > _clock.UtcNow;
+					scope.Span.SetTag("isCoolingDown", isCoolingDown);
+					if (!isCoolingDown)
+					{
+						await fleetManager.ShrinkPoolAsync(pool, poolSizeData.Agents, -deltaAgentCount, cancellationToken);
+						await _poolCollection.TryUpdateAsync(pool, lastScaleDownTime: _clock.UtcNow);
+					}
+					else
+					{
+						TimeSpan? cooldownTimeLeft = pool.LastScaleDownTime + _defaultScaleInCooldown - _clock.UtcNow;
+						_logger.LogDebug("Cannot scale in {PoolName}, it's cooling down for another {TimeLeft} secs", pool.Name, cooldownTimeLeft?.TotalSeconds);
+					}
+				}
+			}
+			catch (Exception ex)
+			{
+				_logger.LogError(ex, "Failed to scale {PoolName}:\n{Exception}", pool.Name, ex);
+				return;
+			}
+
+			_dogStatsd.Gauge("agentpools.autoscale.target", desiredAgentCount, tags: new []{"pool:" + pool.Name});
+			_dogStatsd.Gauge("agentpools.autoscale.current", currentAgentCount, tags: new []{"pool:" + pool.Name});
+		}
+		
+		internal async Task ResizePoolsAsync(List<PoolSizeData> poolSizeDatas, CancellationToken cancellationToken = default)
+		{
+			using IScope _ = GlobalTracer.Instance.BuildSpan("AutoscaleService.ResizePoolsAsync").StartActive();
+			
+			ParallelOptions options = new () { MaxDegreeOfParallelism = MaxParallelTasks, CancellationToken = cancellationToken };
+			await Parallel.ForEachAsync(poolSizeDatas.OrderByDescending(x => x.Agents.Count), options, async (poolSizeData, innerCt) =>
+			{
+				await ResizePoolAsync(poolSizeData, innerCt);
+			});
+		}
+		
+		private IEnumerable<string> GetPropValues(string name)
+		{
+			return name switch
+			{
+				"dayOfWeek" => new List<string> { _clock.UtcNow.DayOfWeek.ToString().ToLower() },
+				_ => Array.Empty<string>()
+			};
+		}
+		
+		/// <summary>
+		/// Instantiate a fleet manager using the list of conditions/configs in <see cref="IPool" />
+		/// </summary>
+		/// <param name="pool">Pool to use</param>
+		/// <returns>A fleet manager with parameters set, as dictated by the pool argument</returns>
+		/// <exception cref="ArgumentException">If fleet manager could not be instantiated</exception>
+		public IFleetManager CreateFleetManager(IPool pool)
+		{
+			foreach (FleetManagerInfo info in pool.FleetManagers)
+			{
+				if (info.Condition == null || info.Condition.Evaluate(GetPropValues))
+				{
+					return _fleetManagerFactory2.CreateFleetManager(info.Type, info.Config);
+				}
+			}
+
+			return _fleetManagerFactory2.CreateFleetManager(FleetManagerType.Default, "{}");
+		}
+
+		/// <summary>
+		/// Instantiate a sizing strategy for the given pool
+		/// Can resolve either from the list of strategies or the old legacy way
+		/// </summary>
+		/// <param name="pool">Pool to use</param>
+		/// <returns>A pool sizing strategy with parameters set</returns>
+		/// <exception cref="ArgumentException"></exception>
+		public IPoolSizeStrategy CreatePoolSizeStrategy(IPool pool)
+		{
+			if (pool.SizeStrategies.Count > 0)
+			{
+				foreach (PoolSizeStrategyInfo info in pool.SizeStrategies)
+				{
+					if (info.Condition == null || info.Condition.Evaluate(GetPropValues))
+					{
+						switch (info.Type)
+						{
+							case PoolSizeStrategy.JobQueue:
+								JobQueueSettings jqSettings = DeserializeConfig<JobQueueSettings>(info.Config);
+								return new JobQueueStrategy(_jobCollection, _graphCollection, _streamService, _clock, _cache, jqSettings);
+							
+							case PoolSizeStrategy.LeaseUtilization:
+								LeaseUtilizationSettings luSettings = DeserializeConfig<LeaseUtilizationSettings>(info.Config);
+								return new LeaseUtilizationStrategy(_agentCollection, _poolCollection, _leaseCollection, _clock, _cache, luSettings);
+							
+							case PoolSizeStrategy.NoOp:
+								return new NoOpPoolSizeStrategy();
+							
+							default:
+								throw new ArgumentException("Invalid pool size strategy type " + info.Type);
+						}
+					}
+				}
+			}
+
+			// These is the legacy way of creating and configuring strategies (list-based approach above is preferred)
+			switch (pool.SizeStrategy)
+			{
+				case PoolSizeStrategy.JobQueue:
+					return new JobQueueStrategy(_jobCollection, _graphCollection, _streamService, _clock, _cache, pool.JobQueueSettings);
+				case PoolSizeStrategy.LeaseUtilization:
+					LeaseUtilizationSettings luSettings = new();
+					if (pool.MinAgents != null) luSettings.MinAgents = pool.MinAgents.Value;
+					if (pool.NumReserveAgents != null) luSettings.NumReserveAgents = pool.NumReserveAgents.Value;
+					return new LeaseUtilizationStrategy(_agentCollection, _poolCollection, _leaseCollection, _clock, _cache, luSettings);
+				case PoolSizeStrategy.NoOp:
+					return new NoOpPoolSizeStrategy();
+				default:
+					throw new ArgumentException("Unknown pool size strategy " + pool.SizeStrategy);
+			}
+		}
+
+		private static T DeserializeConfig<T>(string json)
+		{
+			json = String.IsNullOrEmpty(json) ? "{}" : json;
+			T? config = JsonSerializer.Deserialize<T>(json);
+			if (config == null) throw new ArgumentException("Unable to deserialize config: " + json);
+			return config;
+		}
+	}
+}
diff --git a/Engine/Source/Programs/Horde/Horde.Build/Startup.cs b/Engine/Source/Programs/Horde/Horde.Build/Startup.cs
--- a/Engine/Source/Programs/Horde/Horde.Build/Startup.cs
+++ b/Engine/Source/Programs/Horde/Horde.Build/Startup.cs
@@ -369,7 +369,7 @@ namespace Horde.Build
 			services.AddSingleton(typeof(IAuditLogFactory<>), typeof(AuditLogFactory<>));
 			services.AddSingleton(typeof(ISingletonDocument<>), typeof(SingletonDocument<>));
 
-			services.AddSingleton<AutoscaleServiceV2>();
+			services.AddSingleton<FleetService>();
 			services.AddSingleton<IFleetManagerFactory, FleetManagerFactory>();
 			
 			// Associate IFleetManager interface with the default implementation from config for convenience
@@ -599,7 +599,7 @@ namespace Horde.Build
 			{
 				services.AddHostedService<MongoUpgradeService>();
 
-				services.AddHostedService(provider => provider.GetRequiredService<AutoscaleServiceV2>());
+				services.AddHostedService(provider => provider.GetRequiredService<FleetService>());
 				
 				services.AddHostedService(provider => provider.GetRequiredService<AgentService>());
 				services.AddHostedService(provider => provider.GetRequiredService<PerforceServiceCache>());
