Horde: Rename test to match new name of FleetService

diff --git a/Engine/Source/Programs/Horde/Horde.Build.Tests/Fleet/AutoscaleServiceV2Test.cs b/Engine/Source/Programs/Horde/Horde.Build.Tests/Fleet/AutoscaleServiceV2Test.cs
deleted file mode 100644
diff --git a/Engine/Source/Programs/Horde/Horde.Build.Tests/Fleet/FleetServiceTest.cs b/Engine/Source/Programs/Horde/Horde.Build.Tests/Fleet/FleetServiceTest.cs
new file mode 100644
--- /dev/null
+++ b/Engine/Source/Programs/Horde/Horde.Build.Tests/Fleet/FleetServiceTest.cs
@@ -0,0 +1,369 @@
+// Copyright Epic Games, Inc. All Rights Reserved.
+
+extern alias HordeAgent;
+using System;
+using System.Collections.Generic;
+using System.Linq;
+using System.Text.Json;
+using System.Threading;
+using System.Threading.Tasks;
+using Amazon.AutoScaling;
+using Amazon.EC2;
+using Horde.Build.Agents;
+using Horde.Build.Agents.Fleet;
+using Horde.Build.Agents.Fleet.Providers;
+using Horde.Build.Agents.Pools;
+using Horde.Build.Utilities;
+using Microsoft.Extensions.DependencyInjection;
+using Microsoft.Extensions.Logging;
+using Microsoft.Extensions.Options;
+using Microsoft.VisualStudio.TestTools.UnitTesting;
+using Moq;
+using StatsdClient;
+
+namespace Horde.Build.Tests.Fleet
+{
+	using PoolId = StringId<IPool>;
+
+	public class FleetManagerSpy : IFleetManager
+	{
+		public int ExpandPoolAsyncCallCount { get; private set; }
+		public int ShrinkPoolAsyncCallCount { get; private set; }
+		
+		public Task ExpandPoolAsync(IPool pool, IReadOnlyList<IAgent> agents, int count, CancellationToken cancellationToken)
+		{
+			ExpandPoolAsyncCallCount++;
+			return Task.CompletedTask;
+		}
+
+		public Task ShrinkPoolAsync(IPool pool, IReadOnlyList<IAgent> agents, int count, CancellationToken cancellationToken)
+		{
+			ShrinkPoolAsyncCallCount++;
+			return Task.CompletedTask;
+		}
+
+		public Task<int> GetNumStoppedInstancesAsync(IPool pool, CancellationToken cancellationToken)
+		{
+			throw new NotImplementedException();
+		}
+	}
+
+	public class PoolSizeStrategySpy : IPoolSizeStrategy
+	{
+		public int CallCount { get; private set; }
+		public HashSet<PoolId> PoolIdsSeen { get; } = new();
+		
+		public Task<List<PoolSizeData>> CalcDesiredPoolSizesAsync(List<PoolSizeData> pools)
+		{
+			CallCount++;
+			foreach (PoolSizeData data in pools)
+			{
+				PoolIdsSeen.Add(data.Pool.Id);
+			}
+			return Task.FromResult(pools);
+		}
+
+		public string Name { get; } = "PoolSizeStrategySpy";
+	}
+	
+	public class StubFleetManagerFactory : IFleetManagerFactory
+	{
+		private readonly IFleetManager _fleetManager;
+
+		public StubFleetManagerFactory(IFleetManager fleetManager)
+		{
+			_fleetManager = fleetManager;
+		}
+
+		public IFleetManager CreateFleetManager(FleetManagerType type, string config = "{}")
+		{
+			return _fleetManager;
+		}
+	}
+
+	[TestClass]
+	public class FleetServiceTest : TestSetup
+	{
+		readonly FleetManagerSpy _fleetManagerSpy = new();
+		readonly IDogStatsd _dogStatsD = new NoOpDogStatsd();
+
+		protected override void Dispose(bool disposing)
+		{
+			base.Dispose(disposing);
+
+			if (disposing)
+			{
+				_dogStatsD.Dispose();
+			}
+		}
+
+		[TestMethod]
+		public async Task OnlyEnabledAgentsAreAutoScaled()
+		{
+			using FleetService service = GetAutoscaleService(_fleetManagerSpy);
+			IPool pool = await PoolService.CreatePoolAsync("testPool", null, true, 0, 0, sizeStrategy: PoolSizeStrategy.LeaseUtilization);
+			await CreateAgentAsync(pool, true);
+			await CreateAgentAsync(pool, true);
+			await CreateAgentAsync(pool, false);
+
+			List<PoolSizeData> poolSizeDatas = await service.GetPoolSizeDataAsync();
+			poolSizeDatas = await service.CalculatePoolSizesAsync(poolSizeDatas, CancellationToken.None);
+			Assert.AreEqual(1, poolSizeDatas.Count);
+			Assert.AreEqual(2, poolSizeDatas[0].Agents.Count);
+		}
+		
+		[TestMethod]
+		public async Task ExceptionsAreSwallowedAndLogged()
+		{
+			using FleetService service = GetAutoscaleService(_fleetManagerSpy);
+			IPool pool = await PoolService.CreatePoolAsync("testPool", null, true, 0, 0, sizeStrategy: PoolSizeStrategy.LeaseUtilization);
+			List<PoolSizeData> poolSizeDatas = new()
+			{
+				new PoolSizeData(null!, new(), null),
+				new PoolSizeData(pool, new(), null),
+			};
+
+			poolSizeDatas = await service.CalculatePoolSizesAsync(poolSizeDatas, CancellationToken.None);
+			Assert.AreEqual(1, poolSizeDatas.Count);
+		}
+
+		[TestMethod]
+		public async Task ScaleOutCooldown()
+		{
+			using FleetService service = GetAutoscaleService(_fleetManagerSpy);
+			IPool pool = await PoolService.CreatePoolAsync("testPool", null, true, 0, 0, sizeStrategy: PoolSizeStrategy.NoOp);
+
+			// First scale-out will succeed
+			await service.ResizePoolsAsync(new() { new PoolSizeData(pool, new List<IAgent>(), 1, "Testing") });
+			Assert.AreEqual(1, _fleetManagerSpy.ExpandPoolAsyncCallCount);
+			
+			// Cannot scale-out due to cool-down
+			await service.ResizePoolsAsync(new() { new PoolSizeData(pool, new List<IAgent>(), 2, "Testing") });
+			Assert.AreEqual(1, _fleetManagerSpy.ExpandPoolAsyncCallCount);
+
+			// Wait some time and then try again
+			await Clock.AdvanceAsync(TimeSpan.FromHours(2));
+			await service.ResizePoolsAsync(new() { new PoolSizeData(pool, new List<IAgent>(), 2, "Testing") });
+			Assert.AreEqual(2, _fleetManagerSpy.ExpandPoolAsyncCallCount);
+		}
+		
+		[TestMethod]
+		public async Task ScaleInCooldown()
+		{
+			using FleetService service = GetAutoscaleService(_fleetManagerSpy);
+			IPool pool = await PoolService.CreatePoolAsync("testPool", null, true, 0, 0, sizeStrategy: PoolSizeStrategy.NoOp);
+			IAgent agent1 = await CreateAgentAsync(pool);
+			IAgent agent2 = await CreateAgentAsync(pool);
+
+			// First scale-out will succeed
+			await service.ResizePoolsAsync(new() { new PoolSizeData(pool, new () { agent1, agent2 }, 1, "Testing") });
+			Assert.AreEqual(1, _fleetManagerSpy.ShrinkPoolAsyncCallCount);
+			
+			// Cannot scale-out due to cool-down
+			await service.ResizePoolsAsync(new() { new PoolSizeData(pool, new () { agent1 }, 0, "Testing") });
+			Assert.AreEqual(1, _fleetManagerSpy.ShrinkPoolAsyncCallCount);
+
+			// Wait some time and then try again
+			await Clock.AdvanceAsync(TimeSpan.FromHours(2));
+			await service.ResizePoolsAsync(new() { new PoolSizeData(pool, new () { agent1 }, 0, "Testing") });
+			Assert.AreEqual(2, _fleetManagerSpy.ShrinkPoolAsyncCallCount);
+		}
+
+		private FleetService GetAutoscaleService(IFleetManager fleetManager)
+		{
+			ILoggerFactory loggerFactory = ServiceProvider.GetRequiredService<ILoggerFactory>();
+			IOptions<ServerSettings> serverSettingsOpt = ServiceProvider.GetRequiredService<IOptions<ServerSettings>>();
+			serverSettingsOpt.Value.FleetManagerV2 = FleetManagerType.AwsReuse;
+			
+			FleetService service = new(
+				AgentCollection, GraphCollection, JobCollection, LeaseCollection, PoolCollection, StreamService, _dogStatsD,
+				new StubFleetManagerFactory(fleetManager), Clock, Cache, serverSettingsOpt, loggerFactory.CreateLogger<FleetService>());
+				
+			return service;
+		}
+	}
+
+	[TestClass]
+	public class PoolSizeStrategyFactoryTest : TestSetup
+	{
+		private static int s_poolCount;
+
+		[TestMethod]
+		public async Task CreateJobQueueFromLegacySettings()
+		{
+			IPool pool1 = await PoolService.CreatePoolAsync("test1", sizeStrategy: PoolSizeStrategy.JobQueue);
+			Assert.AreEqual(typeof(JobQueueStrategy), FleetService.CreatePoolSizeStrategy(pool1).GetType());
+			
+			IPool pool2 = await PoolService.CreatePoolAsync("test2", sizeStrategy: PoolSizeStrategy.JobQueue, jobQueueSettings: new JobQueueSettings(22, 33));
+			IPoolSizeStrategy s = FleetService.CreatePoolSizeStrategy(pool2);
+			Assert.AreEqual(typeof(JobQueueStrategy), s.GetType());
+			Assert.AreEqual(22.0, ((JobQueueStrategy)s).Settings.ScaleOutFactor);
+			Assert.AreEqual(33.0, ((JobQueueStrategy)s).Settings.ScaleInFactor);
+		}
+		
+		[TestMethod]
+		public async Task CreateLeaseUtilizationFromLegacySettings()
+		{
+			IPool pool = await PoolService.CreatePoolAsync("test1", sizeStrategy: PoolSizeStrategy.LeaseUtilization);
+			Assert.AreEqual(typeof(LeaseUtilizationStrategy), FleetService.CreatePoolSizeStrategy(pool).GetType());
+		}
+
+		[TestMethod]
+		public async Task CreateJobQueueStrategy()
+		{
+			IPoolSizeStrategy s = await CreateStrategy(new PoolSizeStrategyInfo(PoolSizeStrategy.JobQueue, null, "{\"ScaleOutFactor\": 100, \"ScaleInFactor\": 200}"));
+			Assert.AreEqual(typeof(JobQueueStrategy), s.GetType());
+			Assert.AreEqual(100.0, ((JobQueueStrategy)s).Settings.ScaleOutFactor);
+			Assert.AreEqual(200.0, ((JobQueueStrategy)s).Settings.ScaleInFactor);
+		}
+		
+		[TestMethod]
+		public async Task CreateLeaseUtilizationStrategy()
+		{
+			string config = "{\"SampleTimeSec\": 10, \"NumSamples\": 20, \"NumSamplesForResult\": 30}";
+			IPoolSizeStrategy s = await CreateStrategy(new PoolSizeStrategyInfo(PoolSizeStrategy.LeaseUtilization, null, config));
+			Assert.AreEqual(typeof(LeaseUtilizationStrategy), s.GetType());
+			Assert.AreEqual(10, ((LeaseUtilizationStrategy)s).Settings.SampleTimeSec);
+			Assert.AreEqual(20, ((LeaseUtilizationStrategy)s).Settings.NumSamples);
+			Assert.AreEqual(30, ((LeaseUtilizationStrategy)s).Settings.NumSamplesForResult);
+		}
+		
+		[TestMethod]
+		public async Task EmptyOrInvalidJsonConfig()
+		{
+			await CreateStrategy(new PoolSizeStrategyInfo(PoolSizeStrategy.LeaseUtilization, null, ""));
+			await CreateStrategy(new PoolSizeStrategyInfo(PoolSizeStrategy.LeaseUtilization, null, "{}"));
+			await CreateStrategy(new PoolSizeStrategyInfo(PoolSizeStrategy.LeaseUtilization, null, "  {} "));
+			
+			await Assert.ThrowsExceptionAsync<JsonException>(() => CreateStrategy(new PoolSizeStrategyInfo(PoolSizeStrategy.LeaseUtilization, null, "BAD_JSON")));
+		}
+		
+		[TestMethod]
+		public async Task CreateNoOpStrategy()
+		{
+			IPoolSizeStrategy s = await CreateStrategy(new PoolSizeStrategyInfo(PoolSizeStrategy.NoOp, null, "{}"));
+			Assert.AreEqual(typeof(NoOpPoolSizeStrategy), s.GetType());
+		}
+		
+		[TestMethod]
+		public async Task UnknownConfigFieldsAreIgnored()
+		{
+			IPoolSizeStrategy s = await CreateStrategy(new PoolSizeStrategyInfo(PoolSizeStrategy.JobQueue, null, "{\"BAD-PROPERTY\": 1337}"));
+			Assert.AreEqual(typeof(JobQueueStrategy), s.GetType());
+			Assert.AreEqual(0.25, ((JobQueueStrategy)s).Settings.ScaleOutFactor);
+			Assert.AreEqual(0.9, ((JobQueueStrategy)s).Settings.ScaleInFactor);
+		}
+		
+		[TestMethod]
+		public async Task CreateFromEmptyStrategyList()
+		{
+			IPoolSizeStrategy s = await CreateStrategy();
+			Assert.AreEqual(typeof(NoOpPoolSizeStrategy), s.GetType());
+		}
+		
+		[TestMethod]
+		public async Task ConditionSimple()
+		{
+			IPoolSizeStrategy s = await CreateStrategy(
+				new PoolSizeStrategyInfo(PoolSizeStrategy.LeaseUtilization, "false", "{}"),
+				new PoolSizeStrategyInfo(PoolSizeStrategy.JobQueue, "true", "{\"ScaleOutFactor\": 100, \"ScaleInFactor\": 200}")
+			);
+			Assert.AreEqual(typeof(JobQueueStrategy), s.GetType());
+		}
+		
+		[TestMethod]
+		public async Task ConditionDayOfWeek()
+		{
+			Clock.UtcNow = new DateTime(2022, 9, 5, 15, 0, 0, DateTimeKind.Utc); // A Monday
+			IPoolSizeStrategy s1 = await CreateStrategy(new PoolSizeStrategyInfo(PoolSizeStrategy.LeaseUtilization, "dayOfWeek == 'monday'", "{}"));
+			
+			Clock.UtcNow = new DateTime(2022, 9, 6, 15, 0, 0, DateTimeKind.Utc); // A Tuesday
+			IPoolSizeStrategy s2 = await CreateStrategy(new PoolSizeStrategyInfo(PoolSizeStrategy.LeaseUtilization, "dayOfWeek == 'monday'", "{}"));
+			
+			Assert.AreEqual(typeof(LeaseUtilizationStrategy), s1.GetType());
+			Assert.AreEqual(typeof(NoOpPoolSizeStrategy), s2.GetType());
+		}
+		
+		private async Task<IPoolSizeStrategy> CreateStrategy(params PoolSizeStrategyInfo[] infos)
+		{
+			IPool pool = await PoolService.CreatePoolAsync("testPool-" + s_poolCount++, null, true, 0, 0, sizeStrategy: PoolSizeStrategy.NoOp);
+			await PoolService.UpdatePoolAsync(pool, newSizeStrategies: infos.ToList());
+			return FleetService.CreatePoolSizeStrategy(pool);
+		}
+	}
+
+	[TestClass]
+	public class FleetManagerFactoryTest : TestSetup
+	{
+		private static int s_poolCount;
+		
+		[TestMethod]
+		public async Task CreateNoOp()
+		{
+			IFleetManager fm = await CreateFleetManager(new FleetManagerInfo(FleetManagerType.NoOp, null, "{}"));
+			Assert.AreEqual(typeof(NoOpFleetManager), fm.GetType());
+		}
+		
+		[TestMethod]
+		public async Task CreateAws()
+		{
+			IFleetManager fm = await CreateFleetManager(new FleetManagerInfo(FleetManagerType.Aws, null, "{\"ImageId\": \"bogusImageId\"}"));
+			Assert.AreEqual(typeof(AwsFleetManager), fm.GetType());
+			Assert.AreEqual("bogusImageId", ((AwsFleetManager)fm).Settings.ImageId);
+		}
+		
+		[TestMethod]
+		public async Task CreateAwsReuse()
+		{
+			IFleetManager fm = await CreateFleetManager(new FleetManagerInfo(FleetManagerType.AwsReuse, null, "{\"InstanceTypes\": [\"foo\", \"bar\"]}"));
+			Assert.AreEqual(typeof(AwsReuseFleetManager), fm.GetType());
+			Assert.AreEqual("foo", ((AwsReuseFleetManager)fm).Settings.InstanceTypes![0]);
+			Assert.AreEqual("bar", ((AwsReuseFleetManager)fm).Settings.InstanceTypes![1]);
+		}
+		
+		[TestMethod]
+		public async Task CreateAwsAsg()
+		{
+			IFleetManager fm = await CreateFleetManager(new FleetManagerInfo(FleetManagerType.AwsAsg, null, "{\"Name\": \"bogusName\"}"));
+			Assert.AreEqual(typeof(AwsAsgFleetManager), fm.GetType());
+			Assert.AreEqual("bogusName", ((AwsAsgFleetManager)fm).Settings.Name);
+		}
+
+		[TestMethod]
+		public async Task UnknownConfigFieldsAreIgnored()
+		{
+			IFleetManager fm = await CreateFleetManager(new FleetManagerInfo(FleetManagerType.AwsAsg, null, "{\"Name\": \"bogusName\", \"BAD-PROPERTY\": 1337}"));
+			Assert.AreEqual(typeof(AwsAsgFleetManager), fm.GetType());
+			Assert.AreEqual("bogusName", ((AwsAsgFleetManager)fm).Settings.Name);
+		}
+		
+		[TestMethod]
+		public async Task EmptyConfigThrowsException()
+		{
+			await Assert.ThrowsExceptionAsync<JsonException>(() => CreateFleetManager(new FleetManagerInfo(FleetManagerType.AwsAsg, null, "")));
+		}
+		
+		[TestMethod]
+		public async Task CreateFromEmptyList()
+		{
+			IFleetManager fm = await CreateFleetManager();
+			Assert.AreEqual(typeof(NoOpFleetManager), fm.GetType());
+		}
+		
+		[TestMethod]
+		public async Task ConditionSimple()
+		{
+			IFleetManager s = await CreateFleetManager(
+				new FleetManagerInfo(FleetManagerType.Aws, "false", "{}"),
+				new FleetManagerInfo(FleetManagerType.AwsReuse, "true", "{\"ScaleOutFactor\": 100, \"ScaleInFactor\": 200}")
+			);
+			Assert.AreEqual(typeof(AwsReuseFleetManager), s.GetType());
+		}
+		
+		private async Task<IFleetManager> CreateFleetManager(params FleetManagerInfo[] infos)
+		{
+			IPool pool = await PoolService.CreatePoolAsync("testPool-" + s_poolCount++, null, true, 0, 0, fleetManagers: infos.ToList(), sizeStrategy: PoolSizeStrategy.NoOp);
+			return FleetService.CreateFleetManager(pool);
+		}
+	}
+}
\ No newline at end of file
