Horde: Validate config for pool size strategy in REST API controller

diff --git a/Engine/Source/Programs/Horde/Horde.Build/Agents/Pools/IPool.cs b/Engine/Source/Programs/Horde/Horde.Build/Agents/Pools/IPool.cs
--- a/Engine/Source/Programs/Horde/Horde.Build/Agents/Pools/IPool.cs
+++ b/Engine/Source/Programs/Horde/Horde.Build/Agents/Pools/IPool.cs
@@ -43,8 +43,14 @@ namespace Horde.Build.Agents.Pools
 		/// <summary>
 		/// Constructor
 		/// </summary>
-		public PoolSizeStrategyInfo(PoolSizeStrategy type, Condition? condition, string config)
+		public PoolSizeStrategyInfo(PoolSizeStrategy type, Condition? condition, string? config)
 		{
+			config = String.IsNullOrWhiteSpace(config) ? "{}" : config;
+
+			// Try deserializing to ensure the config is valid JSON
+			// Config can be null due to JSON serializer calling the constructor
+			JsonSerializer.Deserialize<dynamic>(config);
+			
 			Type = type;
 			Condition = condition;
 			Config = config;
