Horde: Improve pool size strategy logging

diff --git a/Engine/Source/Programs/Horde/Horde.Build.Tests/Fleet/FleetServiceTest.cs b/Engine/Source/Programs/Horde/Horde.Build.Tests/Fleet/FleetServiceTest.cs
--- a/Engine/Source/Programs/Horde/Horde.Build.Tests/Fleet/FleetServiceTest.cs
+++ b/Engine/Source/Programs/Horde/Horde.Build.Tests/Fleet/FleetServiceTest.cs
@@ -134,16 +134,16 @@ namespace Horde.Build.Tests.Fleet
 			IPool pool = await PoolService.CreatePoolAsync("testPool", null, true, 0, 0, sizeStrategy: PoolSizeStrategy.NoOp);
 
 			// First scale-out will succeed
-			await service.ResizePoolsAsync(new() { new PoolSizeData(pool, new List<IAgent>(), 1, "Testing") });
+			await service.ResizePoolsAsync(new() { new PoolSizeData(pool, new List<IAgent>(), 1) });
 			Assert.AreEqual(1, _fleetManagerSpy.ExpandPoolAsyncCallCount);
 			
 			// Cannot scale-out due to cool-down
-			await service.ResizePoolsAsync(new() { new PoolSizeData(pool, new List<IAgent>(), 2, "Testing") });
+			await service.ResizePoolsAsync(new() { new PoolSizeData(pool, new List<IAgent>(), 2) });
 			Assert.AreEqual(1, _fleetManagerSpy.ExpandPoolAsyncCallCount);
 
 			// Wait some time and then try again
 			await Clock.AdvanceAsync(TimeSpan.FromHours(2));
-			await service.ResizePoolsAsync(new() { new PoolSizeData(pool, new List<IAgent>(), 2, "Testing") });
+			await service.ResizePoolsAsync(new() { new PoolSizeData(pool, new List<IAgent>(), 2) });
 			Assert.AreEqual(2, _fleetManagerSpy.ExpandPoolAsyncCallCount);
 		}
 		
@@ -156,16 +156,16 @@ namespace Horde.Build.Tests.Fleet
 			IAgent agent2 = await CreateAgentAsync(pool);
 
 			// First scale-out will succeed
-			await service.ResizePoolsAsync(new() { new PoolSizeData(pool, new () { agent1, agent2 }, 1, "Testing") });
+			await service.ResizePoolsAsync(new() { new PoolSizeData(pool, new () { agent1, agent2 }, 1) });
 			Assert.AreEqual(1, _fleetManagerSpy.ShrinkPoolAsyncCallCount);
 			
 			// Cannot scale-out due to cool-down
-			await service.ResizePoolsAsync(new() { new PoolSizeData(pool, new () { agent1 }, 0, "Testing") });
+			await service.ResizePoolsAsync(new() { new PoolSizeData(pool, new () { agent1 }, 0) });
 			Assert.AreEqual(1, _fleetManagerSpy.ShrinkPoolAsyncCallCount);
 
 			// Wait some time and then try again
 			await Clock.AdvanceAsync(TimeSpan.FromHours(2));
-			await service.ResizePoolsAsync(new() { new PoolSizeData(pool, new () { agent1 }, 0, "Testing") });
+			await service.ResizePoolsAsync(new() { new PoolSizeData(pool, new () { agent1 }, 0) });
 			Assert.AreEqual(2, _fleetManagerSpy.ShrinkPoolAsyncCallCount);
 		}
 
diff --git a/Engine/Source/Programs/Horde/Horde.Build/Agents/Fleet/FleetService.cs b/Engine/Source/Programs/Horde/Horde.Build/Agents/Fleet/FleetService.cs
--- a/Engine/Source/Programs/Horde/Horde.Build/Agents/Fleet/FleetService.cs
+++ b/Engine/Source/Programs/Horde/Horde.Build/Agents/Fleet/FleetService.cs
@@ -183,9 +183,15 @@ namespace Horde.Build.Agents.Fleet
 
 			IFleetManager fleetManager = CreateFleetManager(poolSizeData.Pool);
 
-			_logger.LogInformation("{PoolName,-48} Current={Current,4} Target={Target,4} Delta={Delta,4} Status={Status} FleetManager={FleetManager}",
-				pool.Name, currentAgentCount, desiredAgentCount, deltaAgentCount, poolSizeData.StatusMessage, fleetManager.GetType().Name);
-			
+			Dictionary<string, object> logScope = new() { ["FleetManager"] = fleetManager.GetType().Name };
+			logScope["PoolSizeStrategy"] = poolSizeData.Status ?? new Dictionary<string, object>();
+
+			using (_logger.BeginScope(logScope))
+			{
+				_logger.LogInformation("{PoolName} Current={Current} Target={Target} Delta={Delta}",
+					pool.Name, currentAgentCount, desiredAgentCount, deltaAgentCount);
+			}
+
 			try
 			{
 				using IScope scope = GlobalTracer.Instance.BuildSpan("ScalingPool").StartActive();
diff --git a/Engine/Source/Programs/Horde/Horde.Build/Agents/Fleet/IPoolSizeStrategy.cs b/Engine/Source/Programs/Horde/Horde.Build/Agents/Fleet/IPoolSizeStrategy.cs
--- a/Engine/Source/Programs/Horde/Horde.Build/Agents/Fleet/IPoolSizeStrategy.cs
+++ b/Engine/Source/Programs/Horde/Horde.Build/Agents/Fleet/IPoolSizeStrategy.cs
@@ -78,9 +78,9 @@ namespace Horde.Build.Agents.Fleet
 		public int? DesiredAgentCount  { get; }
 		
 		/// <summary>
-		/// Human-readable text describing the status of the pool (data the sizing is based on etc)
+		/// Log-friendly metadata object describing the output of the size calculation through key/values
 		/// </summary>
-		public string StatusMessage { get; }
+		public IReadOnlyDictionary<string, object>? Status { get; }
 
 		/// <summary>
 		/// Constructor
@@ -88,13 +88,13 @@ namespace Horde.Build.Agents.Fleet
 		/// <param name="pool"></param>
 		/// <param name="agents"></param>
 		/// <param name="desiredAgentCount"></param>
-		/// <param name="statusMessage"></param>
-		public PoolSizeData(IPool pool, List<IAgent> agents, int? desiredAgentCount, string statusMessage = "N/A")
+		/// <param name="status"></param>
+		public PoolSizeData(IPool pool, List<IAgent> agents, int? desiredAgentCount, IReadOnlyDictionary<string, object>? status = null)
 		{
 			Pool = pool;
 			Agents = agents;
 			DesiredAgentCount = desiredAgentCount;
-			StatusMessage = statusMessage;
+			Status = status;
 		}
 
 		/// <summary>
@@ -104,11 +104,11 @@ namespace Horde.Build.Agents.Fleet
 		/// <param name="pool"></param>
 		/// <param name="agents"></param>
 		/// <param name="desiredAgentCount"></param>
-		/// <param name="statusMessage"></param>
+		/// <param name="status"></param>
 		/// <returns>A new copy</returns>
-		public PoolSizeData Copy(IPool? pool = null, List<IAgent>? agents = null, int? desiredAgentCount = null, string? statusMessage = null)
+		public PoolSizeData Copy(IPool? pool = null, List<IAgent>? agents = null, int? desiredAgentCount = null, IReadOnlyDictionary<string, object>? status = null)
 		{
-			return new PoolSizeData(pool ?? Pool, agents ?? Agents, desiredAgentCount ?? DesiredAgentCount, statusMessage ?? StatusMessage);
+			return new PoolSizeData(pool ?? Pool, agents ?? Agents, desiredAgentCount ?? DesiredAgentCount, status);
 		}
 	}
 	
@@ -139,7 +139,7 @@ namespace Horde.Build.Agents.Fleet
 		/// <inheritdoc/>
 		public Task<List<PoolSizeData>> CalcDesiredPoolSizesAsync(List<PoolSizeData> pools)
 		{
-			List<PoolSizeData> result = pools.Select(x => new PoolSizeData(x.Pool, x.Agents, x.Agents.Count, "(no-op)")).ToList();
+			List<PoolSizeData> result = pools.Select(x => new PoolSizeData(x.Pool, x.Agents, x.Agents.Count)).ToList();
 			return Task.FromResult(result);
 		}
 
diff --git a/Engine/Source/Programs/Horde/Horde.Build/Agents/Fleet/JobQueueStrategy.cs b/Engine/Source/Programs/Horde/Horde.Build/Agents/Fleet/JobQueueStrategy.cs
--- a/Engine/Source/Programs/Horde/Horde.Build/Agents/Fleet/JobQueueStrategy.cs
+++ b/Engine/Source/Programs/Horde/Horde.Build/Agents/Fleet/JobQueueStrategy.cs
@@ -198,16 +198,27 @@ namespace Horde.Build.Agents.Fleet
 			return pools.Select(current =>
 			{
 				poolQueueSizes.TryGetValue(current.Pool.Id, out int queueSize);
+				
+				Dictionary<string, object> status = new()
+				{
+					["Name"] = GetType().Name,
+					["QueueSize"] = queueSize,
+					["ScaleOutFactor"] = Settings.ScaleOutFactor,
+					["ScaleInFactor"] = Settings.ScaleInFactor,
+					["SamplePeriodMin"] = Settings.SamplePeriodMin,
+					["ReadyTimeThresholdSec"] = Settings.ReadyTimeThresholdSec,
+				};
+				
 				if (queueSize > 0)
 				{
 					int additionalAgentCount = (int)Math.Ceiling(queueSize * Settings.ScaleOutFactor);
 					int desiredAgentCount = current.Agents.Count + additionalAgentCount;
-					return new PoolSizeData(current.Pool, current.Agents, desiredAgentCount, $"QueueSize={queueSize} " + Settings);
+					return new PoolSizeData(current.Pool, current.Agents, desiredAgentCount, status);
 				}
 				else
 				{
 					int desiredAgentCount = (int)(current.Agents.Count * Settings.ScaleInFactor);
-					return new PoolSizeData(current.Pool, current.Agents, desiredAgentCount, "QueueSize={queueSize} " + Settings);
+					return new PoolSizeData(current.Pool, current.Agents, desiredAgentCount, status);
 				}
 			}).ToList();
 		}
diff --git a/Engine/Source/Programs/Horde/Horde.Build/Agents/Fleet/LeaseUtilizationStrategy.cs b/Engine/Source/Programs/Horde/Horde.Build/Agents/Fleet/LeaseUtilizationStrategy.cs
--- a/Engine/Source/Programs/Horde/Horde.Build/Agents/Fleet/LeaseUtilizationStrategy.cs
+++ b/Engine/Source/Programs/Horde/Horde.Build/Agents/Fleet/LeaseUtilizationStrategy.cs
@@ -259,16 +259,23 @@ namespace Horde.Build.Agents.Fleet
 					// Include reserve agent count to ensure pool always can grow
 					int desiredAgentCount = Math.Max(numAgentsUtilized + Settings.NumReserveAgents, Settings.MinAgents);
 					
-					StringBuilder sb = new();
-					sb.AppendFormat("Jobs=[{0}] ", GetDensityMap(poolData.Samples.Select(x => x._jobWork / Math.Max(1, poolData.Agents.Count))));
-					sb.AppendFormat("Total=[{0}] ", GetDensityMap(poolData.Samples.Select(x => x._otherWork / Math.Max(1, poolData.Agents.Count))));
-					sb.AppendFormat("Min=[{0,5:0.0}] ", poolData.Samples.Min(x => x._jobWork));
-					sb.AppendFormat("Max=[{0,5:0.0}] ", poolData.Samples.Max(x => x._jobWork));
-					sb.AppendFormat("Avg=[{0,5:0.0}] ", utilization);
-					sb.AppendFormat("Pct=[{0,5:0.0}] ", poolData.Samples.Sum(x => x._jobWork) / Settings.NumSamples);
-					sb.Append(Settings);
+					Dictionary<string, object> status = new()
+					{
+						["Name"] = GetType().Name,
+						["Jobs"] = GetDensityMap(poolData.Samples.Select(x => x._jobWork / Math.Max(1, poolData.Agents.Count))),
+						["Total"] = GetDensityMap(poolData.Samples.Select(x => x._otherWork / Math.Max(1, poolData.Agents.Count))),
+						["Min"] = poolData.Samples.Min(x => x._jobWork),
+						["Max"] = poolData.Samples.Max(x => x._jobWork),
+						["Avg"] = utilization,
+						["Pct"] = poolData.Samples.Sum(x => x._jobWork) / Settings.NumSamples,
+						["SampleTimeSec"] = Settings.SampleTimeSec,
+						["NumSamples"] = Settings.NumSamples,
+						["NumSamplesForResult"] = Settings.NumSamplesForResult,
+						["MinAgents"] = Settings.MinAgents,
+						["NumReserveAgents"] = Settings.NumReserveAgents,
+					};
 
-					result.Add(new(pool, poolSize.Agents, desiredAgentCount, sb.ToString()));
+					result.Add(new(pool, poolSize.Agents, desiredAgentCount, status));
 				}
 			}
 
