Horde: Refactor pool size strategy interface back to accepting only one pool at a time

With the change to per-pool strategies with accompanying settings, it doesn't make sense to accept more than one any more.

diff --git a/Engine/Source/Programs/Horde/Horde.Build.Tests/Fleet/FleetServiceTest.cs b/Engine/Source/Programs/Horde/Horde.Build.Tests/Fleet/FleetServiceTest.cs
--- a/Engine/Source/Programs/Horde/Horde.Build.Tests/Fleet/FleetServiceTest.cs
+++ b/Engine/Source/Programs/Horde/Horde.Build.Tests/Fleet/FleetServiceTest.cs
@@ -53,14 +53,11 @@ namespace Horde.Build.Tests.Fleet
 		public int CallCount { get; private set; }
 		public HashSet<PoolId> PoolIdsSeen { get; } = new();
 		
-		public Task<List<PoolSizeData>> CalcDesiredPoolSizesAsync(List<PoolSizeData> pools)
+		public Task<PoolSizeResult> CalculatePoolSizeAsync(IPool pool, List<IAgent> agents)
 		{
 			CallCount++;
-			foreach (PoolSizeData data in pools)
-			{
-				PoolIdsSeen.Add(data.Pool.Id);
-			}
-			return Task.FromResult(pools);
+			PoolIdsSeen.Add(pool.Id);
+			return Task.FromResult(new PoolSizeResult(pool, agents, agents.Count));
 		}
 
 		public string Name { get; } = "PoolSizeStrategySpy";
@@ -106,7 +103,7 @@ namespace Horde.Build.Tests.Fleet
 			await CreateAgentAsync(pool, true);
 			await CreateAgentAsync(pool, false);
 
-			List<PoolSizeData> poolSizeDatas = await service.GetPoolSizeDataAsync();
+			List<PoolSizeResult> poolSizeDatas = await service.GetPoolSizeDataAsync();
 			poolSizeDatas = await service.CalculatePoolSizesAsync(poolSizeDatas, CancellationToken.None);
 			Assert.AreEqual(1, poolSizeDatas.Count);
 			Assert.AreEqual(2, poolSizeDatas[0].Agents.Count);
@@ -117,10 +114,10 @@ namespace Horde.Build.Tests.Fleet
 		{
 			using FleetService service = GetAutoscaleService(_fleetManagerSpy);
 			IPool pool = await PoolService.CreatePoolAsync("testPool", null, true, 0, 0, sizeStrategy: PoolSizeStrategy.LeaseUtilization);
-			List<PoolSizeData> poolSizeDatas = new()
+			List<PoolSizeResult> poolSizeDatas = new()
 			{
-				new PoolSizeData(null!, new(), null),
-				new PoolSizeData(pool, new(), null),
+				new PoolSizeResult(null!, new(), null),
+				new PoolSizeResult(pool, new(), null),
 			};
 
 			poolSizeDatas = await service.CalculatePoolSizesAsync(poolSizeDatas, CancellationToken.None);
@@ -134,16 +131,16 @@ namespace Horde.Build.Tests.Fleet
 			IPool pool = await PoolService.CreatePoolAsync("testPool", null, true, 0, 0, sizeStrategy: PoolSizeStrategy.NoOp);
 
 			// First scale-out will succeed
-			await service.ResizePoolsAsync(new() { new PoolSizeData(pool, new List<IAgent>(), 1) });
+			await service.ResizePoolsAsync(new() { new PoolSizeResult(pool, new List<IAgent>(), 1) });
 			Assert.AreEqual(1, _fleetManagerSpy.ExpandPoolAsyncCallCount);
 			
 			// Cannot scale-out due to cool-down
-			await service.ResizePoolsAsync(new() { new PoolSizeData(pool, new List<IAgent>(), 2) });
+			await service.ResizePoolsAsync(new() { new PoolSizeResult(pool, new List<IAgent>(), 2) });
 			Assert.AreEqual(1, _fleetManagerSpy.ExpandPoolAsyncCallCount);
 
 			// Wait some time and then try again
 			await Clock.AdvanceAsync(TimeSpan.FromHours(2));
-			await service.ResizePoolsAsync(new() { new PoolSizeData(pool, new List<IAgent>(), 2) });
+			await service.ResizePoolsAsync(new() { new PoolSizeResult(pool, new List<IAgent>(), 2) });
 			Assert.AreEqual(2, _fleetManagerSpy.ExpandPoolAsyncCallCount);
 		}
 		
@@ -156,16 +153,16 @@ namespace Horde.Build.Tests.Fleet
 			IAgent agent2 = await CreateAgentAsync(pool);
 
 			// First scale-out will succeed
-			await service.ResizePoolsAsync(new() { new PoolSizeData(pool, new () { agent1, agent2 }, 1) });
+			await service.ResizePoolsAsync(new() { new PoolSizeResult(pool, new () { agent1, agent2 }, 1) });
 			Assert.AreEqual(1, _fleetManagerSpy.ShrinkPoolAsyncCallCount);
 			
 			// Cannot scale-out due to cool-down
-			await service.ResizePoolsAsync(new() { new PoolSizeData(pool, new () { agent1 }, 0) });
+			await service.ResizePoolsAsync(new() { new PoolSizeResult(pool, new () { agent1 }, 0) });
 			Assert.AreEqual(1, _fleetManagerSpy.ShrinkPoolAsyncCallCount);
 
 			// Wait some time and then try again
 			await Clock.AdvanceAsync(TimeSpan.FromHours(2));
-			await service.ResizePoolsAsync(new() { new PoolSizeData(pool, new () { agent1 }, 0) });
+			await service.ResizePoolsAsync(new() { new PoolSizeResult(pool, new () { agent1 }, 0) });
 			Assert.AreEqual(2, _fleetManagerSpy.ShrinkPoolAsyncCallCount);
 		}
 
diff --git a/Engine/Source/Programs/Horde/Horde.Build.Tests/Fleet/JobQueueStrategyTest.cs b/Engine/Source/Programs/Horde/Horde.Build.Tests/Fleet/JobQueueStrategyTest.cs
--- a/Engine/Source/Programs/Horde/Horde.Build.Tests/Fleet/JobQueueStrategyTest.cs
+++ b/Engine/Source/Programs/Horde/Horde.Build.Tests/Fleet/JobQueueStrategyTest.cs
@@ -27,7 +27,7 @@ namespace Horde.Build.Tests.Fleet
 		[TestMethod]
 		public async Task GetPoolQueueSizes()
 		{
-			(JobQueueStrategy strategy, PoolSizeData poolSizeData) = await SetUpJobsAsync(1, 5);
+			(JobQueueStrategy strategy, PoolSizeResult poolSizeData) = await SetUpJobsAsync(1, 5);
 			await Clock.AdvanceAsync(TimeSpan.FromSeconds(strategy.Settings.ReadyTimeThresholdSec) + TimeSpan.FromSeconds(5));
 			Dictionary<PoolId, int> poolQueueSizes = await strategy.GetPoolQueueSizesAsync(Clock.UtcNow - TimeSpan.FromHours(2));
 			Assert.AreEqual(1, poolQueueSizes.Count);
@@ -84,16 +84,15 @@ namespace Horde.Build.Tests.Fleet
 
 		public async Task AssertAgentCount(int numBatchesReady, int expectedAgentDelta, bool waitedBeyondThreshold = true, int numAgents = 8)
 		{
-			(JobQueueStrategy strategy, PoolSizeData poolSizeData) = await SetUpJobsAsync(1, numBatchesReady, numAgents);
+			(JobQueueStrategy strategy, PoolSizeResult poolSizeData) = await SetUpJobsAsync(1, numBatchesReady, numAgents);
 			TimeSpan timeToWait = waitedBeyondThreshold
 				? TimeSpan.FromSeconds(strategy.Settings.ReadyTimeThresholdSec) + TimeSpan.FromSeconds(5)
 				: TimeSpan.FromSeconds(15);
 			
 			await Clock.AdvanceAsync(timeToWait);
 
-			List<PoolSizeData> result = await strategy.CalcDesiredPoolSizesAsync(new() { poolSizeData });
-			Assert.AreEqual(1, result.Count);
-			Assert.AreEqual(result[0].Agents.Count + expectedAgentDelta, result[0].DesiredAgentCount);
+			PoolSizeResult result = await strategy.CalculatePoolSizeAsync(poolSizeData.Pool, poolSizeData.Agents);
+			Assert.AreEqual(result.Agents.Count + expectedAgentDelta, result.DesiredAgentCount);
 		}
 	
 		/// <summary>
@@ -101,7 +100,7 @@ namespace Horde.Build.Tests.Fleet
 		/// </summary>
 		/// <param name="numBatchesRunning">Num of job batches that should be in state running</param>
 		/// <param name="numBatchesReady">Num of job batches that should be in state waiting</param>
-		private async Task<(JobQueueStrategy, PoolSizeData)> SetUpJobsAsync(int numBatchesRunning, int numBatchesReady, int numAgents = 8)
+		private async Task<(JobQueueStrategy, PoolSizeResult)> SetUpJobsAsync(int numBatchesRunning, int numBatchesReady, int numAgents = 8)
 		{
 			IPool pool1 = await PoolService.CreatePoolAsync("bogusPool1", null, true, 0, 0);
 			List<IAgent> agents = new();
@@ -110,7 +109,7 @@ namespace Horde.Build.Tests.Fleet
 				agents.Add(await CreateAgentAsync(pool1));
 			}
 			
-			PoolSizeData poolSize = new (pool1, agents, null);
+			PoolSizeResult poolSize = new (pool1, agents, null);
 			
 			string agentTypeName1 = "bogusAgentType1";
 			Dictionary<string, AgentConfig> agentTypes = new() { {agentTypeName1, new() { Pool = new PoolId(pool1.Name) } }, };
diff --git a/Engine/Source/Programs/Horde/Horde.Build.Tests/Fleet/LeaseUtilizationStrategyTest.cs b/Engine/Source/Programs/Horde/Horde.Build.Tests/Fleet/LeaseUtilizationStrategyTest.cs
--- a/Engine/Source/Programs/Horde/Horde.Build.Tests/Fleet/LeaseUtilizationStrategyTest.cs
+++ b/Engine/Source/Programs/Horde/Horde.Build.Tests/Fleet/LeaseUtilizationStrategyTest.cs
@@ -98,9 +98,8 @@ namespace Horde.Build.Tests.Fleet
 		private async Task AssertPoolSizeAsync(IPool pool, int expectedNumAgents)
 		{
 			IPoolSizeStrategy strategy = FleetService.CreatePoolSizeStrategy(pool);
-			List<PoolSizeData> output = await strategy.CalcDesiredPoolSizesAsync(new() { new(pool, _poolAgents, null) });
-			Assert.AreEqual(1, output.Count);
-			Assert.AreEqual(expectedNumAgents, output[0].DesiredAgentCount);
+			PoolSizeResult output = await strategy.CalculatePoolSizeAsync(pool, _poolAgents);
+			Assert.AreEqual(expectedNumAgents, output.DesiredAgentCount);
 		}
 		
 		private async Task<ILease> AddPlaceholderLease(IAgent agent, IPool pool, DateTime startTime, TimeSpan duration)
diff --git a/Engine/Source/Programs/Horde/Horde.Build/Agents/Fleet/ComputeQueueAwsMetricStrategy.cs b/Engine/Source/Programs/Horde/Horde.Build/Agents/Fleet/ComputeQueueAwsMetricStrategy.cs
--- a/Engine/Source/Programs/Horde/Horde.Build/Agents/Fleet/ComputeQueueAwsMetricStrategy.cs
+++ b/Engine/Source/Programs/Horde/Horde.Build/Agents/Fleet/ComputeQueueAwsMetricStrategy.cs
@@ -8,6 +8,7 @@ using System.Threading.Tasks;
 using Amazon.CloudWatch;
 using Amazon.CloudWatch.Model;
 using EpicGames.Horde.Compute;
+using Horde.Build.Agents.Pools;
 using Horde.Build.Compute;
 using Microsoft.Extensions.Logging;
 using OpenTracing;
@@ -41,7 +42,7 @@ namespace Horde.Build.Agents.Fleet
 			Namespace = cloudWatchNamespace;
 		}
 	}
-	
+
 	/// <summary>
 	/// A no-op strategy that reports size of compute task queue for a given pool as AWS CloudWatch metrics
 	/// A metric that later can be used as a source for AWS-controlled auto-scaling policies. 
@@ -50,6 +51,7 @@ namespace Horde.Build.Agents.Fleet
 	{
 		private readonly IAmazonCloudWatch _cloudWatch;
 		private readonly IComputeService _computeService;
+		private readonly ComputeQueueAwsMetricSettings _settings;
 		private readonly ILogger<ComputeQueueAwsMetricStrategy> _logger;
 
 		/// <summary>
@@ -57,11 +59,14 @@ namespace Horde.Build.Agents.Fleet
 		/// </summary>
 		/// <param name="cloudWatch"></param>
 		/// <param name="computeService"></param>
+		/// <param name="settings"></param>
 		/// <param name="logger"></param>
-		public ComputeQueueAwsMetricStrategy(IAmazonCloudWatch cloudWatch, IComputeService computeService, ILogger<ComputeQueueAwsMetricStrategy> logger)
+		public ComputeQueueAwsMetricStrategy(IAmazonCloudWatch cloudWatch, IComputeService computeService,
+			ComputeQueueAwsMetricSettings settings, ILogger<ComputeQueueAwsMetricStrategy> logger)
 		{
 			_cloudWatch = cloudWatch;
 			_computeService = computeService;
+			_settings = settings;
 			_logger = logger;
 		}
 
@@ -69,64 +74,54 @@ namespace Horde.Build.Agents.Fleet
 		public string Name { get; } = "ComputeQueueAwsMetric";
 
 		/// <inheritdoc/>
-		public async Task<List<PoolSizeData>> CalcDesiredPoolSizesAsync(List<PoolSizeData> pools)
+		public async Task<PoolSizeResult> CalculatePoolSizeAsync(IPool pool, List<IAgent> agents)
 		{
 			using IScope scope = GlobalTracer.Instance.BuildSpan("ComputeQueueAwsMetricStrategy.CalcDesiredPoolSizesAsync").StartActive();
 
 			Dictionary<string, List<MetricDatum>> metricsPerCloudWatchNamespace = new();
-			foreach (PoolSizeData poolSizeData in pools)
-			{
-				ComputeQueueAwsMetricSettings? settings = poolSizeData.Pool.ComputeQueueAwsMetricSettings;
-				if (settings == null)
-				{
-					_logger.LogWarning("Pool {PoolId} is configured to use strategy {Strategy} but has no settings configured", poolSizeData.Pool.Id, nameof(ComputeQueueAwsMetricStrategy));
-					continue;
-				}
-				
-				int numAgents = poolSizeData.Agents.Count;
-				int totalCpuCores = poolSizeData.Agents.Select(x => x.Resources.TryGetValue(KnownPropertyNames.LogicalCores, out int numCpuCores) ? numCpuCores : 0).Sum();
-				int numQueuedComputeTasks = await _computeService.GetNumQueuedTasksForPoolAsync(new ClusterId(settings.ComputeClusterId), poolSizeData.Pool);
-				
-				double numQueuedTasksPerAgent = numQueuedComputeTasks / (double)Math.Max(numAgents, 1);
-				double numQueuedTasksPerCores = numQueuedComputeTasks / (double)Math.Max(totalCpuCores, 1);
+			int numAgents = agents.Count;
+			int totalCpuCores = agents.Select(x => x.Resources.TryGetValue(KnownPropertyNames.LogicalCores, out int numCpuCores) ? numCpuCores : 0).Sum();
+			int numQueuedComputeTasks =	await _computeService.GetNumQueuedTasksForPoolAsync(new ClusterId(_settings.ComputeClusterId), pool);
 
-				DateTime now = DateTime.UtcNow;
-				
-				if (!metricsPerCloudWatchNamespace.TryGetValue(settings.Namespace, out List<MetricDatum>? metricDatums))
-				{
-					metricDatums = new List<MetricDatum>();
-					metricsPerCloudWatchNamespace[settings.Namespace] = metricDatums;
-				}
-				
-				List<Dimension> dimensions = new ()
-				{
-					new () { Name = "Pool", Value = poolSizeData.Pool.Id.ToString() },
-					new () { Name = "ComputeCluster", Value = settings.ComputeClusterId },
-				};
-				metricDatums.Add(new ()
-				{
-					MetricName = "NumComputeTasksPerAgent",
-					Dimensions = dimensions,
-					Unit = StandardUnit.Count,
-					Value = numQueuedTasksPerAgent,
-					TimestampUtc = now
-				});
-				metricDatums.Add(new ()
-				{
-					MetricName = "NumComputeTasksPerCpuCores",
-					Dimensions = dimensions,
-					Unit = StandardUnit.Count,
-					Value = numQueuedTasksPerCores,
-					TimestampUtc = now
-				});
+			double numQueuedTasksPerAgent = numQueuedComputeTasks / (double)Math.Max(numAgents, 1);
+			double numQueuedTasksPerCores = numQueuedComputeTasks / (double)Math.Max(totalCpuCores, 1);
+
+			DateTime now = DateTime.UtcNow;
+
+			if (!metricsPerCloudWatchNamespace.TryGetValue(_settings.Namespace, out List<MetricDatum>? metricDatums))
+			{
+				metricDatums = new List<MetricDatum>();
+				metricsPerCloudWatchNamespace[_settings.Namespace] = metricDatums;
 			}
 
-			foreach ((string ns, List<MetricDatum> metricDatums) in metricsPerCloudWatchNamespace)
+			List<Dimension> dimensions = new()
+			{
+				new() { Name = "Pool", Value = pool.Id.ToString() },
+				new() { Name = "ComputeCluster", Value = _settings.ComputeClusterId },
+			};
+			metricDatums.Add(new()
+			{
+				MetricName = "NumComputeTasksPerAgent",
+				Dimensions = dimensions,
+				Unit = StandardUnit.Count,
+				Value = numQueuedTasksPerAgent,
+				TimestampUtc = now
+			});
+			metricDatums.Add(new()
+			{
+				MetricName = "NumComputeTasksPerCpuCores",
+				Dimensions = dimensions,
+				Unit = StandardUnit.Count,
+				Value = numQueuedTasksPerCores,
+				TimestampUtc = now
+			});
+
+			foreach ((string ns, List<MetricDatum> metricDatumsNs) in metricsPerCloudWatchNamespace)
 			{
 				using IScope cwScope = GlobalTracer.Instance.BuildSpan("Putting CloudWatch metrics").StartActive();
 				cwScope.Span.SetTag("namespace", ns);
-				
-				PutMetricDataRequest request = new() { Namespace = ns, MetricData = metricDatums };
+
+				PutMetricDataRequest request = new() { Namespace = ns, MetricData = metricDatumsNs };
 				PutMetricDataResponse response = await _cloudWatch.PutMetricDataAsync(request);
 				scope.Span.SetTag("res.statusCode", (int)response.HttpStatusCode);
 				if (response.HttpStatusCode != HttpStatusCode.OK)
@@ -136,7 +131,7 @@ namespace Horde.Build.Agents.Fleet
 			}
 
 			// Return the input pool size data as-is since we are only observing and not modifying
-			return pools;
+			return new PoolSizeResult(pool, agents, agents.Count);
 		}
 	}
 }
diff --git a/Engine/Source/Programs/Horde/Horde.Build/Agents/Fleet/FleetService.cs b/Engine/Source/Programs/Horde/Horde.Build/Agents/Fleet/FleetService.cs
--- a/Engine/Source/Programs/Horde/Horde.Build/Agents/Fleet/FleetService.cs
+++ b/Engine/Source/Programs/Horde/Horde.Build/Agents/Fleet/FleetService.cs
@@ -115,8 +115,8 @@ namespace Horde.Build.Agents.Fleet
 			Stopwatch stopwatch = Stopwatch.StartNew();
 			using IScope _ = GlobalTracer.Instance.BuildSpan("AutoscaleService.TickAsync").StartActive();
 			
-			List<PoolSizeData> input = await GetPoolSizeDataAsync();
-			List<PoolSizeData> output = await CalculatePoolSizesAsync(input, stoppingToken);
+			List<PoolSizeResult> input = await GetPoolSizeDataAsync();
+			List<PoolSizeResult> output = await CalculatePoolSizesAsync(input, stoppingToken);
 			await ResizePoolsAsync(output, stoppingToken);
 
 			stopwatch.Stop();
@@ -136,10 +136,10 @@ namespace Horde.Build.Agents.Fleet
 			_logger.LogInformation("Autoscaling pools (high frequency) took {ElapsedTime} ms", stopwatch.ElapsedMilliseconds);
 		}
 
-		internal async Task<List<PoolSizeData>> CalculatePoolSizesAsync(List<PoolSizeData> poolSizeDatas, CancellationToken cancellationToken)
+		internal async Task<List<PoolSizeResult>> CalculatePoolSizesAsync(List<PoolSizeResult> poolSizeDatas, CancellationToken cancellationToken)
 		{
 			using IScope _ = GlobalTracer.Instance.BuildSpan("AutoscaleService.CalculatePoolSizesAsync").StartActive();
-			ConcurrentQueue<PoolSizeData> results = new ();
+			ConcurrentQueue<PoolSizeResult> results = new ();
 			ParallelOptions options = new () { MaxDegreeOfParallelism = MaxParallelTasks, CancellationToken = cancellationToken };
 
 			await Parallel.ForEachAsync(poolSizeDatas, options, async (input, innerCt) =>
@@ -147,7 +147,7 @@ namespace Horde.Build.Agents.Fleet
 				try
 				{
 					IPoolSizeStrategy sizeStrategy = CreatePoolSizeStrategy(input.Pool);
-					PoolSizeData output = (await sizeStrategy.CalcDesiredPoolSizesAsync(new List<PoolSizeData> { input }))[0];
+					PoolSizeResult output = await sizeStrategy.CalculatePoolSizeAsync(input.Pool, input.Agents);
 					results.Enqueue(output);
 				}
 				catch (Exception e)
@@ -159,32 +159,32 @@ namespace Horde.Build.Agents.Fleet
 			return results.ToList();
 		}
 
-		internal async Task<List<PoolSizeData>> GetPoolSizeDataAsync()
+		internal async Task<List<PoolSizeResult>> GetPoolSizeDataAsync()
 		{
 			List<IAgent> agents = await _agentCollection.FindAsync(status: AgentStatus.Ok, enabled: true);
 			List<IAgent> GetAgentsInPool(PoolId poolId) => agents.FindAll(a => a.GetPools().Any(p => p == poolId));
 			List<IPool> pools = await _poolCollection.GetAsync();
 
-			return pools.Select(pool => new PoolSizeData(pool, GetAgentsInPool(pool.Id), null)).ToList();
+			return pools.Select(pool => new PoolSizeResult(pool, GetAgentsInPool(pool.Id), null)).ToList();
 		}
 
-		internal async Task ResizePoolAsync(PoolSizeData poolSizeData, CancellationToken cancellationToken)
+		internal async Task ResizePoolAsync(PoolSizeResult poolSizeResult, CancellationToken cancellationToken)
 		{
-			IPool pool = poolSizeData.Pool;
+			IPool pool = poolSizeResult.Pool;
 
-			if (!pool.EnableAutoscaling || poolSizeData.DesiredAgentCount == null)
+			if (!pool.EnableAutoscaling || poolSizeResult.DesiredAgentCount == null)
 			{
 				return;
 			}
 
-			int currentAgentCount = poolSizeData.Agents.Count;
-			int desiredAgentCount = poolSizeData.DesiredAgentCount.Value;
+			int currentAgentCount = poolSizeResult.Agents.Count;
+			int desiredAgentCount = poolSizeResult.DesiredAgentCount.Value;
 			int deltaAgentCount = desiredAgentCount - currentAgentCount;
 
-			IFleetManager fleetManager = CreateFleetManager(poolSizeData.Pool);
+			IFleetManager fleetManager = CreateFleetManager(poolSizeResult.Pool);
 
 			Dictionary<string, object> logScope = new() { ["FleetManager"] = fleetManager.GetType().Name };
-			logScope["PoolSizeStrategy"] = poolSizeData.Status ?? new Dictionary<string, object>();
+			logScope["PoolSizeStrategy"] = poolSizeResult.Status ?? new Dictionary<string, object>();
 
 			using (_logger.BeginScope(logScope))
 			{
@@ -207,7 +207,7 @@ namespace Horde.Build.Agents.Fleet
 					scope.Span.SetTag("isCoolingDown", isCoolingDown);
 					if (!isCoolingDown)
 					{
-						await fleetManager.ExpandPoolAsync(pool, poolSizeData.Agents, deltaAgentCount, cancellationToken);
+						await fleetManager.ExpandPoolAsync(pool, poolSizeResult.Agents, deltaAgentCount, cancellationToken);
 						await _poolCollection.TryUpdateAsync(pool, lastScaleUpTime: _clock.UtcNow);
 					}
 					else
@@ -224,7 +224,7 @@ namespace Horde.Build.Agents.Fleet
 					scope.Span.SetTag("isCoolingDown", isCoolingDown);
 					if (!isCoolingDown)
 					{
-						await fleetManager.ShrinkPoolAsync(pool, poolSizeData.Agents, -deltaAgentCount, cancellationToken);
+						await fleetManager.ShrinkPoolAsync(pool, poolSizeResult.Agents, -deltaAgentCount, cancellationToken);
 						await _poolCollection.TryUpdateAsync(pool, lastScaleDownTime: _clock.UtcNow);
 					}
 					else
@@ -244,7 +244,7 @@ namespace Horde.Build.Agents.Fleet
 			_dogStatsd.Gauge("agentpools.autoscale.current", currentAgentCount, tags: new []{"pool:" + pool.Name});
 		}
 		
-		internal async Task ResizePoolsAsync(List<PoolSizeData> poolSizeDatas, CancellationToken cancellationToken = default)
+		internal async Task ResizePoolsAsync(List<PoolSizeResult> poolSizeDatas, CancellationToken cancellationToken = default)
 		{
 			using IScope _ = GlobalTracer.Instance.BuildSpan("AutoscaleService.ResizePoolsAsync").StartActive();
 			
@@ -308,6 +308,11 @@ namespace Horde.Build.Agents.Fleet
 								LeaseUtilizationSettings luSettings = DeserializeConfig<LeaseUtilizationSettings>(info.Config);
 								return new LeaseUtilizationStrategy(_agentCollection, _poolCollection, _leaseCollection, _clock, _cache, luSettings);
 							
+							// Disabled until moved to a separate factory class as FleetService should not contain AWS-specific classes
+							// case PoolSizeStrategy.ComputeQueueAwsMetric:
+							// 	ComputeQueueAwsMetricSettings cqamSettings = DeserializeConfig<ComputeQueueAwsMetricSettings>(info.Config);
+							// 	return new ComputeQueueAwsMetricStrategy(_awsCloudWatch, _computeService, cqamSettings);
+							
 							case PoolSizeStrategy.NoOp:
 								return new NoOpPoolSizeStrategy();
 							
diff --git a/Engine/Source/Programs/Horde/Horde.Build/Agents/Fleet/IPoolSizeStrategy.cs b/Engine/Source/Programs/Horde/Horde.Build/Agents/Fleet/IPoolSizeStrategy.cs
--- a/Engine/Source/Programs/Horde/Horde.Build/Agents/Fleet/IPoolSizeStrategy.cs
+++ b/Engine/Source/Programs/Horde/Horde.Build/Agents/Fleet/IPoolSizeStrategy.cs
@@ -1,7 +1,6 @@
 // Copyright Epic Games, Inc. All Rights Reserved.
 
 using System.Collections.Generic;
-using System.Linq;
 using System.Threading.Tasks;
 using Horde.Build.Agents.Pools;
 
@@ -60,7 +59,7 @@ namespace Horde.Build.Agents.Fleet
 	/// <summary>
 	/// Class for specifying and grouping data together required for calculating pool size
 	/// </summary>
-	public class PoolSizeData
+	public class PoolSizeResult
 	{
 		/// <summary>
 		/// Pool being resized
@@ -89,7 +88,7 @@ namespace Horde.Build.Agents.Fleet
 		/// <param name="agents"></param>
 		/// <param name="desiredAgentCount"></param>
 		/// <param name="status"></param>
-		public PoolSizeData(IPool pool, List<IAgent> agents, int? desiredAgentCount, IReadOnlyDictionary<string, object>? status = null)
+		public PoolSizeResult(IPool pool, List<IAgent> agents, int? desiredAgentCount, IReadOnlyDictionary<string, object>? status = null)
 		{
 			Pool = pool;
 			Agents = agents;
@@ -106,9 +105,9 @@ namespace Horde.Build.Agents.Fleet
 		/// <param name="desiredAgentCount"></param>
 		/// <param name="status"></param>
 		/// <returns>A new copy</returns>
-		public PoolSizeData Copy(IPool? pool = null, List<IAgent>? agents = null, int? desiredAgentCount = null, IReadOnlyDictionary<string, object>? status = null)
+		public PoolSizeResult Copy(IPool? pool = null, List<IAgent>? agents = null, int? desiredAgentCount = null, IReadOnlyDictionary<string, object>? status = null)
 		{
-			return new PoolSizeData(pool ?? Pool, agents ?? Agents, desiredAgentCount ?? DesiredAgentCount, status);
+			return new PoolSizeResult(pool ?? Pool, agents ?? Agents, desiredAgentCount ?? DesiredAgentCount, status);
 		}
 	}
 	
@@ -120,9 +119,10 @@ namespace Horde.Build.Agents.Fleet
 		/// <summary>
 		/// Calculate the adequate number of agents to be online for given pools
 		/// </summary>
-		/// <param name="pools">Pools including attached agents</param>
-		/// <returns></returns>
-		Task<List<PoolSizeData>> CalcDesiredPoolSizesAsync(List<PoolSizeData> pools);
+		/// <param name="pool">Pool to calculate size for</param>
+		/// <param name="agents">Available agents</param>
+		/// <returns>A result containing the desired agent count</returns>
+		Task<PoolSizeResult> CalculatePoolSizeAsync(IPool pool, List<IAgent> agents);
 		
 		/// <summary>
 		/// Name of the strategy
@@ -137,10 +137,9 @@ namespace Horde.Build.Agents.Fleet
 	public class NoOpPoolSizeStrategy : IPoolSizeStrategy
 	{
 		/// <inheritdoc/>
-		public Task<List<PoolSizeData>> CalcDesiredPoolSizesAsync(List<PoolSizeData> pools)
+		public Task<PoolSizeResult> CalculatePoolSizeAsync(IPool pool, List<IAgent> agents)
 		{
-			List<PoolSizeData> result = pools.Select(x => new PoolSizeData(x.Pool, x.Agents, x.Agents.Count)).ToList();
-			return Task.FromResult(result);
+			return Task.FromResult(new PoolSizeResult(pool, agents, agents.Count));
 		}
 
 		/// <inheritdoc/>
diff --git a/Engine/Source/Programs/Horde/Horde.Build/Agents/Fleet/JobQueueStrategy.cs b/Engine/Source/Programs/Horde/Horde.Build/Agents/Fleet/JobQueueStrategy.cs
--- a/Engine/Source/Programs/Horde/Horde.Build/Agents/Fleet/JobQueueStrategy.cs
+++ b/Engine/Source/Programs/Horde/Horde.Build/Agents/Fleet/JobQueueStrategy.cs
@@ -183,7 +183,7 @@ namespace Horde.Build.Agents.Fleet
 		}
 
 		/// <inheritdoc/>
-		public async Task<List<PoolSizeData>> CalcDesiredPoolSizesAsync(List<PoolSizeData> pools)
+		public async Task<PoolSizeResult> CalculatePoolSizeAsync(IPool pool, List<IAgent> agents)
 		{
 			DateTimeOffset minCreateTime = _clock.UtcNow - TimeSpan.FromMinutes(Settings.SamplePeriodMin);
 
@@ -195,32 +195,29 @@ namespace Horde.Build.Agents.Fleet
 				_cache.Set(CacheKey, poolQueueSizes, TimeSpan.FromSeconds(60));
 			}
 
-			return pools.Select(current =>
+			poolQueueSizes.TryGetValue(pool.Id, out int queueSize);
+			
+			Dictionary<string, object> status = new()
 			{
-				poolQueueSizes.TryGetValue(current.Pool.Id, out int queueSize);
-				
-				Dictionary<string, object> status = new()
-				{
-					["Name"] = GetType().Name,
-					["QueueSize"] = queueSize,
-					["ScaleOutFactor"] = Settings.ScaleOutFactor,
-					["ScaleInFactor"] = Settings.ScaleInFactor,
-					["SamplePeriodMin"] = Settings.SamplePeriodMin,
-					["ReadyTimeThresholdSec"] = Settings.ReadyTimeThresholdSec,
-				};
-				
-				if (queueSize > 0)
-				{
-					int additionalAgentCount = (int)Math.Ceiling(queueSize * Settings.ScaleOutFactor);
-					int desiredAgentCount = current.Agents.Count + additionalAgentCount;
-					return new PoolSizeData(current.Pool, current.Agents, desiredAgentCount, status);
-				}
-				else
-				{
-					int desiredAgentCount = (int)(current.Agents.Count * Settings.ScaleInFactor);
-					return new PoolSizeData(current.Pool, current.Agents, desiredAgentCount, status);
-				}
-			}).ToList();
+				["Name"] = GetType().Name,
+				["QueueSize"] = queueSize,
+				["ScaleOutFactor"] = Settings.ScaleOutFactor,
+				["ScaleInFactor"] = Settings.ScaleInFactor,
+				["SamplePeriodMin"] = Settings.SamplePeriodMin,
+				["ReadyTimeThresholdSec"] = Settings.ReadyTimeThresholdSec,
+			};
+			
+			if (queueSize > 0)
+			{
+				int additionalAgentCount = (int)Math.Ceiling(queueSize * Settings.ScaleOutFactor);
+				int desiredAgentCount = agents.Count + additionalAgentCount;
+				return new PoolSizeResult(pool, agents, desiredAgentCount, status);
+			}
+			else
+			{
+				int desiredAgentCount = (int)(agents.Count * Settings.ScaleInFactor);
+				return new PoolSizeResult(pool, agents, desiredAgentCount, status);
+			}
 		}
 	}
 }
diff --git a/Engine/Source/Programs/Horde/Horde.Build/Agents/Fleet/LeaseUtilizationStrategy.cs b/Engine/Source/Programs/Horde/Horde.Build/Agents/Fleet/LeaseUtilizationStrategy.cs
--- a/Engine/Source/Programs/Horde/Horde.Build/Agents/Fleet/LeaseUtilizationStrategy.cs
+++ b/Engine/Source/Programs/Horde/Horde.Build/Agents/Fleet/LeaseUtilizationStrategy.cs
@@ -230,7 +230,7 @@ namespace Horde.Build.Agents.Fleet
 		public string Name { get; } = "LeaseUtilization";
 
 		/// <inheritdoc/>
-		public async Task<List<PoolSizeData>> CalcDesiredPoolSizesAsync(List<PoolSizeData> pools)
+		public async Task<PoolSizeResult> CalculatePoolSizeAsync(IPool pool, List<IAgent> agents)
 		{
 			Dictionary<PoolId, PoolData> poolToData;
 			
@@ -241,45 +241,34 @@ namespace Horde.Build.Agents.Fleet
 				poolToData = await GetPoolDataAsync();
 				_cache.Set(CacheKey, poolToData, TimeSpan.FromSeconds(60));
 			}
-			
-			List<PoolSizeData> result = new();
 
-			foreach (PoolData poolData in poolToData.Values.OrderByDescending(x => x.Agents.Count))
+			PoolData poolData = poolToData[pool.Id];
+			
+			double utilization = poolData.Samples.Select(x => x._jobWork).OrderByDescending(x => x).Skip(Settings.NumSamples - Settings.NumSamplesForResult).First();
+		
+			// Number of agents in use over the sampling period. Can never be greater than number of agents available in pool.
+			int numAgentsUtilized = (int)utilization;
+			
+			// Include reserve agent count to ensure pool always can grow
+			int desiredAgentCount = Math.Max(numAgentsUtilized + Settings.NumReserveAgents, Settings.MinAgents);
+			
+			Dictionary<string, object> status = new()
 			{
-				IPool pool = poolData.Pool;
-
-				PoolSizeData? poolSize = pools.Find(x => x.Pool.Id == pool.Id);
-				if (poolSize != null)
-				{
-					double utilization = poolData.Samples.Select(x => x._jobWork).OrderByDescending(x => x).Skip(Settings.NumSamples - Settings.NumSamplesForResult).First();
-				
-					// Number of agents in use over the sampling period. Can never be greater than number of agents available in pool.
-					int numAgentsUtilized = (int)utilization;
-					
-					// Include reserve agent count to ensure pool always can grow
-					int desiredAgentCount = Math.Max(numAgentsUtilized + Settings.NumReserveAgents, Settings.MinAgents);
-					
-					Dictionary<string, object> status = new()
-					{
-						["Name"] = GetType().Name,
-						["Jobs"] = GetDensityMap(poolData.Samples.Select(x => x._jobWork / Math.Max(1, poolData.Agents.Count))),
-						["Total"] = GetDensityMap(poolData.Samples.Select(x => x._otherWork / Math.Max(1, poolData.Agents.Count))),
-						["Min"] = poolData.Samples.Min(x => x._jobWork),
-						["Max"] = poolData.Samples.Max(x => x._jobWork),
-						["Avg"] = utilization,
-						["Pct"] = poolData.Samples.Sum(x => x._jobWork) / Settings.NumSamples,
-						["SampleTimeSec"] = Settings.SampleTimeSec,
-						["NumSamples"] = Settings.NumSamples,
-						["NumSamplesForResult"] = Settings.NumSamplesForResult,
-						["MinAgents"] = Settings.MinAgents,
-						["NumReserveAgents"] = Settings.NumReserveAgents,
-					};
-
-					result.Add(new(pool, poolSize.Agents, desiredAgentCount, status));
-				}
-			}
+				["Name"] = GetType().Name,
+				["Jobs"] = GetDensityMap(poolData.Samples.Select(x => x._jobWork / Math.Max(1, poolData.Agents.Count))),
+				["Total"] = GetDensityMap(poolData.Samples.Select(x => x._otherWork / Math.Max(1, poolData.Agents.Count))),
+				["Min"] = poolData.Samples.Min(x => x._jobWork),
+				["Max"] = poolData.Samples.Max(x => x._jobWork),
+				["Avg"] = utilization,
+				["Pct"] = poolData.Samples.Sum(x => x._jobWork) / Settings.NumSamples,
+				["SampleTimeSec"] = Settings.SampleTimeSec,
+				["NumSamples"] = Settings.NumSamples,
+				["NumSamplesForResult"] = Settings.NumSamplesForResult,
+				["MinAgents"] = Settings.MinAgents,
+				["NumReserveAgents"] = Settings.NumReserveAgents,
+			};
 
-			return result;
+			return new PoolSizeResult(pool, agents, desiredAgentCount, status); 
 		}
 		
 		/// <summary>
