Horde: Fix issue with quarantined steps not being recorded in spans
#rnx

diff --git a/Engine/Source/Programs/Horde/Horde.Build.Tests/IssueServiceTests.cs b/Engine/Source/Programs/Horde/Horde.Build.Tests/IssueServiceTests.cs
--- a/Engine/Source/Programs/Horde/Horde.Build.Tests/IssueServiceTests.cs
+++ b/Engine/Source/Programs/Horde/Horde.Build.Tests/IssueServiceTests.cs
@@ -168,6 +168,8 @@ namespace Horde.Build.Tests
 			_perforce.AddChange(_mainStreamId, 115, 75, jerry, jerry, "Description\n#ROBOMERGE-SOURCE: CL 75 in //UE4/Release/...", new string[] { "a/e.cpp", "a/foo.cpp" });
 			_perforce.AddChange(_mainStreamId, 120, 120, chris, tim, "Description\n#ROBOMERGE-OWNER: Tim", new string[] { "a/f.cpp" });
 			_perforce.AddChange(_mainStreamId, 125, chris, "Description", new string[] { "a/g.cpp" });
+			_perforce.AddChange(_mainStreamId, 130, anne, "Description", new string[] { "a/g.cpp" });
+			_perforce.AddChange(_mainStreamId, 135, jerry, "Description", new string[] { "a/g.cpp" });
 
 			List<INode> nodes = new List<INode>();
 			nodes.Add(MockNode("Update Version Files", NodeAnnotations.Empty));
@@ -1913,12 +1915,21 @@ namespace Horde.Build.Tests
 		{
 			int issueId;
 			DateTime lastSeenAt;
+			int hour = 0;
 
 			// #1
+			// Scenario: Job succeeds establishing first success
+			{
+				IJob job = CreateJob(_mainStreamId, 100, "Test Build", _graph, TimeSpan.FromHours(hour++));
+				await UpdateCompleteStep(job, 0, 0, JobStepOutcome.Success);
+			}
+
+
+			// #2
 			// Scenario: Warning in first step
 			// Expected: Default issue is created
 			{
-				IJob job = CreateJob(_mainStreamId, 105, "Test Build", _graph);
+				IJob job = CreateJob(_mainStreamId, 105, "Test Build", _graph, TimeSpan.FromHours(hour++));
 				await AddEvent(job, 0, 0, new { level = nameof(LogLevel.Warning) }, EventSeverity.Warning);
 				await UpdateCompleteStep(job, 0, 0, JobStepOutcome.Warnings);
 
@@ -1930,6 +1941,13 @@ namespace Horde.Build.Tests
 
 				issueId = issues[0].Id;
 				lastSeenAt = issues[0].LastSeenAt;
+
+				List<IIssueSpan> spans = await IssueCollection.FindSpansAsync(issues[0].Id);
+				Assert.AreEqual(1, spans.Count);
+
+				IIssueDetails details = await IssueService.GetIssueDetailsAsync(issues[0]);
+				Assert.AreEqual(details.Steps.Count, 1);
+
 			}
 
 			// assign to bob
@@ -1938,27 +1956,27 @@ namespace Horde.Build.Tests
 			// Mark issue as quarantined
 			await IssueService.UpdateIssueAsync(issueId, quarantinedById: _jerryId);
 
-			// #2
+			// #3
 			// Scenario: Job succeeds
-			// Expected: Issue is not marked resolved
+			// Expected: Issue is not marked resolved, though step is added to span history
 			{
-				IJob job = CreateJob(_mainStreamId, 115, "Test Build", _graph);
+				IJob job = CreateJob(_mainStreamId, 115, "Test Build", _graph, TimeSpan.FromHours(hour++));
 				await UpdateCompleteStep(job, 0, 0, JobStepOutcome.Success);
 
 				IIssue? issue = await IssueCollection.GetIssueAsync(issueId);
 				Assert.IsNull(issue!.ResolvedAt);
+				Assert.IsNull(issue!.VerifiedAt);
 				Assert.AreEqual(issue!.OwnerId, _bobId);
 
-				List<IIssueSpan> spans = await IssueCollection.FindSpansAsync(issue.Id);
-				Assert.AreEqual(spans.Count, 1);
-
+				List<IIssue> issues = await IssueCollection.FindIssuesAsync();
+				Assert.AreEqual(1, issues.Count);
 			}
 
-			// #3
+			// #4
 			// Scenario: Job fails
 			// Expected: Existing issue is updated
 			{
-				IJob job = CreateJob(_mainStreamId, 125, "Test Build", _graph);
+				IJob job = CreateJob(_mainStreamId, 125, "Test Build", _graph, TimeSpan.FromHours(hour++));
 				await AddEvent(job, 0, 0, new { level = nameof(LogLevel.Warning) }, EventSeverity.Warning);
 				await UpdateCompleteStep(job, 0, 0, JobStepOutcome.Warnings);
 
@@ -1970,16 +1988,21 @@ namespace Horde.Build.Tests
 
 				Assert.AreEqual(issueId, issues[0].Id);
 				Assert.AreNotEqual(lastSeenAt, issues[0].LastSeenAt);
+
+				// make sure 4 steps have been recorded
+				IIssueDetails details = await IssueService.GetIssueDetailsAsync(issues[0]);
+				Assert.AreEqual(4, details.Steps.Count);
+
 			}
 
 			// Mark issue as not quarantined
 			await IssueService.UpdateIssueAsync(issueId, quarantinedById: UserId.Empty);
 
-			// #4
+			// #5
 			// Scenario: Job succeeds
 			// Expected: Issue is marked resolved and closed
 			{
-				IJob job = CreateJob(_mainStreamId, 135, "Test Build", _graph);
+				IJob job = CreateJob(_mainStreamId, 130, "Test Build", _graph, TimeSpan.FromHours(hour++));
 				await UpdateCompleteStep(job, 0, 0, JobStepOutcome.Success);
 
 				IIssue? issue = await IssueCollection.GetIssueAsync(issueId);
@@ -1990,7 +2013,6 @@ namespace Horde.Build.Tests
 												
 			}
 
-
 		}
 
 		[TestMethod]
diff --git a/Engine/Source/Programs/Horde/Horde.Build/Issues/IssueService.cs b/Engine/Source/Programs/Horde/Horde.Build/Issues/IssueService.cs
--- a/Engine/Source/Programs/Horde/Horde.Build/Issues/IssueService.cs
+++ b/Engine/Source/Programs/Horde/Horde.Build/Issues/IssueService.cs
@@ -585,13 +585,8 @@ namespace Horde.Build.Issues
 						}
 					}
 
-					// Only update sentinels for issues which aren't quarantined
-					List<IIssue> quarantined = await _issueCollection.FindIssuesAsync(openSpans.Select(x => x.IssueId).Distinct());
-					quarantined = quarantined.Where(x => x.QuarantinedByUserId != null).ToList();
-					List<IIssueSpan> sentinels = openSpans.Where(x => !quarantined.Any(y => x.IssueId == y.Id)).ToList();
-
 					// Try to update the sentinels for any other open steps
-					if (!await TryUpdateSentinelsAsync(sentinels, stream, job, batch, step))
+					if (!await TryUpdateSentinelsAsync(openSpans, stream, job, batch, step))
 					{
 						continue;
 					}
@@ -1274,14 +1269,34 @@ namespace Horde.Build.Issues
 				}
 				else if (job.Change > span.LastFailure.Change && (span.NextSuccess == null || job.Change < span.NextSuccess.Change))
 				{
-					NewIssueStepData newNextSucccess = new NewIssueStepData(job.Change, IssueSeverity.Unspecified, job.Name, job.Id, batch.Id, step.Id, step.StartTimeUtc ?? default, step.LogId, null, false);
+					
+					IIssue? issue = await _issueCollection.GetIssueAsync(span.IssueId);
+					bool quarantined = false;
+					if (issue != null && issue.QuarantinedByUserId != null)
+					{
+						quarantined = true;
+					}
 
-					if (await _issueCollection.TryUpdateSpanAsync(span, newNextSuccess: newNextSucccess) == null)
+					if (!quarantined)
 					{
-						return false;
+						NewIssueStepData newNextSucccess = new NewIssueStepData(job.Change, IssueSeverity.Unspecified, job.Name, job.Id, batch.Id, step.Id, step.StartTimeUtc ?? default, step.LogId, null, false);
+
+						if (await _issueCollection.TryUpdateSpanAsync(span, newNextSuccess: newNextSucccess) == null)
+						{
+							return false;
+						}
+
+						_logger.LogInformation("Set next success for issue {IssueId}, template {TemplateId}, node {Node} as job {JobId}, cl {Change}", span.IssueId, job.TemplateId, span.NodeName, job.Id, job.Change);
+
+					}
+					else
+					{						
+						NewIssueStepData newStep = new NewIssueStepData(job, batch, step, IssueSeverity.Unspecified, null, false);
+						await _issueCollection.AddStepAsync(span.Id, newStep);
+						_logger.LogInformation("Adding step to quarantined issue {IssueId}, template {TemplateId}, node {Node} job {JobId} batch {BatchId} step {StepId} cl {Change}", span.IssueId, job.TemplateId, span.NodeName, job.Id, batch.Id, step.Id, job.Change);
 					}
 
-					_logger.LogInformation("Set next success for issue {IssueId}, template {TemplateId}, node {Node} as job {JobId}, cl {Change}", span.IssueId, job.TemplateId, span.NodeName, job.Id, job.Change);
+					
 					await UpdateIssueDerivedDataAsync(span.IssueId);
 				}
 			}
