Horde: Re-fix conforms being scheduled if a job fails due to an invalid shelf.

diff --git a/Engine/Source/Programs/Horde/Horde.Build/Jobs/JobService.cs b/Engine/Source/Programs/Horde/Horde.Build/Jobs/JobService.cs
--- a/Engine/Source/Programs/Horde/Horde.Build/Jobs/JobService.cs
+++ b/Engine/Source/Programs/Horde/Horde.Build/Jobs/JobService.cs
@@ -860,7 +860,7 @@ namespace Horde.Build.Jobs
 						newError = error.Value;
 
 						// Find the agent and set the conform flag
-						if (checkForBadAgent)
+						if (newError == JobStepBatchError.Incomplete && checkForBadAgent)
 						{
 							for (; ; )
 							{
