Horde: Add support for enumerating bucket contents within IStorageBackend.

diff --git a/Engine/Source/Programs/Horde/Horde.Build/Storage/Backends/AwsStorageBackend.cs b/Engine/Source/Programs/Horde/Horde.Build/Storage/Backends/AwsStorageBackend.cs
--- a/Engine/Source/Programs/Horde/Horde.Build/Storage/Backends/AwsStorageBackend.cs
+++ b/Engine/Source/Programs/Horde/Horde.Build/Storage/Backends/AwsStorageBackend.cs
@@ -4,7 +4,7 @@ using System;
 using System.Collections.Generic;
 using System.Globalization;
 using System.IO;
-using System.Text;
+using System.Runtime.CompilerServices;
 using System.Threading;
 using System.Threading.Tasks;
 using Amazon;
@@ -118,6 +118,11 @@ namespace Horde.Build.Storage.Backends
 		/// </summary>
 		private readonly SemaphoreSlim _semaphore;
 
+		/// <summary>
+		/// Prefix for objects in the bucket
+		/// </summary>
+		private readonly string _pathPrefix;
+
 		/// <summary>
 		/// Logger interface
 		/// </summary>
@@ -138,6 +143,12 @@ namespace Horde.Build.Storage.Backends
 			_semaphore = new SemaphoreSlim(16);
 			_logger = logger;
 
+			_pathPrefix = (_options.AwsBucketPath ?? String.Empty).TrimEnd('/');
+			if (_pathPrefix.Length > 0)
+			{
+				_pathPrefix += '/';
+			}
+
 			logger.LogInformation("Created AWS storage backend for bucket {BucketName} using credentials {Credentials} {CredentialsStr}", options.AwsBucketName, awsOptions.Credentials.GetType(), awsOptions.Credentials.ToString());
 		}
 
@@ -239,22 +250,7 @@ namespace Horde.Build.Storage.Backends
 			}
 		}
 
-		string GetFullPath(string path)
-		{
-			if (_options.AwsBucketPath == null)
-			{
-				return path;
-			}
-
-			StringBuilder result = new StringBuilder();
-			result.Append(_options.AwsBucketPath);
-			if (result.Length > 0 && result[^1] != '/')
-			{
-				result.Append('/');
-			}
-			result.Append(path);
-			return result.ToString();
-		}
+		string GetFullPath(string path) => _pathPrefix + path;
 
 		/// <inheritdoc/>
 		public Task<Stream?> TryReadAsync(string path, CancellationToken cancellationToken)
@@ -474,5 +470,42 @@ namespace Horde.Build.Storage.Backends
 				return false;
 			}
 		}
+
+		/// <inheritdoc/>
+		public async IAsyncEnumerable<string> EnumerateAsync([EnumeratorCancellation] CancellationToken cancellationToken = default)
+		{
+			ListObjectsV2Request request = new ListObjectsV2Request();
+			request.BucketName = _options.AwsBucketName;
+			if(_pathPrefix.Length > 0)
+			{
+				request.Prefix = _pathPrefix;
+			}
+
+			for (; ; )
+			{
+				ListObjectsV2Response response = await _client.ListObjectsV2Async(request, cancellationToken);
+				foreach (S3Object obj in response.S3Objects)
+				{
+					string path = obj.Key;
+					if (path.StartsWith(_pathPrefix, StringComparison.Ordinal))
+					{
+						yield return path.Substring(_pathPrefix.Length);
+					}
+					else
+					{
+						_logger.LogError("Unexpected object enumerated from {AwsBucketName} - expected object \"{Path}\" to start with \"{AwsBucketPath}\"", _options.AwsBucketName, path, _pathPrefix);
+					}
+				}
+
+				if (response.IsTruncated)
+				{
+					request.ContinuationToken = response.NextContinuationToken;
+				}
+				else
+				{
+					break;
+				}
+			}
+		}
 	}
 }
diff --git a/Engine/Source/Programs/Horde/Horde.Build/Storage/Backends/FileStorageBackend.cs b/Engine/Source/Programs/Horde/Horde.Build/Storage/Backends/FileStorageBackend.cs
--- a/Engine/Source/Programs/Horde/Horde.Build/Storage/Backends/FileStorageBackend.cs
+++ b/Engine/Source/Programs/Horde/Horde.Build/Storage/Backends/FileStorageBackend.cs
@@ -1,10 +1,13 @@
 // Copyright Epic Games, Inc. All Rights Reserved.
 
 using System;
+using System.Collections.Generic;
 using System.IO;
+using System.Runtime.CompilerServices;
 using System.Threading;
 using System.Threading.Tasks;
 using EpicGames.Core;
+using EpicGames.Horde.Storage;
 using OpenTracing;
 using OpenTracing.Util;
 
@@ -148,5 +151,44 @@ namespace Horde.Build.Storage.Backends
 			FileReference.Delete(location);
 			return Task.CompletedTask;
 		}
+
+		/// <inheritdoc/>
+		public async IAsyncEnumerable<string> EnumerateAsync([EnumeratorCancellation] CancellationToken cancellationToken = default)
+		{
+			Stack<IEnumerator<DirectoryInfo>> queue = new Stack<IEnumerator<DirectoryInfo>>();
+			try
+			{
+				queue.Push(new List<DirectoryInfo> { _baseDir.ToDirectoryInfo() }.GetEnumerator());
+				while (queue.Count > 0)
+				{
+					IEnumerator<DirectoryInfo> top = queue.Peek();
+					if (!top.MoveNext())
+					{
+						top.Dispose();
+						queue.Pop();
+						continue;
+					}
+
+					DirectoryInfo current = top.Current;
+					foreach (FileInfo fileInfo in current.EnumerateFiles("*"))
+					{
+						string path = fileInfo.FullName.Substring(_baseDir.FullName.Length + 1).Replace(Path.DirectorySeparatorChar, '/');
+						yield return path;
+					}
+
+					queue.Push(current.EnumerateDirectories().GetEnumerator());
+
+					cancellationToken.ThrowIfCancellationRequested();
+					await Task.Yield();
+				}
+			}
+			finally
+			{
+				while (queue.TryPop(out IEnumerator<DirectoryInfo>? enumerator))
+				{
+					enumerator.Dispose();
+				}
+			}
+		}
 	}
 }
diff --git a/Engine/Source/Programs/Horde/Horde.Build/Storage/Backends/MemoryStorageBackend.cs b/Engine/Source/Programs/Horde/Horde.Build/Storage/Backends/MemoryStorageBackend.cs
--- a/Engine/Source/Programs/Horde/Horde.Build/Storage/Backends/MemoryStorageBackend.cs
+++ b/Engine/Source/Programs/Horde/Horde.Build/Storage/Backends/MemoryStorageBackend.cs
@@ -1,7 +1,9 @@
 // Copyright Epic Games, Inc. All Rights Reserved.
 
 using System.Collections.Concurrent;
+using System.Collections.Generic;
 using System.IO;
+using System.Runtime.CompilerServices;
 using System.Threading;
 using System.Threading.Tasks;
 
@@ -64,5 +66,16 @@ namespace Horde.Build.Storage.Backends
 			_pathToData.TryRemove(path, out _);
 			return Task.CompletedTask;
 		}
+
+		/// <inheritdoc/>
+		public async IAsyncEnumerable<string> EnumerateAsync([EnumeratorCancellation] CancellationToken cancellationToken = default)
+		{
+			foreach (string path in _pathToData.Keys)
+			{
+				yield return path;
+				cancellationToken.ThrowIfCancellationRequested();
+				await Task.Yield();
+			}
+		}
 	}
 }
diff --git a/Engine/Source/Programs/Horde/Horde.Build/Storage/IStorageBackend.cs b/Engine/Source/Programs/Horde/Horde.Build/Storage/IStorageBackend.cs
--- a/Engine/Source/Programs/Horde/Horde.Build/Storage/IStorageBackend.cs
+++ b/Engine/Source/Programs/Horde/Horde.Build/Storage/IStorageBackend.cs
@@ -1,6 +1,7 @@
 // Copyright Epic Games, Inc. All Rights Reserved.
 
 using System;
+using System.Collections.Generic;
 using System.IO;
 using System.Threading;
 using System.Threading.Tasks;
@@ -54,6 +55,13 @@ namespace Horde.Build.Storage
 		/// <param name="cancellationToken">Cancellation token for the operation</param>
 		/// <returns>Async task</returns>
 		Task DeleteAsync(string path, CancellationToken cancellationToken = default);
+
+		/// <summary>
+		/// Enumerates all the objects in the store
+		/// </summary>
+		/// <param name="cancellationToken">Cancellation token for the operation</param>
+		/// <returns>Sequence of object paths</returns>
+		IAsyncEnumerable<string> EnumerateAsync(CancellationToken cancellationToken = default);
 	}
 
 	/// <summary>
@@ -97,6 +105,9 @@ namespace Horde.Build.Storage
 
 			/// <inheritdoc/>
 			public Task<bool> ExistsAsync(string path, CancellationToken cancellationToken) => _inner.ExistsAsync(path, cancellationToken);
+
+			/// <inheritdoc/>
+			public IAsyncEnumerable<string> EnumerateAsync(CancellationToken cancellationToken = default) => _inner.EnumerateAsync(cancellationToken);
 		}
 
 		/// <summary>
