Horde: Fix duplicate step entries in quarantined issue span
#rnx

diff --git a/Engine/Source/Programs/Horde/Horde.Build.Tests/IssueServiceTests.cs b/Engine/Source/Programs/Horde/Horde.Build.Tests/IssueServiceTests.cs
--- a/Engine/Source/Programs/Horde/Horde.Build.Tests/IssueServiceTests.cs
+++ b/Engine/Source/Programs/Horde/Horde.Build.Tests/IssueServiceTests.cs
@@ -1991,7 +1991,7 @@ namespace Horde.Build.Tests
 
 				// make sure 4 steps have been recorded
 				IIssueDetails details = await IssueService.GetIssueDetailsAsync(issues[0]);
-				Assert.AreEqual(4, details.Steps.Count);
+				Assert.AreEqual(3, details.Steps.Count);
 
 			}
 
diff --git a/Engine/Source/Programs/Horde/Horde.Build/Issues/IssueService.cs b/Engine/Source/Programs/Horde/Horde.Build/Issues/IssueService.cs
--- a/Engine/Source/Programs/Horde/Horde.Build/Issues/IssueService.cs
+++ b/Engine/Source/Programs/Horde/Horde.Build/Issues/IssueService.cs
@@ -1268,16 +1268,9 @@ namespace Horde.Build.Issues
 					await UpdateIssueDerivedDataAsync(span.IssueId);
 				}
 				else if (job.Change > span.LastFailure.Change && (span.NextSuccess == null || job.Change < span.NextSuccess.Change))
-				{
-					
+				{					
 					IIssue? issue = await _issueCollection.GetIssueAsync(span.IssueId);
-					bool quarantined = false;
-					if (issue != null && issue.QuarantinedByUserId != null)
-					{
-						quarantined = true;
-					}
-
-					if (!quarantined)
+					if (issue == null || issue.QuarantinedByUserId == null)
 					{
 						NewIssueStepData newNextSucccess = new NewIssueStepData(job.Change, IssueSeverity.Unspecified, job.Name, job.Id, batch.Id, step.Id, step.StartTimeUtc ?? default, step.LogId, null, false);
 
@@ -1290,12 +1283,17 @@ namespace Horde.Build.Issues
 
 					}
 					else
-					{						
-						NewIssueStepData newStep = new NewIssueStepData(job, batch, step, IssueSeverity.Unspecified, null, false);
-						await _issueCollection.AddStepAsync(span.Id, newStep);
-						_logger.LogInformation("Adding step to quarantined issue {IssueId}, template {TemplateId}, node {Node} job {JobId} batch {BatchId} step {StepId} cl {Change}", span.IssueId, job.TemplateId, span.NodeName, job.Id, batch.Id, step.Id, job.Change);
+					{
+						// If the issue is quarantined, spans may have been updated by log events, though in the case they aren't
+						// need to add a step to the span history
+						List<IIssueStep> steps = await _issueCollection.FindStepsAsync(span.Id);
+						if (steps.FirstOrDefault(x => x.JobId == job.Id && x.StepId == step.Id) == null)
+						{
+							NewIssueStepData newStep = new NewIssueStepData(job, batch, step, IssueSeverity.Unspecified, null, false);
+							await _issueCollection.AddStepAsync(span.Id, newStep);
+							_logger.LogInformation("Adding step to quarantined issue {IssueId}, template {TemplateId}, node {Node} job {JobId} batch {BatchId} step {StepId} cl {Change}", span.IssueId, job.TemplateId, span.NodeName, job.Id, batch.Id, step.Id, job.Change);
+						}
 					}
-
 					
 					await UpdateIssueDerivedDataAsync(span.IssueId);
 				}
