Horde: Fix generic errors being masked by a previous compiler match
#rnx

diff --git a/Engine/Source/Programs/Horde/Horde.Build.Tests/IssueServiceTests.cs b/Engine/Source/Programs/Horde/Horde.Build.Tests/IssueServiceTests.cs
--- a/Engine/Source/Programs/Horde/Horde.Build.Tests/IssueServiceTests.cs
+++ b/Engine/Source/Programs/Horde/Horde.Build.Tests/IssueServiceTests.cs
@@ -456,23 +456,28 @@ namespace Horde.Build.Tests
 			{
 				string[] lines =
 				{
-					@"< enterprise.max.test_maxscript_datasmith_export.from_maxscript_to_unreal[vray_materials.ms] >",
-					@"  [ :ERROR: ] [2022.14.06-12:08:55] [log_parser] keyword found 'logpython: error:'",
-					@"   > [2022.06.14-12.08.44:338][  1]LogPython: Error: One or more instance material is loaded but not expected : ['VRayMtl__2_.VRayMtl__2_', 'VRayMtl__3_.VRayMtl__3_', 'VRayMtl__1_.VRayMtl__1_', 'VRayMtl__4_.VRayMtl__4_']",
-					@"     [2022.06.14-12.08.44:338][  1]LogPython: Error: One or more instance material expected but not loaded: ['VRayMtl_2.VRayMtl_2', 'VRayMtl_3.VRayMtl_3', 'VRayMtl_4.VRayMtl_4', 'VRayMtl_7.VRayMtl_7']",
-					@"   > File ""D:\build\++UE5\Sync\Engine\Saved\pydrover\session[2022.14.06-11.50.11]\from_maxscript_to_unreal_4d5057656576\ue_log[2022.14.06-12.08.31].txt"", line 1085",
-					@"< enterprise.max.test_maxscript_datasmith_export.from_maxscript_to_unreal[vray_materials.ms] >"
+					@"Engine/Source/Editor/SparseVolumeTexture/Private/SparseVolumeTextureOpenVDB.h(38): warning: include path has multiple slashes (<openvdb//math/Half.h>)",
+					@"Engine/Source/Editor/SparseVolumeTexture/Private/SparseVolumeTextureOpenVDB.h(38): warning: include path has multiple slashes (<openvdb//math/Half.h>)",
+					@"Error: include cycle: 0: VulkanContext.h -> VulkanRenderpass.h -> VulkanContext.h",
+					@"Error: Unable to continue until this cycle has been removed.",
+					@"Took 692.6159064s to run IncludeTool.exe, ExitCode=1",
+					@"ERROR: IncludeTool.exe terminated with an exit code indicating an error (1)",
+					@"       while executing task <Spawn Exe=""D:\build\++UE5\Sync\Engine\Binaries\DotNET\IncludeTool\IncludeTool.exe"" Arguments=""-Mode=Scan -Target=UnrealEditor -Platform=Linux -Configuration=Development -WorkingDir=D:\build\++UE5\Sync\Working"" LogOutput=""True"" ErrorLevel=""1"" />",
+					@"      at D:\build\++UE5\Sync\Engine\Restricted\NotForLicensees\Build\DevStreams.xml(938)",
+					@"       (see d:\build\++UE5\Sync\Engine\Programs\AutomationTool\Saved\Logs\Log.txt for full exception trace)"
 				};
 
 				IJob job = CreateJob(_mainStreamId, 105, "Test Build", _graph);
 				await ParseEventsAsync(job, 0, 0, lines);
-				await UpdateCompleteStep(job, 0, 0, JobStepOutcome.Warnings);
+				await UpdateCompleteStep(job, 0, 0, JobStepOutcome.Failure);
 
 				List<IIssue> issues = await IssueCollection.FindIssuesAsync();
-				Assert.AreEqual(1, issues.Count);
+				Assert.AreEqual(2, issues.Count);
 				Assert.AreEqual(IssueSeverity.Error, issues[0].Severity);
+				Assert.AreEqual(IssueSeverity.Warning, issues[1].Severity);
 
 				Assert.AreEqual("Errors in Update Version Files", issues[0].Summary);
+				Assert.AreEqual("Compile warnings in SparseVolumeTextureOpenVDB.h", issues[1].Summary);				
 			}
 		}
 
diff --git a/Engine/Source/Programs/Horde/Horde.Build/Issues/Handlers/CompileIssueHandler.cs b/Engine/Source/Programs/Horde/Horde.Build/Issues/Handlers/CompileIssueHandler.cs
--- a/Engine/Source/Programs/Horde/Horde.Build/Issues/Handlers/CompileIssueHandler.cs
+++ b/Engine/Source/Programs/Horde/Horde.Build/Issues/Handlers/CompileIssueHandler.cs
@@ -42,7 +42,7 @@ namespace Horde.Build.Issues.Handlers
 			return eventId == KnownLogEvents.Compiler || eventId == KnownLogEvents.UHT || eventId == KnownLogEvents.AutomationTool_SourceFileLine || eventId == KnownLogEvents.MSBuild;
 		}
 
-		static bool IsMaskedEventId(EventId id) => id == KnownLogEvents.Generic || id == KnownLogEvents.ExitCode || id == KnownLogEvents.Systemic_Xge_BuildFailed;
+		static bool IsMaskedEventId(EventId id) => id == KnownLogEvents.ExitCode || id == KnownLogEvents.Systemic_Xge_BuildFailed;
 
 		/// <inheritdoc/>
 		public override void TagEvents(IJob job, INode node, IReadOnlyNodeAnnotations annotations, IReadOnlyList<IssueEvent> stepEvents)
