Horde: Deduplicate job scheduled notifications

Depending on pool setup, these notifications can be fired a lot. But we are usually only interested of one notification per job/pool, so deduplicating them is fine.

diff --git a/Engine/Source/Programs/Horde/Horde.Build.Tests/NotificationServiceTests.cs b/Engine/Source/Programs/Horde/Horde.Build.Tests/NotificationServiceTests.cs
--- a/Engine/Source/Programs/Horde/Horde.Build.Tests/NotificationServiceTests.cs
+++ b/Engine/Source/Programs/Horde/Horde.Build.Tests/NotificationServiceTests.cs
@@ -17,6 +17,7 @@ using Microsoft.Extensions.DependencyInjection;
 using Microsoft.VisualStudio.TestTools.UnitTesting;
 using Moq;
 using Horde.Build.Logs;
+using Microsoft.Extensions.Caching.Memory;
 
 namespace Horde.Build.Tests
 {
@@ -116,6 +117,38 @@ namespace Horde.Build.Tests
 			Assert.AreEqual(2, fakeSink.JobScheduledNotifications.Count);
 			Assert.AreEqual(1, fakeSink.JobScheduledCallCount);
 		}
+		
+		[TestMethod]
+		public async Task JobScheduledNotificationsAreDeduplicated()
+		{
+			FakeNotificationSink fakeSink = ServiceProvider.GetRequiredService<FakeNotificationSink>();
+			NotificationService service = (NotificationService)ServiceProvider.GetRequiredService<INotificationService>();
+			await service._ticker.StartAsync();
+			Fixture fixture = await CreateFixtureAsync();
+			IPool pool = await PoolService.CreatePoolAsync("BogusPool", properties: new Dictionary<string, string>());
+
+			service.NotifyJobScheduled(pool, false, fixture.Job1, fixture.Graph, SubResourceId.Random());
+			service.NotifyJobScheduled(pool, false, fixture.Job1, fixture.Graph, SubResourceId.Random());
+			service.NotifyJobScheduled(pool, false, fixture.Job1, fixture.Graph, SubResourceId.Random());
+			
+			// Currently no good way to wait for NotifyJobScheduled() to complete as the execution is completely async in background task (see ExecuteAsync)
+			await Task.Delay(1000);
+			await Clock.AdvanceAsync(service._notificationBatchInterval + TimeSpan.FromMinutes(5));
+			
+			// Only one job scheduled notification should have been sent, despite queuing three
+			Assert.AreEqual(1, fakeSink.JobScheduledNotifications.Count);
+			
+			// Clear the cache by compacting it 100%
+			MemoryCache cache = (MemoryCache)ServiceProvider.GetRequiredService<IMemoryCache>();
+			cache.Compact(1.0);
+			
+			// Notify of exactly the same job again
+			service.NotifyJobScheduled(pool, false, fixture.Job1, fixture.Graph, SubResourceId.Random());
+			
+			await Task.Delay(1000);
+			await Clock.AdvanceAsync(service._notificationBatchInterval + TimeSpan.FromMinutes(5));
+			Assert.AreEqual(2, fakeSink.JobScheduledNotifications.Count);
+		}
 
 		//public void NotifyLabelUpdate(IJob Job, IReadOnlyList<(LabelState, LabelOutcome)> OldLabelStates, IReadOnlyList<(LabelState, LabelOutcome)> NewLabelStates)
 		//{
diff --git a/Engine/Source/Programs/Horde/Horde.Build/Notifications/INotificationService.cs b/Engine/Source/Programs/Horde/Horde.Build/Notifications/INotificationService.cs
--- a/Engine/Source/Programs/Horde/Horde.Build/Notifications/INotificationService.cs
+++ b/Engine/Source/Programs/Horde/Horde.Build/Notifications/INotificationService.cs
@@ -1,5 +1,6 @@
 // Copyright Epic Games, Inc. All Rights Reserved.
 
+using System;
 using System.Collections.Generic;
 using System.Diagnostics.CodeAnalysis;
 using System.Security.Claims;
@@ -16,14 +17,15 @@ namespace Horde.Build.Notifications
 {
 	/// <summary>
 	/// Marker interface for (serializable) notifications
+	/// Implements IEquatable primarily for being IMemoryCache-compatible.
 	/// </summary>
 	[SuppressMessage("Design", "CA1040:Avoid empty interfaces", Justification = "Marker interface for generic type handling")]
-	public interface INotification { }
+	public interface INotification<T> : IEquatable<T> { }
 
 	/// <summary>
 	/// Notification for job scheduled events
 	/// </summary>
-	public class JobScheduledNotification : INotification
+	public class JobScheduledNotification : INotification<JobScheduledNotification>
 	{
 		/// <summary>Job ID</summary>
 		public string JobId { get; }
@@ -46,6 +48,29 @@ namespace Horde.Build.Notifications
 			JobName = jobName;
 			PoolName = poolName;
 		}
+
+		/// <inheritdoc />
+		public bool Equals(JobScheduledNotification? other)
+		{
+			if (ReferenceEquals(null, other)) { return false; }
+			if (ReferenceEquals(this, other)) { return true; }
+			return JobId == other.JobId && JobName == other.JobName && PoolName == other.PoolName;
+		}
+
+		/// <inheritdoc />
+		public override bool Equals(object? obj)
+		{
+			if (ReferenceEquals(null, obj)) { return false; }
+			if (ReferenceEquals(this, obj)) { return true; }
+			if (obj.GetType() != GetType()) { return false; }
+			return Equals((JobScheduledNotification)obj);
+		}
+
+		/// <inheritdoc />
+		public override int GetHashCode()
+		{
+			return HashCode.Combine(JobId, JobName, PoolName);
+		}
 	}
 	
 	/// <summary>
diff --git a/Engine/Source/Programs/Horde/Horde.Build/Notifications/NotificationService.cs b/Engine/Source/Programs/Horde/Horde.Build/Notifications/NotificationService.cs
--- a/Engine/Source/Programs/Horde/Horde.Build/Notifications/NotificationService.cs
+++ b/Engine/Source/Programs/Horde/Horde.Build/Notifications/NotificationService.cs
@@ -20,6 +20,7 @@ using Horde.Build.Streams;
 using Horde.Build.Users;
 using Horde.Build.Utilities;
 using HordeCommon;
+using Microsoft.Extensions.Caching.Memory;
 using Microsoft.Extensions.Hosting;
 using Microsoft.Extensions.Logging;
 using Microsoft.Extensions.Options;
@@ -87,6 +88,17 @@ namespace Horde.Build.Notifications
 		/// Instance of the <see cref="IDogStatsd"/>.
 		/// </summary>
 		private readonly IDogStatsd _dogStatsd;
+		
+		/// <summary>
+		/// Cache for de-duplicating queued notifications
+		/// </summary>
+		private readonly IMemoryCache _cache;
+		
+		/// <summary>
+		/// Lock object for manipulating the above cache
+		/// Used since batch notification queue handling is run async.
+		/// </summary>
+		private readonly object _cacheLock = new object();
 
 		/// <summary>
 		/// Connection pool for Redis databases
@@ -128,7 +140,22 @@ namespace Horde.Build.Notifications
 		/// <summary>
 		/// Constructor
 		/// </summary>
-		public NotificationService(IEnumerable<INotificationSink> sinks, IOptionsMonitor<ServerSettings> settings, ILogger<NotificationService> logger, IGraphCollection graphCollection, ISubscriptionCollection subscriptionCollection, INotificationTriggerCollection triggerCollection, IUserCollection userCollection, JobService jobService, StreamService streamService, IssueService issueService, ILogFileService logFileService, IDogStatsd dogStatsd, RedisService redisService, IClock clock)
+		public NotificationService(
+			IEnumerable<INotificationSink> sinks,
+			IOptionsMonitor<ServerSettings> settings,
+			ILogger<NotificationService> logger,
+			IGraphCollection graphCollection,
+			ISubscriptionCollection subscriptionCollection,
+			INotificationTriggerCollection triggerCollection,
+			IUserCollection userCollection,
+			JobService jobService,
+			StreamService streamService,
+			IssueService issueService,
+			ILogFileService logFileService,
+			IDogStatsd dogStatsd,
+			IMemoryCache cache,
+			RedisService redisService,
+			IClock clock)
 		{
 			_sinks = sinks.ToList();
 			_settings = settings;
@@ -142,6 +169,7 @@ namespace Horde.Build.Notifications
 			_issueService = issueService;
 			_logFileService = logFileService;
 			_dogStatsd = dogStatsd;
+			_cache = cache;
 			_redisConnectionPool = redisService.ConnectionPool;
 
 			issueService.OnIssueUpdated += NotifyIssueUpdated;
@@ -284,10 +312,24 @@ namespace Horde.Build.Notifications
 		/// <summary>
 		/// Enqueue a notification in Redis for batch sending later on 
 		/// </summary>
-		/// <param name="notification"></param>
+		/// <param name="notification">Notification to enqueue</param>
+		/// <param name="deduplicate">True if notification should be deduplicated</param>
 		/// <typeparam name="T">Any INotification type</typeparam>
-		private async Task EnqueueNotificationForBatchSending<T>(T notification) where T : INotification
+		private async Task EnqueueNotificationForBatchSending<T>(T notification, bool deduplicate = true) where T : INotification<T>
 		{
+			lock (_cacheLock)
+			{
+				if (deduplicate)
+				{
+					// Use cache to deduplicate notifications
+					if (_cache.TryGetValue(notification, out object? _))
+					{
+						return;
+					}
+					_cache.Set(notification, notification, _notificationBatchInterval / 2);					
+				}
+			}
+
 			try
 			{
 				byte[] data = JsonSerializer.SerializeToUtf8Bytes(notification);
@@ -317,7 +359,7 @@ namespace Horde.Build.Notifications
 		/// <typeparam name="T"></typeparam>
 		/// <returns>Deserialized notifications</returns>
 		/// <exception cref="Exception"></exception>
-		private async Task<List<T>> GetAllQueuedNotificationsAsync<T>() where T : INotification
+		private async Task<List<T>> GetAllQueuedNotificationsAsync<T>() where T : INotification<T>
 		{
 			IDatabase redis = _redisConnectionPool.GetDatabase();
 
diff --git a/Engine/Source/Programs/Horde/Horde.Build/ServerSettings.cs b/Engine/Source/Programs/Horde/Horde.Build/ServerSettings.cs
--- a/Engine/Source/Programs/Horde/Horde.Build/ServerSettings.cs
+++ b/Engine/Source/Programs/Horde/Horde.Build/ServerSettings.cs
@@ -463,7 +463,7 @@ namespace Horde.Build
 		public string? UpdateStreamsNotificationChannel { get; set; }
 		
 		/// <summary>
-		/// Slack channel to send job related notifications to
+		/// Slack channel to send job related notifications to. Multiple channels can be specified, separated by ;
 		/// </summary>
 		public string? JobNotificationChannel { get; set; }
 
