Horde: Fix fatal link warnings not having error issue severity
#rnx

diff --git a/Engine/Source/Programs/Horde/Horde.Build.Tests/IssueServiceTests.cs b/Engine/Source/Programs/Horde/Horde.Build.Tests/IssueServiceTests.cs
--- a/Engine/Source/Programs/Horde/Horde.Build.Tests/IssueServiceTests.cs
+++ b/Engine/Source/Programs/Horde/Horde.Build.Tests/IssueServiceTests.cs
@@ -1452,6 +1452,45 @@ namespace Horde.Build.Tests
 			}
 		}
 
+		[TestMethod]
+		public async Task SymbolIssueTest4()
+		{
+			// #1
+			// Scenario: Job step completes successfully at CL 105
+			// Expected: No issues are created
+			{
+				IJob job = CreateJob(_mainStreamId, 105, "Test Build", _graph);
+				await UpdateCompleteStep(job, 0, 0, JobStepOutcome.Success);
+
+				List<IIssue> issues = await IssueCollection.FindIssuesAsync();
+				Assert.AreEqual(0, issues.Count);
+			}
+
+			// #2
+			// Scenario: Job step fails at CL 120
+			// Expected: Creates a linker issue with severity error due to fatal warnings 
+			{
+				string[] lines =
+				{
+					@"ld: warning: direct access in function 'void Eigen::internal::evaluateProductBlockingSizesHeuristic<Eigen::half, Eigen::half, 1, long>(long&, long&, long&, long)' from file '../../EngineTest/Intermediate/Build/Mac/x86_64/EngineTest/Development/ORT/inverse.cc.o' to global weak symbol 'guard variable for Eigen::internal::manage_caching_sizes(Eigen::Action, long*, long*, long*)::m_cacheSizes' from file '../../EngineTest/Intermediate/Build/Mac/x86_64/EngineTest/Development/DynamicMesh/Module.DynamicMesh.4_of_5.cpp.o' means the weak symbol cannot be overridden at runtime. This was likely caused by different translation units being compiled with different visibility settings.",
+					@"ld: fatal warning(s) induced error (-fatal_warnings)",
+					@"clang: error: linker command failed with exit code 1 (use -v to see invocation)"              
+				};
+
+				IJob job = CreateJob(_mainStreamId, 120, "Test Build", _graph);
+				await ParseEventsAsync(job, 0, 0, lines);
+				await UpdateCompleteStep(job, 0, 0, JobStepOutcome.Failure);
+
+				List<IIssue> issues = await IssueCollection.FindIssuesAsync();
+				Assert.AreEqual(1, issues.Count);
+
+				IIssue issue = issues[0];
+
+				Assert.AreEqual(IssueSeverity.Error, issue.Severity);
+			}
+		}
+
+
 		[TestMethod]
 		public async Task LinkerIssueTest2()
 		{
diff --git a/Engine/Source/Programs/UnrealBuildTool/Matchers/LinkEventMatcher.cs b/Engine/Source/Programs/UnrealBuildTool/Matchers/LinkEventMatcher.cs
--- a/Engine/Source/Programs/UnrealBuildTool/Matchers/LinkEventMatcher.cs
+++ b/Engine/Source/Programs/UnrealBuildTool/Matchers/LinkEventMatcher.cs
@@ -47,8 +47,8 @@ namespace UnrealBuildTool.Matchers
 			bool isWarning = false;
 			while (cursor.IsMatch(lineCount, s_undefinedReferencePattern))
 			{
-				isError |= cursor.Contains("error:");
-				isWarning |= cursor.Contains("warning:");
+				isError |= cursor.Contains(lineCount, "error:");
+				isWarning |= cursor.Contains(lineCount, "warning:");
 				lineCount++;
 			}
 			if (lineCount > 0)
