Horde: Show quaratined step successes with success color
#rnx

diff --git a/Engine/Source/Programs/Horde/HordeDashboard/src/components/IssueViewV2.tsx b/Engine/Source/Programs/Horde/HordeDashboard/src/components/IssueViewV2.tsx
--- a/Engine/Source/Programs/Horde/HordeDashboard/src/components/IssueViewV2.tsx
+++ b/Engine/Source/Programs/Horde/HordeDashboard/src/components/IssueViewV2.tsx
@@ -867,7 +867,9 @@ const StreamCanvas: React.FC = () => {
                   if (step.severity === IssueSeverity.Error) {
                      color = colors.get(StatusColor.Failure)!;
                   }
-
+                  if (index && step.severity === IssueSeverity.Unspecified) {
+                     color = colors.get(StatusColor.Success)!;
+                  }
 
                   let baseTime = details.fakeCLTimes.get(step.change)!.getTime() / 1000;
 
