Horde: Remove per-agent setting to use different software channels. The channel to use is now configured via a condition evaluated on the agent, defined in the global config file.

diff --git a/Engine/Source/Programs/Horde/Horde.Build.Tests/Fixture.cs b/Engine/Source/Programs/Horde/Horde.Build.Tests/Fixture.cs
--- a/Engine/Source/Programs/Horde/Horde.Build.Tests/Fixture.cs
+++ b/Engine/Source/Programs/Horde/Horde.Build.Tests/Fixture.cs
@@ -122,7 +122,7 @@ namespace Horde.Build.Tests
 				"text/plain", job1ArtifactStream);
 
 			Agent1Name = "testAgent1";
-			Agent1 = await agentService.CreateAgentAsync(Agent1Name, true, null, null);
+			Agent1 = await agentService.CreateAgentAsync(Agent1Name, true, null);
 		}
 
 		private class FixtureGraph : IGraph
diff --git a/Engine/Source/Programs/Horde/Horde.Build.Tests/JobServiceTests.cs b/Engine/Source/Programs/Horde/Horde.Build.Tests/JobServiceTests.cs
--- a/Engine/Source/Programs/Horde/Horde.Build.Tests/JobServiceTests.cs
+++ b/Engine/Source/Programs/Horde/Horde.Build.Tests/JobServiceTests.cs
@@ -203,7 +203,7 @@ namespace Horde.Build.Tests
 		[TestMethod]
 		public async Task TestRunEarly()
 		{
-			IAgent? agent = await AgentService.CreateAgentAsync("TestAgent", true, null, new List<StringId<IPool>> { new StringId<IPool>("win") });
+			IAgent? agent = await AgentService.CreateAgentAsync("TestAgent", true, new List<StringId<IPool>> { new StringId<IPool>("win") });
 			await AgentService.CreateSessionAsync(agent, AgentStatus.Ok, new List<string>(), new Dictionary<string, int>(), null);
 
 			const string ProjectConfigRevision = "projectconfig";
diff --git a/Engine/Source/Programs/Horde/Horde.Build.Tests/JobTaskSourceTests.cs b/Engine/Source/Programs/Horde/Horde.Build.Tests/JobTaskSourceTests.cs
--- a/Engine/Source/Programs/Horde/Horde.Build.Tests/JobTaskSourceTests.cs
+++ b/Engine/Source/Programs/Horde/Horde.Build.Tests/JobTaskSourceTests.cs
@@ -170,7 +170,7 @@ namespace Horde.Build.Tests
 
 			if (shouldCreateAgent)
 			{
-				IAgent? agent = await AgentService.CreateAgentAsync("TestAgent", isAgentEnabled, null, new List<StringId<IPool>> { pool.Id });
+				IAgent? agent = await AgentService.CreateAgentAsync("TestAgent", isAgentEnabled, new List<StringId<IPool>> { pool.Id });
 				await AgentService.CreateSessionAsync(agent, AgentStatus.Ok, new List<string>(), new Dictionary<string, int>(), null);
 			}
 			
diff --git a/Engine/Source/Programs/Horde/Horde.Build.Tests/TestSetup.cs b/Engine/Source/Programs/Horde/Horde.Build.Tests/TestSetup.cs
--- a/Engine/Source/Programs/Horde/Horde.Build.Tests/TestSetup.cs
+++ b/Engine/Source/Programs/Horde/Horde.Build.Tests/TestSetup.cs
@@ -319,7 +319,7 @@ namespace Horde.Build.Tests
 		private static int s_agentIdCounter = 1;
 		public async Task<IAgent> CreateAgentAsync(IPool pool, bool enabled = true)
 		{
-			IAgent agent = await AgentService.CreateAgentAsync("TestAgent" + s_agentIdCounter++, enabled, null, new List<StringId<IPool>> { pool.Id });
+			IAgent agent = await AgentService.CreateAgentAsync("TestAgent" + s_agentIdCounter++, enabled, new List<StringId<IPool>> { pool.Id });
 			agent = await AgentService.CreateSessionAsync(agent, AgentStatus.Ok, new List<string>(), new Dictionary<string, int>(), null);
 			return agent;
 		}
diff --git a/Engine/Source/Programs/Horde/Horde.Build/Agents/AgentCollection.cs b/Engine/Source/Programs/Horde/Horde.Build/Agents/AgentCollection.cs
--- a/Engine/Source/Programs/Horde/Horde.Build/Agents/AgentCollection.cs
+++ b/Engine/Source/Programs/Horde/Horde.Build/Agents/AgentCollection.cs
@@ -10,7 +10,6 @@ using Google.Protobuf.WellKnownTypes;
 using Horde.Build.Acls;
 using Horde.Build.Agents.Pools;
 using Horde.Build.Agents.Sessions;
-using Horde.Build.Agents.Software;
 using Horde.Build.Auditing;
 using Horde.Build.Server;
 using Horde.Build.Utilities;
@@ -22,7 +21,6 @@ using MongoDB.Driver;
 
 namespace Horde.Build.Agents
 {
-	using AgentSoftwareChannelName = StringId<AgentSoftwareChannels>;
 	using PoolId = StringId<IPool>;
 	using SessionId = ObjectId<ISession>;
 
@@ -83,9 +81,6 @@ namespace Horde.Build.Agents
 			public List<string>? Properties { get; set; }
 			public Dictionary<string, int>? Resources { get; set; }
 
-			[BsonIgnoreIfNull]
-			public AgentSoftwareChannelName? Channel { get; set; }
-
 			[BsonIgnoreIfNull]
 			public string? LastUpgradeVersion { get; set; }
 
@@ -135,12 +130,11 @@ namespace Horde.Build.Agents
 			{
 			}
 
-			public AgentDocument(AgentId id, bool enabled, AgentSoftwareChannelName? channel, List<PoolId> pools)
+			public AgentDocument(AgentId id, bool enabled, List<PoolId> pools)
 			{
 				Id = id;
 				Acl = new Acl();
 				Enabled = enabled;
-				Channel = channel;
 				Pools = pools;
 			}
 		}
@@ -165,9 +159,9 @@ namespace Horde.Build.Agents
 		}
 
 		/// <inheritdoc/>
-		public async Task<IAgent> AddAsync(AgentId id, bool enabled, AgentSoftwareChannelName? channel, List<PoolId>? pools)
+		public async Task<IAgent> AddAsync(AgentId id, bool enabled, List<PoolId>? pools)
 		{
-			AgentDocument agent = new AgentDocument(id, enabled, channel, pools ?? new List<PoolId>());
+			AgentDocument agent = new AgentDocument(id, enabled, pools ?? new List<PoolId>());
 			await _agents.InsertOneAsync(agent);
 			return agent;
 		}
@@ -289,7 +283,7 @@ namespace Horde.Build.Agents
 		}
 
 		/// <inheritdoc/>
-		public async Task<IAgent?> TryUpdateSettingsAsync(IAgent agentInterface, bool? enabled = null, bool? requestConform = null, bool? requestFullConform = null, bool? requestRestart = null, bool? requestShutdown = null, string? shutdownReason = null, AgentSoftwareChannelName? channel = null, List<PoolId>? pools = null, Acl? acl = null, string? comment = null)
+		public async Task<IAgent?> TryUpdateSettingsAsync(IAgent agentInterface, bool? enabled = null, bool? requestConform = null, bool? requestFullConform = null, bool? requestRestart = null, bool? requestShutdown = null, string? shutdownReason = null, List<PoolId>? pools = null, Acl? acl = null, string? comment = null)
 		{
 			AgentDocument agent = (AgentDocument)agentInterface;
 
@@ -343,17 +337,6 @@ namespace Horde.Build.Agents
 				updates.Add(updateBuilder.Set(x => x.LastShutdownReason, shutdownReason));
 			}
 
-			if (channel != null)
-			{
-				if (channel.Value == AgentSoftwareService.DefaultChannelName)
-				{
-					updates.Add(updateBuilder.Unset(x => x.Channel));
-				}
-				else
-				{
-					updates.Add(updateBuilder.Set(x => x.Channel, channel));
-				}
-			}
 			if (acl != null)
 			{
 				updates.Add(Acl.CreateUpdate<AgentDocument>(x => x.Acl!, acl));
@@ -367,7 +350,7 @@ namespace Horde.Build.Agents
 			IAgent? newAgent = await TryUpdateAsyncWithRetries(agent, updateBuilder.Combine(updates));
 			if (newAgent != null)
 			{
-				if (newAgent.RequestRestart != agent.RequestRestart || newAgent.RequestConform != agent.RequestConform || newAgent.RequestShutdown != agent.RequestShutdown || newAgent.Channel != agent.Channel)
+				if (newAgent.RequestRestart != agent.RequestRestart || newAgent.RequestConform != agent.RequestConform || newAgent.RequestShutdown != agent.RequestShutdown)
 				{
 					await PublishUpdateEventAsync(agent.Id);
 				}
diff --git a/Engine/Source/Programs/Horde/Horde.Build/Agents/AgentMessages.cs b/Engine/Source/Programs/Horde/Horde.Build/Agents/AgentMessages.cs
--- a/Engine/Source/Programs/Horde/Horde.Build/Agents/AgentMessages.cs
+++ b/Engine/Source/Programs/Horde/Horde.Build/Agents/AgentMessages.cs
@@ -33,11 +33,6 @@ namespace Horde.Build.Agents
 		/// </summary>
 		public bool Ephemeral { get; set; }
 
-		/// <summary>
-		/// Per-agent override for the desired agent software channel
-		/// </summary>
-		public string? Channel { get; set; }
-
 		/// <summary>
 		/// Pools for this agent
 		/// </summary>
@@ -99,11 +94,6 @@ namespace Horde.Build.Agents
 		/// </summary>
 		public bool? RequestShutdown { get; set; }
 
-		/// <summary>
-		/// Per-agent override for the desired agent software channel
-		/// </summary>
-		public string? Channel { get; set; }
-
 		/// <summary>
 		/// Pools for this agent
 		/// </summary>
@@ -404,11 +394,6 @@ namespace Horde.Build.Agents
 		/// </summary>
 		public string? Version { get; set; }
 
-		/// <summary>
-		/// Per-agent override for the desired agent software channel
-		/// </summary>
-		public string? Channel { get; set; }
-
 		/// <summary>
 		/// Properties for the agent
 		/// </summary>
@@ -481,12 +466,7 @@ namespace Horde.Build.Agents
 			NextConformTime =   agent.LastConformTime;
 			ConformAttemptCount = agent.ConformAttemptCount;
 			Version = agent.Version?.ToString() ?? "Unknown";
-			if(agent.Channel != null)
-			{
-				Version += $" ({agent.Channel})";
-			}
 			Version = agent.Version?.ToString();
-			Channel = agent.Channel?.ToString();
 			UpdateTime = agent.UpdateTime;
 			Pools = agent.GetPools().Select(x => x.ToString()).ToList();
 			Workspaces = agent.Workspaces.ConvertAll(x => new GetAgentWorkspaceResponse(x));
diff --git a/Engine/Source/Programs/Horde/Horde.Build/Agents/AgentService.cs b/Engine/Source/Programs/Horde/Horde.Build/Agents/AgentService.cs
--- a/Engine/Source/Programs/Horde/Horde.Build/Agents/AgentService.cs
+++ b/Engine/Source/Programs/Horde/Horde.Build/Agents/AgentService.cs
@@ -15,7 +15,6 @@ using Horde.Build.Acls;
 using Horde.Build.Agents.Leases;
 using Horde.Build.Agents.Pools;
 using Horde.Build.Agents.Sessions;
-using Horde.Build.Agents.Software;
 using Horde.Build.Auditing;
 using Horde.Build.Server;
 using Horde.Build.Tasks;
@@ -31,7 +30,6 @@ using StatsdClient;
 
 namespace Horde.Build.Agents
 {
-	using AgentSoftwareChannelName = StringId<AgentSoftwareChannels>;
 	using LeaseId = ObjectId<ILease>;
 	using PoolId = StringId<IPool>;
 	using SessionId = ObjectId<ISession>;
@@ -181,12 +179,11 @@ namespace Horde.Build.Agents
 		/// </summary>
 		/// <param name="name">Name of the agent</param>
 		/// <param name="enabled">Whether the agent is currently enabled</param>
-		/// <param name="channel">Override for the desired software version</param>
 		/// <param name="pools">Pools for this agent</param>
 		/// <returns>Unique id for the agent</returns>
-		public Task<IAgent> CreateAgentAsync(string name, bool enabled, AgentSoftwareChannelName? channel, List<PoolId>? pools)
+		public Task<IAgent> CreateAgentAsync(string name, bool enabled, List<PoolId>? pools)
 		{
-			return Agents.AddAsync(new AgentId(name), enabled, channel, pools);
+			return Agents.AddAsync(new AgentId(name), enabled, pools);
 		}
 
 		/// <summary>
diff --git a/Engine/Source/Programs/Horde/Horde.Build/Agents/AgentsController.cs b/Engine/Source/Programs/Horde/Horde.Build/Agents/AgentsController.cs
--- a/Engine/Source/Programs/Horde/Horde.Build/Agents/AgentsController.cs
+++ b/Engine/Source/Programs/Horde/Horde.Build/Agents/AgentsController.cs
@@ -58,8 +58,7 @@ namespace Horde.Build.Agents
 				return Forbid(AclAction.CreateAgent);
 			}
 
-			AgentSoftwareChannelName? channel = String.IsNullOrEmpty(request.Channel) ? (AgentSoftwareChannelName?)null : new AgentSoftwareChannelName(request.Channel);
-			IAgent agent = await _agentService.CreateAgentAsync(request.Name, request.Enabled, channel, request.Pools?.ConvertAll(x => new PoolId(x)));
+			IAgent agent = await _agentService.CreateAgentAsync(request.Name, request.Enabled, request.Pools?.ConvertAll(x => new PoolId(x)));
 
 			return new CreateAgentResponse(agent.Id.ToString());
 		}
@@ -143,7 +142,6 @@ namespace Horde.Build.Agents
 		public async Task<ActionResult> UpdateAgentAsync(AgentId agentId, [FromBody] UpdateAgentRequest update)
 		{
 			List<PoolId>? updatePools = update.Pools?.ConvertAll(x => new PoolId(x));
-			AgentSoftwareChannelName? channel = String.IsNullOrEmpty(update.Channel) ? (AgentSoftwareChannelName?)null : new AgentSoftwareChannelName(update.Channel);
 
 			string userName = User.GetUserName() ?? "Unknown";
 
@@ -164,7 +162,7 @@ namespace Horde.Build.Agents
 					return Forbid(AclAction.ChangePermissions, agentId);
 				}
 
-				IAgent? newAgent = await _agentService.Agents.TryUpdateSettingsAsync(agent, update.Enabled, update.RequestConform, update.RequestFullConform, update.RequestRestart, update.RequestShutdown, $"Manual ({userName})", channel, update.Pools?.ConvertAll(x => new PoolId(x)), Acl.Merge(agent.Acl, update.Acl), update.Comment);
+				IAgent? newAgent = await _agentService.Agents.TryUpdateSettingsAsync(agent, update.Enabled, update.RequestConform, update.RequestFullConform, update.RequestRestart, update.RequestShutdown, $"Manual ({userName})", update.Pools?.ConvertAll(x => new PoolId(x)), Acl.Merge(agent.Acl, update.Acl), update.Comment);
 				if (newAgent == null)
 				{
 					continue;
diff --git a/Engine/Source/Programs/Horde/Horde.Build/Agents/IAgent.cs b/Engine/Source/Programs/Horde/Horde.Build/Agents/IAgent.cs
--- a/Engine/Source/Programs/Horde/Horde.Build/Agents/IAgent.cs
+++ b/Engine/Source/Programs/Horde/Horde.Build/Agents/IAgent.cs
@@ -26,7 +26,6 @@ using MongoDB.Bson.Serialization.Attributes;
 
 namespace Horde.Build.Agents
 {
-	using AgentSoftwareChannelName = StringId<AgentSoftwareChannels>;
 	using LeaseId = ObjectId<ILease>;
 	using LogId = ObjectId<ILogFile>;
 	using PoolId = StringId<IPool>;
@@ -404,11 +403,6 @@ namespace Horde.Build.Agents
 		/// </summary>
 		public IReadOnlyDictionary<string, int> Resources { get; }
 
-		/// <summary>
-		/// Channel for the software running on this agent. Uses <see cref="AgentSoftwareService.DefaultChannelName"/> if not specified
-		/// </summary>
-		public AgentSoftwareChannelName? Channel { get; }
-
 		/// <summary>
 		/// Last upgrade that was attempted
 		/// </summary>
diff --git a/Engine/Source/Programs/Horde/Horde.Build/Agents/IAgentCollection.cs b/Engine/Source/Programs/Horde/Horde.Build/Agents/IAgentCollection.cs
--- a/Engine/Source/Programs/Horde/Horde.Build/Agents/IAgentCollection.cs
+++ b/Engine/Source/Programs/Horde/Horde.Build/Agents/IAgentCollection.cs
@@ -6,14 +6,12 @@ using System.Threading.Tasks;
 using Horde.Build.Acls;
 using Horde.Build.Agents.Pools;
 using Horde.Build.Agents.Sessions;
-using Horde.Build.Agents.Software;
 using Horde.Build.Auditing;
 using Horde.Build.Utilities;
 using HordeCommon;
 
 namespace Horde.Build.Agents
 {
-	using AgentSoftwareChannelName = StringId<AgentSoftwareChannels>;
 	using PoolId = StringId<IPool>;
 	using SessionId = ObjectId<ISession>;
 
@@ -27,9 +25,8 @@ namespace Horde.Build.Agents
 		/// </summary>
 		/// <param name="id">Id for the new agent</param>
 		/// <param name="enabled">Whether the agent is enabled or not</param>
-		/// <param name="channel">Channel to use for software run by this agent</param>
 		/// <param name="pools">Pools for the agent</param>
-		Task<IAgent> AddAsync(AgentId id, bool enabled, AgentSoftwareChannelName? channel = null, List<PoolId>? pools = null);
+		Task<IAgent> AddAsync(AgentId id, bool enabled, List<PoolId>? pools = null);
 
 		/// <summary>
 		/// Deletes an agent
@@ -82,12 +79,11 @@ namespace Horde.Build.Agents
 		/// <param name="requestRestart">Whether to request the machine be restarted</param>
 		/// <param name="requestShutdown">Whether to request the machine be shut down</param>
 		/// <param name="shutdownReason">The reason for shutting down agent, ex. Autoscaler/Manual/Unexpected</param>
-		/// <param name="channel">Override for the desired software channel</param>
 		/// <param name="pools">List of pools for the agent</param>
 		/// <param name="acl">New ACL for this agent</param>
 		/// <param name="comment">New comment</param>
 		/// <returns>Version of the software that needs to be installed on the agent. Null if the agent is running the correct version.</returns>
-		Task<IAgent?> TryUpdateSettingsAsync(IAgent agent, bool? enabled = null, bool? requestConform = null, bool? requestFullConform = null, bool? requestRestart = null, bool? requestShutdown = null, string? shutdownReason = null, AgentSoftwareChannelName? channel = null, List<PoolId>? pools = null, Acl? acl = null, string? comment = null);
+		Task<IAgent?> TryUpdateSettingsAsync(IAgent agent, bool? enabled = null, bool? requestConform = null, bool? requestFullConform = null, bool? requestRestart = null, bool? requestShutdown = null, string? shutdownReason = null, List<PoolId>? pools = null, Acl? acl = null, string? comment = null);
 
 		/// <summary>
 		/// Update the current workspaces for an agent.
diff --git a/Engine/Source/Programs/Horde/Horde.Build/Agents/Software/AgentSoftwareService.cs b/Engine/Source/Programs/Horde/Horde.Build/Agents/Software/AgentSoftwareService.cs
--- a/Engine/Source/Programs/Horde/Horde.Build/Agents/Software/AgentSoftwareService.cs
+++ b/Engine/Source/Programs/Horde/Horde.Build/Agents/Software/AgentSoftwareService.cs
@@ -68,11 +68,6 @@ namespace Horde.Build.Agents.Software
 	[SingletonDocument("agent-software-channels", "5f455039d97900f2b6c735a9")]
 	public class AgentSoftwareChannels : SingletonBase
 	{
-		/// <summary>
-		/// The next version number
-		/// </summary>
-		public int NextVersion { get; set; } = 1;
-
 		/// <summary>
 		/// List of channels
 		/// </summary>
@@ -270,6 +265,7 @@ namespace Horde.Build.Agents.Software
 		{
 			return _collection.GetAsync(version);
 		}
+
 		async Task RegisterDefaultAgent(int delayMs)
 		{
 			await Task.Delay(delayMs);
diff --git a/Engine/Source/Programs/Horde/Horde.Build/Server/GlobalConfig.cs b/Engine/Source/Programs/Horde/Horde.Build/Server/GlobalConfig.cs
--- a/Engine/Source/Programs/Horde/Horde.Build/Server/GlobalConfig.cs
+++ b/Engine/Source/Programs/Horde/Horde.Build/Server/GlobalConfig.cs
@@ -9,6 +9,7 @@ using EpicGames.Horde.Common;
 using EpicGames.Perforce;
 using EpicGames.Serialization;
 using Horde.Build.Acls;
+using Horde.Build.Agents.Software;
 using Horde.Build.Projects;
 using Horde.Build.Storage;
 using Horde.Build.Tools;
@@ -17,7 +18,8 @@ using Horde.Build.Utilities;
 namespace Horde.Build.Server
 {
 	using ProjectId = StringId<IProject>;
-	
+	using AgentSoftwareChannelName = StringId<AgentSoftwareChannels>;
+
 	/// <summary>
 	/// Global configuration
 	/// </summary>
@@ -40,6 +42,11 @@ namespace Horde.Build.Server
 		/// </summary>
 		public List<PerforceCluster> PerforceClusters { get; set; } = new List<PerforceCluster>();
 
+		/// <summary>
+		/// List of costs of a particular agent type
+		/// </summary>
+		public List<AgentSoftwareConfig> Software { get; set; } = new List<AgentSoftwareConfig>();
+
 		/// <summary>
 		/// List of costs of a particular agent type
 		/// </summary>
@@ -202,6 +209,22 @@ namespace Horde.Build.Server
 		public List<DevicePlatformConfig> Platforms { get; set; } = new List<DevicePlatformConfig>();
 	}
 
+	/// <summary>
+	/// Selects different agent software versions by evaluating a condition
+	/// </summary>
+	public class AgentSoftwareConfig
+	{
+		/// <summary>
+		/// Channel name
+		/// </summary>
+		public AgentSoftwareChannelName Channel { get; set; }
+
+		/// <summary>
+		/// Condition for using this channel
+		/// </summary>
+		public Condition? Condition { get; set; }
+	}
+
 	/// <summary>
 	/// Describes the monetary cost of agents matching a particular criteris
 	/// </summary>
diff --git a/Engine/Source/Programs/Horde/Horde.Build/Server/RpcService.cs b/Engine/Source/Programs/Horde/Horde.Build/Server/RpcService.cs
--- a/Engine/Source/Programs/Horde/Horde.Build/Server/RpcService.cs
+++ b/Engine/Source/Programs/Horde/Horde.Build/Server/RpcService.cs
@@ -302,7 +302,7 @@ namespace Horde.Build.Server
 					throw new StructuredRpcException(StatusCode.PermissionDenied, "User is not authenticated to create new agents");
 				}
 
-				agent = await _agentService.CreateAgentAsync(request.Name, true, null, null);
+				agent = await _agentService.CreateAgentAsync(request.Name, true, null);
 			}
 
 			// Make sure we're allowed to create sessions on this agent
diff --git a/Engine/Source/Programs/Horde/Horde.Build/ServerSettings.cs b/Engine/Source/Programs/Horde/Horde.Build/ServerSettings.cs
--- a/Engine/Source/Programs/Horde/Horde.Build/ServerSettings.cs
+++ b/Engine/Source/Programs/Horde/Horde.Build/ServerSettings.cs
@@ -1,11 +1,13 @@
 // Copyright Epic Games, Inc. All Rights Reserved.
 
 using System;
+using System.ComponentModel.DataAnnotations;
 using System.Linq;
+using EpicGames.Horde.Common;
 using EpicGames.Horde.Storage;
 using EpicGames.Horde.Storage.Nodes;
 using Horde.Build.Agents.Fleet;
-using Horde.Build.Perforce;
+using Horde.Build.Agents.Software;
 using Horde.Build.Storage.Backends;
 using Horde.Build.Utilities;
 using Serilog.Events;
diff --git a/Engine/Source/Programs/Horde/Horde.Build/Tasks/UpgradeTaskSource.cs b/Engine/Source/Programs/Horde/Horde.Build/Tasks/UpgradeTaskSource.cs
--- a/Engine/Source/Programs/Horde/Horde.Build/Tasks/UpgradeTaskSource.cs
+++ b/Engine/Source/Programs/Horde/Horde.Build/Tasks/UpgradeTaskSource.cs
@@ -10,6 +10,7 @@ using Horde.Build.Agents.Leases;
 using Horde.Build.Agents.Software;
 using Horde.Build.Jobs;
 using Horde.Build.Logs;
+using Horde.Build.Server;
 using Horde.Build.Utilities;
 using HordeCommon;
 using HordeCommon.Rpc.Tasks;
@@ -28,12 +29,14 @@ namespace Horde.Build.Tasks
 		public override TaskSourceFlags Flags => TaskSourceFlags.AllowWhenDisabled | TaskSourceFlags.AllowDuringDowntime;
 
 		readonly AgentSoftwareService _agentSoftwareService;
+		readonly GlobalsService _globalsService;
 		readonly ILogFileService _logService;
 		readonly IClock _clock;
 
-		public UpgradeTaskSource(AgentSoftwareService agentSoftwareService, ILogFileService logService, IClock clock)
+		public UpgradeTaskSource(AgentSoftwareService agentSoftwareService, GlobalsService globalsService, ILogFileService logService, IClock clock)
 		{
 			_agentSoftwareService = agentSoftwareService;
+			_globalsService = globalsService;
 			_logService = logService;
 			_clock = clock;
 
@@ -81,7 +84,19 @@ namespace Horde.Build.Tasks
 		/// <returns>Unique id of the client version this agent should be running</returns>
 		public async Task<string?> GetRequiredSoftwareVersion(IAgent agent)
 		{
-			AgentSoftwareChannelName channelName = agent.Channel ?? AgentSoftwareService.DefaultChannelName;
+			IGlobals globals = await _globalsService.GetAsync();
+			GlobalConfig config = globals.Config;
+
+			AgentSoftwareChannelName channelName = AgentSoftwareService.DefaultChannelName;
+			foreach (AgentSoftwareConfig softwareConfig in config.Software)
+			{
+				if (softwareConfig.Condition != null && agent.SatisfiesCondition(softwareConfig.Condition))
+				{
+					channelName = softwareConfig.Channel;
+					break;
+				}
+			}
+
 			IAgentSoftwareChannel? channel = await _agentSoftwareService.GetCachedChannelAsync(channelName);
 			return channel?.Version;
 		}
