Horde: Allow updating the current software version for a channel separately to the actual data.

diff --git a/Engine/Source/Programs/Horde/Horde.Build/Agents/Software/AgentSoftwareController.Messages.cs b/Engine/Source/Programs/Horde/Horde.Build/Agents/Software/AgentSoftwareController.Messages.cs
--- a/Engine/Source/Programs/Horde/Horde.Build/Agents/Software/AgentSoftwareController.Messages.cs
+++ b/Engine/Source/Programs/Horde/Horde.Build/Agents/Software/AgentSoftwareController.Messages.cs
@@ -4,6 +4,28 @@ using System;
 
 namespace Horde.Build.Agents.Software
 {
+	/// <summary>
+	/// Upload a new version of the agent software
+	/// </summary>
+	public class UploadAgentSoftwareResponse
+	{
+		/// <summary>
+		/// Version number of the uploaded software
+		/// </summary>
+		public string Version { get; set; } = String.Empty;
+	}
+
+	/// <summary>
+	/// Update the current software version
+	/// </summary>
+	public class UpdateAgentSoftwareChannelRequest
+	{
+		/// <summary>
+		/// Version of the software to use
+		/// </summary>
+		public string? Version { get; set; }
+	}
+
 	/// <summary>
 	/// Information about an agent software channel
 	/// </summary>
diff --git a/Engine/Source/Programs/Horde/Horde.Build/Agents/Software/AgentSoftwareController.cs b/Engine/Source/Programs/Horde/Horde.Build/Agents/Software/AgentSoftwareController.cs
--- a/Engine/Source/Programs/Horde/Horde.Build/Agents/Software/AgentSoftwareController.cs
+++ b/Engine/Source/Programs/Horde/Horde.Build/Agents/Software/AgentSoftwareController.cs
@@ -93,6 +93,25 @@ namespace Horde.Build.Agents.Software
 			return new GetAgentSoftwareChannelResponse(channel).ApplyFilter(filter);
 		}
 
+		/// <summary>
+		/// Finds all uploaded software matching the given criteria
+		/// </summary>
+		/// <param name="name">Name of the channel to get</param>
+		/// <param name="filter">Filter for the properties to return</param>
+		/// <returns>Http response</returns>
+		[HttpPut]
+		[Route("/api/v1/agentsoftware/{name}")]
+		public async Task<ActionResult<object>> SetSoftwareAsync(string name, [FromBody] UpdateAgentSoftwareChannelRequest request, [FromQuery] PropertyFilter? filter = null)
+		{
+			if (!await _aclService.AuthorizeAsync(AclAction.UploadSoftware, User))
+			{
+				return Forbid();
+			}
+
+			await _agentSoftwareService.UpdateChannelAsync(new AgentSoftwareChannelName(name), User.GetUserName(), request.Version);
+			return Ok();
+		}
+
 		/// <summary>
 		/// Uploads a new agent zip file
 		/// </summary>
@@ -101,7 +120,7 @@ namespace Horde.Build.Agents.Software
 		/// <returns>Http result code</returns>
 		[HttpPost]
 		[Route("/api/v1/agentsoftware/{name}/zip")]
-		public async Task<ActionResult> SetArchiveAsync(string name, [FromForm] IFormFile file)
+		public async Task<ActionResult<UploadAgentSoftwareResponse>> UploadArchiveAsync(string name, [FromForm] IFormFile file)
 		{
 			if (!await _aclService.AuthorizeAsync(AclAction.UploadSoftware, User))
 			{
@@ -118,8 +137,10 @@ namespace Horde.Build.Agents.Software
 				data = memoryStream.ToArray();
 			}
 
-			await _agentSoftwareService.SetArchiveAsync(new AgentSoftwareChannelName(name), User.Identity?.Name, data);
-			return Ok();
+			string version = await _agentSoftwareService.UploadArchiveAsync(data);
+			await _agentSoftwareService.UpdateChannelAsync(new AgentSoftwareChannelName(name), User.Identity?.Name, version);
+
+			return new UploadAgentSoftwareResponse { Version = version };
 		}
 
 		/// <summary>
diff --git a/Engine/Source/Programs/Horde/Horde.Build/Agents/Software/AgentSoftwareService.cs b/Engine/Source/Programs/Horde/Horde.Build/Agents/Software/AgentSoftwareService.cs
--- a/Engine/Source/Programs/Horde/Horde.Build/Agents/Software/AgentSoftwareService.cs
+++ b/Engine/Source/Programs/Horde/Horde.Build/Agents/Software/AgentSoftwareService.cs
@@ -208,17 +208,24 @@ namespace Horde.Build.Agents.Software
 		/// <summary>
 		/// Updates a new software revision
 		/// </summary>
-		/// <param name="name">Name of the channel</param>
-		/// <param name="author">Name of the user uploading this file</param>
 		/// <param name="data">The input data stream. This should be a zip archive containing the HordeAgent executable.</param>
-		/// <returns>Unique id for the file</returns>
-		public async Task<string> SetArchiveAsync(AgentSoftwareChannelName name, string? author, byte[] data)
+		/// <returns>Version number for the archive</returns>
+		public async Task<string> UploadArchiveAsync(byte[] data)
 		{
 			// Upload the software
 			string version = AgentUtilities.ReadVersion(data);
 			await _collection.AddAsync(version, data);
+			return version;
+		}
 
-			// Update the channel
+		/// <summary>
+		/// Updates the software being used by a particular channel
+		/// </summary>
+		/// <param name="name">Name of the channel to update</param>
+		/// <param name="author">Name of the user updating the software</param>
+		/// <param name="version">New version of the software to use</param>
+		public async Task UpdateChannelAsync(AgentSoftwareChannelName name, string? author, string version)
+		{
 			for(; ;)
 			{
 				AgentSoftwareChannels instance = await _singleton.GetAsync();
@@ -233,7 +240,6 @@ namespace Horde.Build.Agents.Software
 					break;
 				}
 			}
-			return version;
 		}
 
 		/// <summary>
@@ -320,7 +326,8 @@ namespace Horde.Build.Agents.Software
 
 			bytes = AgentUtilities.UpdateAppSettings(bytes, agentSettings);
 
-			string version = await SetArchiveAsync(new AgentSoftwareChannelName("default"), null, bytes);
+			string version = await UploadArchiveAsync(bytes);
+			await UpdateChannelAsync(new AgentSoftwareChannelName("default"), null, version);
 
 			if (!DirectoryReference.Exists(agentHashFile.Directory))
 			{
diff --git a/Engine/Source/Programs/Horde/Horde.Build/Server/RpcService.cs b/Engine/Source/Programs/Horde/Horde.Build/Server/RpcService.cs
--- a/Engine/Source/Programs/Horde/Horde.Build/Server/RpcService.cs
+++ b/Engine/Source/Programs/Horde/Horde.Build/Server/RpcService.cs
@@ -956,7 +956,11 @@ namespace Horde.Build.Server
 				throw new StructuredRpcException(StatusCode.PermissionDenied, "Access to software is forbidden");
 			}
 
-			string version = await _agentSoftwareService.SetArchiveAsync(new AgentSoftwareChannelName(request.Channel), null, request.Data.ToArray());
+			string version = await _agentSoftwareService.UploadArchiveAsync(request.Data.ToArray());
+			foreach (string channel in request.Channel.Split('+', ';'))
+			{
+				await _agentSoftwareService.UpdateChannelAsync(new AgentSoftwareChannelName(channel), null, version);
+			}
 
 			UploadSoftwareResponse response = new UploadSoftwareResponse();
 			response.Version = version;
