Horde: Fix warnings.

diff --git a/Engine/Source/Programs/Horde/Horde.Build/Agents/Software/AgentSoftwareController.cs b/Engine/Source/Programs/Horde/Horde.Build/Agents/Software/AgentSoftwareController.cs
--- a/Engine/Source/Programs/Horde/Horde.Build/Agents/Software/AgentSoftwareController.cs
+++ b/Engine/Source/Programs/Horde/Horde.Build/Agents/Software/AgentSoftwareController.cs
@@ -97,6 +97,7 @@ namespace Horde.Build.Agents.Software
 		/// Finds all uploaded software matching the given criteria
 		/// </summary>
 		/// <param name="name">Name of the channel to get</param>
+		/// <param name="request">Request body</param>
 		/// <param name="filter">Filter for the properties to return</param>
 		/// <returns>Http response</returns>
 		[HttpPut]
diff --git a/Engine/Source/Programs/Horde/Horde.Build/Agents/Software/AgentSoftwareService.cs b/Engine/Source/Programs/Horde/Horde.Build/Agents/Software/AgentSoftwareService.cs
--- a/Engine/Source/Programs/Horde/Horde.Build/Agents/Software/AgentSoftwareService.cs
+++ b/Engine/Source/Programs/Horde/Horde.Build/Agents/Software/AgentSoftwareService.cs
@@ -224,20 +224,23 @@ namespace Horde.Build.Agents.Software
 		/// <param name="name">Name of the channel to update</param>
 		/// <param name="author">Name of the user updating the software</param>
 		/// <param name="version">New version of the software to use</param>
-		public async Task UpdateChannelAsync(AgentSoftwareChannelName name, string? author, string version)
+		public async Task UpdateChannelAsync(AgentSoftwareChannelName name, string? author, string? version)
 		{
-			for(; ;)
+			if (version != null)
 			{
-				AgentSoftwareChannels instance = await _singleton.GetAsync();
+				for (; ; )
+				{
+					AgentSoftwareChannels instance = await _singleton.GetAsync();
 
-				AgentSoftwareChannel channel = instance.FindOrAddChannel(name);
-				channel.ModifiedBy = author;
-				channel.ModifiedTime = DateTime.UtcNow;
-				channel.Version = version;
+					AgentSoftwareChannel channel = instance.FindOrAddChannel(name);
+					channel.ModifiedBy = author;
+					channel.ModifiedTime = DateTime.UtcNow;
+					channel.Version = version;
 
-				if (await _singleton.TryUpdateAsync(instance))
-				{
-					break;
+					if (await _singleton.TryUpdateAsync(instance))
+					{
+						break;
+					}
 				}
 			}
 		}
