Horde: Additional logging for Slack socket management

Also fix issue escalation being tied to having a socket token set up.

diff --git a/Engine/Source/Programs/Horde/Horde.Build/Notifications/Sinks/SlackNotificationSink.cs b/Engine/Source/Programs/Horde/Horde.Build/Notifications/Sinks/SlackNotificationSink.cs
--- a/Engine/Source/Programs/Horde/Horde.Build/Notifications/Sinks/SlackNotificationSink.cs
+++ b/Engine/Source/Programs/Horde/Horde.Build/Notifications/Sinks/SlackNotificationSink.cs
@@ -13,6 +13,7 @@ using System.Text.Json.Serialization;
 using System.Text.RegularExpressions;
 using System.Threading;
 using System.Threading.Tasks;
+using System.Timers;
 using EpicGames.Core;
 using EpicGames.Redis;
 using EpicGames.Redis.Utility;
@@ -219,6 +220,7 @@ namespace Horde.Build.Notifications.Sinks
 		readonly HashSet<string>? _allowUsers;
 		readonly IExternalIssueService _externalIssueService;
 		readonly JsonSerializerOptions _jsonSerializerOptions;
+		readonly ITicker _escalateTicker;
 		readonly RedisSortedSet<int> _escalateIssues;
 		readonly IClock _clock;
 		readonly ILogger _logger;
@@ -254,6 +256,8 @@ namespace Horde.Build.Notifications.Sinks
 			_clock = clock;
 			_logger = logger;
 
+			_escalateTicker = clock.AddSharedTicker<SlackNotificationSink>(TimeSpan.FromMinutes(1.0), EscalateAsync, _logger);
+
 			_issueQueueTicker = clock.AddSharedTicker("slack-issues", TimeSpan.FromMinutes(1.0), ProcessIssueQueueAsync, _logger);
 			_redisIssueQueue = new RedisList<int>(redisService.ConnectionPool, "slack/issue-queue");
 			_redisIssueLockPrefix = "slack/issues/";
@@ -284,11 +288,13 @@ namespace Horde.Build.Notifications.Sinks
 		public override async Task StartAsync(CancellationToken cancellationToken)
 		{
 			await _issueQueueTicker.StartAsync();
+			await _escalateTicker.StartAsync();
 		}
 
 		/// <inheritdoc/>
 		public override async Task StopAsync(CancellationToken cancellationToken)
 		{
+			await _escalateTicker.StopAsync();
 			await _issueQueueTicker.StopAsync();
 		}
 
@@ -2260,39 +2266,33 @@ namespace Horde.Build.Notifications.Sinks
 		/// <inheritdoc/>
 		protected override async Task ExecuteAsync(CancellationToken stoppingToken)
 		{
-			if (!String.IsNullOrEmpty(_settings.SlackSocketToken))
+			if (String.IsNullOrEmpty(_settings.SlackSocketToken))
 			{
-				using ITicker ticker = _clock.AddSharedTicker<SlackNotificationSink>(TimeSpan.FromMinutes(1.0), EscalateAsync, _logger);
-				await ticker.StartAsync();
+				_logger.LogInformation("No Slack socket token configured; will not be able to respond to interactive messages.");
+				return;
+			}
 
+			while (!stoppingToken.IsCancellationRequested)
+			{
 				try
 				{
-					while (!stoppingToken.IsCancellationRequested)
+					Uri? webSocketUrl = await GetWebSocketUrlAsync(stoppingToken);
+					if (webSocketUrl == null)
 					{
-						try
-						{
-							Uri? webSocketUrl = await GetWebSocketUrlAsync(stoppingToken);
-							if (webSocketUrl == null)
-							{
-								await Task.Delay(TimeSpan.FromSeconds(5.0), stoppingToken);
-								continue;
-							}
-							await HandleSocketAsync(webSocketUrl, stoppingToken);
-						}
-						catch (OperationCanceledException)
-						{
-							break;
-						}
-						catch (Exception ex)
-						{
-							_logger.LogError(ex, "Exception while updating Slack socket");
-							await Task.Delay(TimeSpan.FromSeconds(5.0), stoppingToken);
-						}
+						_logger.LogWarning("Unable to get Slack websocket URL. Pausing before retry.");
+						await Task.Delay(TimeSpan.FromSeconds(5.0), stoppingToken);
+						continue;
 					}
+					await HandleSocketAsync(webSocketUrl, stoppingToken);
 				}
-				finally
+				catch (OperationCanceledException)
 				{
-					await ticker.StopAsync();
+					break;
+				}
+				catch (Exception ex)
+				{
+					_logger.LogError(ex, "Exception while updating Slack socket");
+					await Task.Delay(TimeSpan.FromSeconds(5.0), stoppingToken);
 				}
 			}
 		}
@@ -2319,6 +2319,7 @@ namespace Horde.Build.Notifications.Sinks
 				return null;
 			}
 
+			_logger.LogInformation("Using Slack websocket url: {Url}", socketResponse.Url);
 			return socketResponse.Url;
 		}
 
@@ -2330,6 +2331,8 @@ namespace Horde.Build.Notifications.Sinks
 		/// <returns></returns>
 		private async Task HandleSocketAsync(Uri socketUrl, CancellationToken stoppingToken)
 		{
+			_logger.LogInformation("Opening Slack socket {Url}", socketUrl);
+
 			using ClientWebSocket socket = new ClientWebSocket();
 			await socket.ConnectAsync(socketUrl, stoppingToken);
 
