Horde: Fix Slack interactions not working correctly. Overloaded StartAsync()/StopAsync() methods in BackgroundService were not calling base implementation, preventing background task from being created.

diff --git a/Engine/Source/Programs/Horde/Horde.Build/Notifications/Sinks/SlackNotificationSink.cs b/Engine/Source/Programs/Horde/Horde.Build/Notifications/Sinks/SlackNotificationSink.cs
--- a/Engine/Source/Programs/Horde/Horde.Build/Notifications/Sinks/SlackNotificationSink.cs
+++ b/Engine/Source/Programs/Horde/Horde.Build/Notifications/Sinks/SlackNotificationSink.cs
@@ -282,11 +282,15 @@ namespace Horde.Build.Notifications.Sinks
 
 			_userCache.Dispose();
 			_httpClient.Dispose();
+			_issueQueueTicker.Dispose();
+			_escalateTicker.Dispose();
 		}
 
 		/// <inheritdoc/>
 		public override async Task StartAsync(CancellationToken cancellationToken)
 		{
+			base.StartAsync();
+
 			await _issueQueueTicker.StartAsync();
 			await _escalateTicker.StartAsync();
 		}
@@ -294,6 +298,8 @@ namespace Horde.Build.Notifications.Sinks
 		/// <inheritdoc/>
 		public override async Task StopAsync(CancellationToken cancellationToken)
 		{
+			base.StopAsync();
+
 			await _escalateTicker.StopAsync();
 			await _issueQueueTicker.StopAsync();
 		}
