Horde: Fix build error.

diff --git a/Engine/Source/Programs/Horde/Horde.Build/Notifications/Sinks/SlackNotificationSink.cs b/Engine/Source/Programs/Horde/Horde.Build/Notifications/Sinks/SlackNotificationSink.cs
--- a/Engine/Source/Programs/Horde/Horde.Build/Notifications/Sinks/SlackNotificationSink.cs
+++ b/Engine/Source/Programs/Horde/Horde.Build/Notifications/Sinks/SlackNotificationSink.cs
@@ -289,7 +289,7 @@ namespace Horde.Build.Notifications.Sinks
 		/// <inheritdoc/>
 		public override async Task StartAsync(CancellationToken cancellationToken)
 		{
-			base.StartAsync();
+			base.StartAsync(cancellationToken);
 
 			await _issueQueueTicker.StartAsync();
 			await _escalateTicker.StartAsync();
@@ -298,7 +298,7 @@ namespace Horde.Build.Notifications.Sinks
 		/// <inheritdoc/>
 		public override async Task StopAsync(CancellationToken cancellationToken)
 		{
-			base.StopAsync();
+			base.StopAsync(cancellationToken);
 
 			await _escalateTicker.StopAsync();
 			await _issueQueueTicker.StopAsync();
