Horde: Add user names to issue audit logs
#rnx

diff --git a/Engine/Source/Programs/Horde/Horde.Build/Issues/IssueCollection.cs b/Engine/Source/Programs/Horde/Horde.Build/Issues/IssueCollection.cs
--- a/Engine/Source/Programs/Horde/Horde.Build/Issues/IssueCollection.cs
+++ b/Engine/Source/Programs/Horde/Horde.Build/Issues/IssueCollection.cs
@@ -520,8 +520,60 @@ namespace Horde.Build.Issues
 			return newIssue;
 		}
 
-		static void LogIssueChanges(ILogger issueLogger, Issue oldIssue, Issue newIssue)
+		static void GetRelevantUserIds(List<IIssue> issues, HashSet<UserId> ids)
 		{
+			foreach (IIssue issue in issues)
+			{
+				if (issue.OwnerId != null)
+				{
+					ids.Add(issue.OwnerId.Value);
+				}
+				if (issue.NominatedById != null)
+				{
+					ids.Add(issue.NominatedById.Value);
+				}
+				if (issue.ResolvedById != null)
+				{
+					ids.Add(issue.ResolvedById.Value);
+				}
+				if (issue.QuarantinedByUserId != null)
+				{
+					ids.Add(issue.QuarantinedByUserId.Value);
+				}
+				if (issue.ForceClosedByUserId != null)
+				{
+					ids.Add(issue.ForceClosedByUserId.Value);
+				}
+			}
+		}
+
+		async void LogIssueChanges(ILogger issueLogger, Issue oldIssue, Issue newIssue)
+		{
+			HashSet<UserId> ids = new HashSet<UserId>();
+			GetRelevantUserIds(new List<IIssue>() { oldIssue, newIssue }, ids);
+
+			Dictionary<UserId, string> userNames = new Dictionary<UserId, string>() { [IIssue.ResolvedByUnknownId] = "Horde", [IIssue.ResolvedByTimeoutId] = "Horde"};
+
+			if (ids.Count > 0)
+			{
+				List<IUser> users = await _userCollection.FindUsersAsync(ids);
+				users.ForEach(x => { userNames[x.Id] = x.Name; });
+			}
+
+			string GetUserName(UserId? UserId)
+			{
+				if (UserId == null)
+				{
+					return "null";
+				}
+				string? name;
+				if (userNames.TryGetValue(UserId.Value, out name))
+				{
+					return name;
+				}
+				return "Unknown User";
+			}
+
 			if (newIssue.Severity != oldIssue.Severity)
 			{
 				issueLogger.LogInformation("Changed severity to {Severity}", newIssue.Severity);
@@ -541,23 +593,23 @@ namespace Horde.Build.Issues
 			if (newIssue.OwnerId != oldIssue.OwnerId)
 			{
 				if (newIssue.NominatedById != null)
-				{
-					issueLogger.LogInformation("User {UserId} was nominated by {NominatedByUserId}", newIssue.OwnerId, newIssue.NominatedById);
+				{					
+					issueLogger.LogInformation("User {UserName} ({UserId}) was nominated by {NominatedByUserName} ({NominatedByUserId})", GetUserName(newIssue.OwnerId), newIssue.OwnerId, GetUserName(newIssue.NominatedById), newIssue.NominatedById);
 				}
 				else
 				{
-					issueLogger.LogInformation("User {UserId} was nominated by default", newIssue.OwnerId);
+					issueLogger.LogInformation("User {UserName} ({UserId}) was nominated by default", GetUserName(newIssue.OwnerId), newIssue.OwnerId);
 				}
 			}
 			if (newIssue.AcknowledgedAt != oldIssue.AcknowledgedAt)
 			{
 				if (newIssue.AcknowledgedAt == null)
 				{
-					issueLogger.LogInformation("Issue was un-acknowledged by {UserId}", oldIssue.OwnerId);
+					issueLogger.LogInformation("Issue was un-acknowledged by {UserName} ({UserId})", GetUserName(oldIssue.OwnerId), oldIssue.OwnerId);
 				}
 				else
 				{
-					issueLogger.LogInformation("Issue was acknowledged by {UserId}", newIssue.OwnerId);
+					issueLogger.LogInformation("Issue was acknowledged by {UserName} ({UserId})", GetUserName(newIssue.OwnerId), newIssue.OwnerId);
 				}
 			}
 			if (newIssue.FixChange != oldIssue.FixChange)
@@ -579,7 +631,7 @@ namespace Horde.Build.Issues
 				}
 				else
 				{
-					issueLogger.LogInformation("Resolved by {UserId}", newIssue.ResolvedById);
+					issueLogger.LogInformation("Resolved by {UserName} ({UserId})", GetUserName(newIssue.ResolvedById), newIssue.ResolvedById);
 				}
 			}
 
@@ -623,7 +675,7 @@ namespace Horde.Build.Issues
 			{
 				if (newIssue.QuarantinedByUserId != null)
 				{
-					issueLogger.LogInformation("Quarantined by {UserId}", newIssue.QuarantinedByUserId);
+					issueLogger.LogInformation("Quarantined by {UserName} ({UserId})", GetUserName(newIssue.QuarantinedByUserId), newIssue.QuarantinedByUserId);
 				}
 				else
 				{
@@ -635,7 +687,7 @@ namespace Horde.Build.Issues
 			{
 				if (newIssue.ForceClosedByUserId != null)
 				{
-					issueLogger.LogInformation("Forced closed by {UserId}", newIssue.ForceClosedByUserId);
+					issueLogger.LogInformation("Forced closed by {UserName} ({UserId})", GetUserName(newIssue.ForceClosedByUserId), newIssue.ForceClosedByUserId);
 				}
 				else
 				{
