Horde: Add new AWS fleet manager more resilient to insufficient capacity errors

Doesn't replace the old AwsReuseFleetManager but adds a new clone of the old one called AwsRecylingFleetManager. The old one will be retired and removed but until the new one is proven they can run side-by-side. Pools individually need to opt into "AwsRecycle" through config to benefit.

diff --git a/Engine/Source/Programs/Horde/Horde.Build.Tests/Fleet/AwsRecyclingFleetManagerTest.cs b/Engine/Source/Programs/Horde/Horde.Build.Tests/Fleet/AwsRecyclingFleetManagerTest.cs
new file mode 100644
--- /dev/null
+++ b/Engine/Source/Programs/Horde/Horde.Build.Tests/Fleet/AwsRecyclingFleetManagerTest.cs
@@ -0,0 +1,192 @@
+// Copyright Epic Games, Inc. All Rights Reserved.
+
+using System;
+using System.Collections.Generic;
+using System.Linq;
+using System.Threading;
+using System.Threading.Tasks;
+using Amazon.EC2;
+using Amazon.EC2.Model;
+using Horde.Build.Agents.Pools;
+using Horde.Build.Agents.Fleet;
+using Microsoft.VisualStudio.TestTools.UnitTesting;
+using Horde.Build.Agents;
+using Horde.Build.Agents.Fleet.Providers;
+using Microsoft.Extensions.Logging;
+using JsonSerializer = System.Text.Json.JsonSerializer;
+
+namespace Horde.Build.Tests.Fleet
+{
+	[TestClass]
+	public class AwsRecyclingFleetManagerTest : TestSetup
+	{
+		[TestMethod]
+		public async Task ExpandOneAzPartialCapacityWithFallback()
+		{
+			FakeAmazonEc2 ec2 = new ();
+			Instance i1 = ec2.AddInstance(FakeAmazonEc2.StateStopped, InstanceType.M5Large, FakeAmazonEc2.AzUsEast1A);
+			Instance i2 = ec2.AddInstance(FakeAmazonEc2.StateStopped, InstanceType.M5Large, FakeAmazonEc2.AzUsEast1A);
+			ec2.SetCapacity(FakeAmazonEc2.AzUsEast1A, InstanceType.M5Large, 1);
+			ec2.SetCapacity(FakeAmazonEc2.AzUsEast1A, InstanceType.C5Large, 2);
+
+			await ExpandPoolAsync(ec2.Get(), 2, new(new List<string> { InstanceType.C5Large }));
+			Assert.AreEqual(0, ec2.GetStoppedInstanceCount());
+			Assert.AreEqual(2, ec2.GetPendingInstanceCount());
+			Assert.AreEqual(InstanceType.C5Large, ec2.GetInstance(i1.InstanceId)!.InstanceType);
+			Assert.AreEqual(InstanceType.C5Large, ec2.GetInstance(i2.InstanceId)!.InstanceType);
+		}
+		
+		[TestMethod]
+		public async Task ExpandOneAzPartialCapacityWithMultipleFallbacks()
+		{
+			FakeAmazonEc2 ec2 = new ();
+			Instance i1 = ec2.AddInstance(FakeAmazonEc2.StateStopped, InstanceType.M5Large, FakeAmazonEc2.AzUsEast1A);
+			Instance i2 = ec2.AddInstance(FakeAmazonEc2.StateStopped, InstanceType.M5Large, FakeAmazonEc2.AzUsEast1A);
+			ec2.SetCapacity(FakeAmazonEc2.AzUsEast1A, InstanceType.M5Large, 1);
+			ec2.SetCapacity(FakeAmazonEc2.AzUsEast1A, InstanceType.C5Large, 2);
+
+			await ExpandPoolAsync(ec2.Get(), 2, new(new List<string> { InstanceType.M5Xlarge, InstanceType.R5Large, InstanceType.C5Large }));
+			Assert.AreEqual(0, ec2.GetStoppedInstanceCount());
+			Assert.AreEqual(2, ec2.GetPendingInstanceCount());
+			Assert.AreEqual(InstanceType.C5Large, ec2.GetInstance(i1.InstanceId)!.InstanceType);
+			Assert.AreEqual(InstanceType.C5Large, ec2.GetInstance(i2.InstanceId)!.InstanceType);
+		}
+		
+		[TestMethod]
+		public async Task ExpandOneAzPartialCapacityNoFallback()
+		{
+			FakeAmazonEc2 ec2 = new ();
+			ec2.AddInstance(FakeAmazonEc2.StateStopped, InstanceType.M5Large, FakeAmazonEc2.AzUsEast1A);
+			ec2.AddInstance(FakeAmazonEc2.StateStopped, InstanceType.M5Large, FakeAmazonEc2.AzUsEast1A);
+			ec2.SetCapacity(FakeAmazonEc2.AzUsEast1A, InstanceType.M5Large, 1);
+
+			await ExpandPoolAsync(ec2.Get(), 2, new());
+			Assert.AreEqual(2, ec2.GetStoppedInstanceCount());
+			Assert.AreEqual(0, ec2.GetPendingInstanceCount());
+		}
+		
+		[TestMethod]
+		public async Task ExpandTwoAzsPartialCapacityNoFallback()
+		{
+			FakeAmazonEc2 ec2 = new ();
+			ec2.AddInstance(FakeAmazonEc2.StateStopped, InstanceType.M5Large, FakeAmazonEc2.AzUsEast1A);
+			ec2.AddInstance(FakeAmazonEc2.StateStopped, InstanceType.M5Large, FakeAmazonEc2.AzUsEast1B);
+			ec2.SetCapacity(FakeAmazonEc2.AzUsEast1A, InstanceType.M5Large, 1);
+			ec2.SetCapacity(FakeAmazonEc2.AzUsEast1B, InstanceType.M5Large, 0);
+
+			await ExpandPoolAsync(ec2.Get(), 2, new());
+			Assert.AreEqual(1, ec2.GetStoppedInstanceCount());
+			Assert.AreEqual(1, ec2.GetPendingInstanceCount());
+		}
+		
+		[TestMethod]
+		public async Task ExpandMoreThanNumStoppedInstances()
+		{
+			FakeAmazonEc2 ec2 = new ();
+			ec2.AddInstance(FakeAmazonEc2.StateStopped, InstanceType.M5Large, FakeAmazonEc2.AzUsEast1A);
+			ec2.AddInstance(FakeAmazonEc2.StateStopped, InstanceType.M5Large, FakeAmazonEc2.AzUsEast1B);
+			ec2.SetCapacity(FakeAmazonEc2.AzUsEast1A, InstanceType.M5Large, 1);
+			ec2.SetCapacity(FakeAmazonEc2.AzUsEast1B, InstanceType.M5Large, 1);
+
+			Assert.AreEqual(2, ec2.GetStoppedInstanceCount());
+			await ExpandPoolAsync(ec2.Get(), 10, new());
+			Assert.AreEqual(0, ec2.GetStoppedInstanceCount());
+			Assert.AreEqual(2, ec2.GetPendingInstanceCount());
+			
+			// TODO: Verify remainder is reported correctly
+		}
+		
+		[TestMethod]
+		public async Task ExpandTwoAzsFullCapacityNoFallback()
+		{
+			FakeAmazonEc2 ec2 = new ();
+			ec2.AddInstance(FakeAmazonEc2.StateStopped, InstanceType.M5Large, FakeAmazonEc2.AzUsEast1A);
+			ec2.AddInstance(FakeAmazonEc2.StateStopped, InstanceType.M5Large, FakeAmazonEc2.AzUsEast1B);
+			ec2.SetCapacity(FakeAmazonEc2.AzUsEast1A, InstanceType.M5Large, 1);
+			ec2.SetCapacity(FakeAmazonEc2.AzUsEast1B, InstanceType.M5Large, 1);
+
+			await ExpandPoolAsync(ec2.Get(), 2, new());
+			Assert.AreEqual(0, ec2.GetStoppedInstanceCount());
+			Assert.AreEqual(2, ec2.GetPendingInstanceCount());
+		}
+		
+		[TestMethod]
+		public async Task ExpandTwoAzsFullCapacityEmptyFallback()
+		{
+			FakeAmazonEc2 ec2 = new ();
+			ec2.AddInstance(FakeAmazonEc2.StateStopped, InstanceType.M5Large, FakeAmazonEc2.AzUsEast1A);
+			ec2.AddInstance(FakeAmazonEc2.StateStopped, InstanceType.M5Large, FakeAmazonEc2.AzUsEast1B);
+			ec2.SetCapacity(FakeAmazonEc2.AzUsEast1A, InstanceType.M5Large, 1);
+			ec2.SetCapacity(FakeAmazonEc2.AzUsEast1B, InstanceType.M5Large, 1);
+
+			await ExpandPoolAsync(ec2.Get(), 2, new(new List<string>()));
+			Assert.AreEqual(0, ec2.GetStoppedInstanceCount());
+			Assert.AreEqual(2, ec2.GetPendingInstanceCount());
+		}
+		
+		[TestMethod]
+		public async Task ExpandLaunchesLastUsedInstancesFirst()
+		{
+			FakeAmazonEc2 ec2 = new ();
+			Instance i1 = ec2.AddInstance(FakeAmazonEc2.StateStopped, InstanceType.M5Large, launchTime: DateTime.Parse("Aug 18, 2018"));
+			Instance i2 = ec2.AddInstance(FakeAmazonEc2.StateStopped, InstanceType.M5Large, launchTime: DateTime.Parse("Aug 5, 2018"));
+			Instance i3 = ec2.AddInstance(FakeAmazonEc2.StateStopped, InstanceType.M5Large, launchTime: DateTime.Parse("Aug 22, 2018"));
+			Instance i4 = ec2.AddInstance(FakeAmazonEc2.StateStopped, InstanceType.M5Large, launchTime: DateTime.Parse("Aug 1, 2018"));
+			Instance i5 = ec2.AddInstance(FakeAmazonEc2.StateStopped, InstanceType.M5Large, launchTime: DateTime.Parse("Aug 28, 2018"));
+			ec2.SetCapacity(FakeAmazonEc2.AzUsEast1A, InstanceType.M5Large, 2);
+
+			await ExpandPoolAsync(ec2.Get(), 2, new(new List<string>()));
+			Assert.AreEqual(3, ec2.GetStoppedInstanceCount());
+			Assert.AreEqual(2, ec2.GetPendingInstanceCount());
+			Assert.AreEqual(FakeAmazonEc2.StatePending.Name, ec2.GetInstance(i2.InstanceId)!.State.Name);
+			Assert.AreEqual(FakeAmazonEc2.StatePending.Name, ec2.GetInstance(i4.InstanceId)!.State.Name);
+		}
+
+		[TestMethod]
+		public void DistributesInstanceRequestsEvenly()
+		{
+			AssertInstanceCounts(0, new () { ["a"] = 3, ["b"] = 4 }, new() { ["a"] = 0, ["b"] = 0 }, 0);
+			AssertInstanceCounts(0, new () { }, new() { }, 0);
+			AssertInstanceCounts(1, new () { ["a"] = 1, ["b"] = 1, ["c"] = 1 }, new() { ["a"] = 1, ["b"] = 0, ["c"] = 0 }, 0);
+			AssertInstanceCounts(2, new () { ["a"] = 1, ["b"] = 1, ["c"] = 1 }, new() { ["a"] = 1, ["b"] = 1, ["c"] = 0 }, 0);
+			AssertInstanceCounts(5, new () { ["a"] = 10 }, new() { ["a"] = 5 }, 0);
+			AssertInstanceCounts(6, new () { ["a"] = 10, ["b"] = 2 }, new() { ["a"] = 4, ["b"] = 2 }, 0);
+			AssertInstanceCounts(3, new () { ["a"] = 1, ["b"] = 1, ["c"] = 1 }, new() { ["a"] = 1, ["b"] = 1, ["c"] = 1 }, 0);
+			
+			// Remainder > 0
+			AssertInstanceCounts(10, new () { ["a"] = 3, ["b"] = 4 }, new() { ["a"] = 3, ["b"] = 4 }, 3);
+			AssertInstanceCounts(10, new () { }, new() { }, 10);
+		}
+
+		private static void AssertInstanceCounts(int requestCount, Dictionary<string, int> availableCountPerAz, Dictionary<string, int> expected, int expectedRemainder)
+		{
+			Dictionary<string, int> actual = AwsRecyclingFleetManager.DistributeInstanceRequestsOverAzs(
+				requestCount, availableCountPerAz, out int actualRemainder, false);
+			try
+			{
+				CollectionAssert.AreEquivalent(expected.ToList(), actual.ToList());
+				Assert.AreEqual(expectedRemainder, actualRemainder);
+			}
+			catch (Exception)
+			{
+				Console.WriteLine("Expected=" + JsonSerializer.Serialize(expected));
+				Console.WriteLine("  Actual=" + JsonSerializer.Serialize(actual));
+				throw;
+			}
+		}
+
+		private async Task ExpandPoolAsync(IAmazonEC2 ec2, int numRequestedInstances, AwsRecyclingFleetManagerSettings settings)
+		{
+			using ILoggerFactory loggerFactory = LoggerFactory.Create(builder =>
+				{
+					builder.SetMinimumLevel(LogLevel.Debug);
+					builder.AddSimpleConsole(options => { options.SingleLine = true; });
+				});
+
+			ILogger<AwsReuseFleetManager> logger = loggerFactory.CreateLogger<AwsReuseFleetManager>();
+			IPool pool = await PoolService.CreatePoolAsync("testPool", null, true, 0, 0, sizeStrategy: PoolSizeStrategy.NoOp);
+			AwsRecyclingFleetManager manager = new (ec2, AgentCollection, settings, logger);
+			await manager.ExpandPoolAsync(pool, new List<IAgent>(), numRequestedInstances, CancellationToken.None);
+		}
+	}
+}
\ No newline at end of file
diff --git a/Engine/Source/Programs/Horde/Horde.Build.Tests/Fleet/AwsReuseFleetManagerTest.cs b/Engine/Source/Programs/Horde/Horde.Build.Tests/Fleet/AwsReuseFleetManagerTest.cs
--- a/Engine/Source/Programs/Horde/Horde.Build.Tests/Fleet/AwsReuseFleetManagerTest.cs
+++ b/Engine/Source/Programs/Horde/Horde.Build.Tests/Fleet/AwsReuseFleetManagerTest.cs
@@ -21,23 +21,25 @@ namespace Horde.Build.Tests.Fleet
 		[TestMethod]
 		public async Task ExpandOneAgent()
 		{
-			FakeAmazonEc2 fakeAmazonEc2 = new ();
-			Instance i = fakeAmazonEc2.AddInstance(FakeAmazonEc2.StateStopped, InstanceType.M5Large);
+			FakeAmazonEc2 ec2 = new ();
+			Instance i = ec2.AddInstance(FakeAmazonEc2.StateStopped, InstanceType.M5Large);
+			ec2.SetCapacity(FakeAmazonEc2.AzUsEast1A, InstanceType.M5Large, 1);
 
-			await ExpandPoolAsync(fakeAmazonEc2.Get(), 1, new());
-			Assert.AreEqual(InstanceType.M5Large, fakeAmazonEc2.Instances[i.InstanceId].InstanceType);
-			Assert.AreEqual(FakeAmazonEc2.StatePending, fakeAmazonEc2.Instances[i.InstanceId].State);
+			await ExpandPoolAsync(ec2.Get(), 1, new());
+			Assert.AreEqual(InstanceType.M5Large, ec2.Instances[i.InstanceId].InstanceType);
+			Assert.AreEqual(FakeAmazonEc2.StatePending, ec2.Instances[i.InstanceId].State);
 		}
 		
 		[TestMethod]
-		public async Task ExpandWithInstanceTypeChange2()
+		public async Task ExpandWithInstanceTypeChange()
 		{
-			FakeAmazonEc2 fakeAmazonEc2 = new ();
-			Instance i = fakeAmazonEc2.AddInstance(FakeAmazonEc2.StateStopped, InstanceType.M5Large);
+			FakeAmazonEc2 ec2 = new ();
+			Instance i = ec2.AddInstance(FakeAmazonEc2.StateStopped, InstanceType.M5Large);
+			ec2.SetCapacity(FakeAmazonEc2.AzUsEast1A, InstanceType.M54xlarge, 1);
 
-			await ExpandPoolAsync(fakeAmazonEc2.Get(), 1, new(new List<string> { InstanceType.M54xlarge }));
-			Assert.AreEqual(InstanceType.M54xlarge, fakeAmazonEc2.Instances[i.InstanceId].InstanceType);
-			Assert.AreEqual(FakeAmazonEc2.StatePending, fakeAmazonEc2.Instances[i.InstanceId].State);
+			await ExpandPoolAsync(ec2.Get(), 1, new(new List<string> { InstanceType.M54xlarge }));
+			Assert.AreEqual(InstanceType.M54xlarge, ec2.Instances[i.InstanceId].InstanceType);
+			Assert.AreEqual(FakeAmazonEc2.StatePending, ec2.Instances[i.InstanceId].State);
 		}
 
 		private async Task ExpandPoolAsync(IAmazonEC2 ec2, int numRequestedInstances, AwsReuseFleetManagerSettings settings)
diff --git a/Engine/Source/Programs/Horde/Horde.Build.Tests/Fleet/FakeAmazonEc2.cs b/Engine/Source/Programs/Horde/Horde.Build.Tests/Fleet/FakeAmazonEc2.cs
--- a/Engine/Source/Programs/Horde/Horde.Build.Tests/Fleet/FakeAmazonEc2.cs
+++ b/Engine/Source/Programs/Horde/Horde.Build.Tests/Fleet/FakeAmazonEc2.cs
@@ -21,6 +21,8 @@ public class FakeAmazonEc2
 {
 	public static readonly InstanceState StatePending = new() { Code = 0, Name = "pending" };
 	public static readonly InstanceState StateStopped = new() { Code = 80, Name = "stopped" };
+	public const string AzUsEast1A = "us-east-1a";
+	public const string AzUsEast1B = "us-east-1b";
 	
 	private readonly Mock<IAmazonEC2> _mock;
 	private readonly Dictionary<string, Instance> _instances = new();
@@ -28,6 +30,12 @@ public class FakeAmazonEc2
 	
 	internal IReadOnlyDictionary<string, Instance> Instances => _instances;
 
+	private readonly Dictionary<string, Dictionary<string, int>> _availabilityZoneCapacity = new()
+	{
+		{AzUsEast1A, new () },
+		{AzUsEast1B, new () }
+	};
+
 	public FakeAmazonEc2()
 	{
 		_mock = new Mock<IAmazonEC2>(MockBehavior.Strict);
@@ -49,9 +57,23 @@ public class FakeAmazonEc2
 		return _mock.Object;
 	}
 
-	public Instance AddInstance(InstanceState state, InstanceType type)
+	public void SetCapacity(string availabilityZone, string instanceType, int capacity)
+	{
+		Dictionary<string, int> temp = _availabilityZoneCapacity.GetValueOrDefault(availabilityZone, new());
+		temp[instanceType] = capacity;
+		_availabilityZoneCapacity[availabilityZone] = temp;
+	}
+
+	public Instance AddInstance(InstanceState state, InstanceType type,
+		string availabilityZone = AzUsEast1A, DateTime? launchTime = null)
 	{
-		Instance i = new() { InstanceId = "bogus-instance-" + _instanceIdCounter++, State = state, InstanceType = type };
+		Instance i = new()
+		{
+			InstanceId = "bogus-instance-" + _instanceIdCounter++,
+			State = state, InstanceType = type,
+			Placement = new Placement(availabilityZone),
+			LaunchTime = launchTime ?? DateTime.UnixEpoch
+		};
 		_instances.Add(i.InstanceId, i);
 		return i;
 	}
@@ -61,6 +83,9 @@ public class FakeAmazonEc2
 		return _instances[instanceId];
 	}
 
+	public int GetStoppedInstanceCount() { return _instances.Values.Count(x => x.State == StateStopped); }
+	public int GetPendingInstanceCount() { return _instances.Values.Count(x => x.State == StatePending); }
+
 	private Task<DescribeInstancesResponse> DescribeInstancesAsync(DescribeInstancesRequest request, CancellationToken cancellationToken)
 	{
 		return Task.FromResult(new DescribeInstancesResponse
@@ -73,6 +98,10 @@ public class FakeAmazonEc2
 	private Task<StartInstancesResponse> StartInstancesAsync(StartInstancesRequest request, CancellationToken cancellationToken)
 	{
 		List<InstanceStateChange> stateChanges = new();
+		
+		// Check if capacity in each AZ can handle this start instance request.
+		// If not, abort with an exception prior attempting to start them and modifying in-memory state.
+		Dictionary<string, Dictionary<string, int>> azInstanceTypeCapacities = new (_availabilityZoneCapacity);
 		foreach (string instanceId in request.InstanceIds)
 		{
 			if (!_instances.ContainsKey(instanceId))
@@ -80,6 +109,20 @@ public class FakeAmazonEc2
 				throw new ArgumentException("Unknown instance ID " + instanceId);
 			}
 
+			string az = _instances[instanceId].Placement.AvailabilityZone;
+			string instanceType = _instances[instanceId].InstanceType;
+			int capacity = azInstanceTypeCapacities[az].GetValueOrDefault(instanceType, 0);
+			
+			if (capacity <= 0)
+			{
+				throw new AmazonEC2Exception("Insufficient capacity.");
+			}
+
+			azInstanceTypeCapacities[az][instanceType] = capacity - 1;
+		}
+		
+		foreach (string instanceId in request.InstanceIds)
+		{
 			if (_instances[instanceId].State.Name == StateStopped.Name)
 			{
 				stateChanges.Add(new()
diff --git a/Engine/Source/Programs/Horde/Horde.Build/Agents/Fleet/FleetManagerFactory.cs b/Engine/Source/Programs/Horde/Horde.Build/Agents/Fleet/FleetManagerFactory.cs
--- a/Engine/Source/Programs/Horde/Horde.Build/Agents/Fleet/FleetManagerFactory.cs
+++ b/Engine/Source/Programs/Horde/Horde.Build/Agents/Fleet/FleetManagerFactory.cs
@@ -63,6 +63,8 @@ public sealed class FleetManagerFactory : IFleetManagerFactory
 				new AwsFleetManager(GetAwsEc2(type), _agentCollection, DeserializeSettings<AwsFleetManagerSettings>(config), _loggerFactory.CreateLogger<AwsFleetManager>()),
 			FleetManagerType.AwsReuse =>
 				new AwsReuseFleetManager(GetAwsEc2(type), _agentCollection, DeserializeSettings<AwsReuseFleetManagerSettings>(config), _loggerFactory.CreateLogger<AwsReuseFleetManager>()),
+			FleetManagerType.AwsRecycle =>
+				new AwsRecyclingFleetManager(GetAwsEc2(type), _agentCollection, DeserializeSettings<AwsRecyclingFleetManagerSettings>(config), _loggerFactory.CreateLogger<AwsReuseFleetManager>()),
 			FleetManagerType.AwsAsg =>
 				new AwsAsgFleetManager(GetAwsAutoScaling(type), DeserializeSettings<AwsAsgSettings>(config), _loggerFactory.CreateLogger<AwsAsgFleetManager>()),
 			_ => throw new ArgumentException("Unknown fleet manager type " + type)
diff --git a/Engine/Source/Programs/Horde/Horde.Build/Agents/Fleet/FleetService.cs b/Engine/Source/Programs/Horde/Horde.Build/Agents/Fleet/FleetService.cs
--- a/Engine/Source/Programs/Horde/Horde.Build/Agents/Fleet/FleetService.cs
+++ b/Engine/Source/Programs/Horde/Horde.Build/Agents/Fleet/FleetService.cs
@@ -44,7 +44,7 @@ namespace Horde.Build.Agents.Fleet
 		private readonly IPoolCollection _poolCollection;
 		private readonly StreamService _streamService;
 		private readonly IDogStatsd _dogStatsd;
-		private readonly IFleetManagerFactory _fleetManagerFactory2;
+		private readonly IFleetManagerFactory _fleetManagerFactory;
 		private readonly IClock _clock;
 		private readonly IMemoryCache _cache;
 		private readonly ITicker _ticker;
@@ -78,7 +78,7 @@ namespace Horde.Build.Agents.Fleet
 			_poolCollection = poolCollection;
 			_streamService = streamService;
 			_dogStatsd = dogStatsd;
-			_fleetManagerFactory2 = fleetManagerFactory;
+			_fleetManagerFactory = fleetManagerFactory;
 			_clock = clock;
 			_cache = cache;
 			_logger = logger;
@@ -276,11 +276,11 @@ namespace Horde.Build.Agents.Fleet
 			{
 				if (info.Condition == null || info.Condition.Evaluate(GetPropValues))
 				{
-					return _fleetManagerFactory2.CreateFleetManager(info.Type, info.Config);
+					return _fleetManagerFactory.CreateFleetManager(info.Type, info.Config);
 				}
 			}
 
-			return _fleetManagerFactory2.CreateFleetManager(FleetManagerType.Default, "{}");
+			return _fleetManagerFactory.CreateFleetManager(FleetManagerType.Default, "{}");
 		}
 
 		/// <summary>
diff --git a/Engine/Source/Programs/Horde/Horde.Build/Agents/Fleet/IFleetManager.cs b/Engine/Source/Programs/Horde/Horde.Build/Agents/Fleet/IFleetManager.cs
--- a/Engine/Source/Programs/Horde/Horde.Build/Agents/Fleet/IFleetManager.cs
+++ b/Engine/Source/Programs/Horde/Horde.Build/Agents/Fleet/IFleetManager.cs
@@ -33,6 +33,11 @@ namespace Horde.Build.Agents.Fleet
 		/// </summary>
 		AwsReuse,
 		
+		/// <summary>
+		/// <see cref="AwsRecyclingFleetManager" />
+		/// </summary>
+		AwsRecycle,
+		
 		/// <summary>
 		/// <see cref="AwsAsgFleetManager" />
 		/// </summary>
diff --git a/Engine/Source/Programs/Horde/Horde.Build/Agents/Fleet/Providers/AwsRecyclingFleetManager.cs b/Engine/Source/Programs/Horde/Horde.Build/Agents/Fleet/Providers/AwsRecyclingFleetManager.cs
new file mode 100644
--- /dev/null
+++ b/Engine/Source/Programs/Horde/Horde.Build/Agents/Fleet/Providers/AwsRecyclingFleetManager.cs
@@ -0,0 +1,373 @@
+// Copyright Epic Games, Inc. All Rights Reserved.
+
+using System;
+using System.Collections.Generic;
+using System.Linq;
+using System.Net;
+using System.Threading;
+using System.Threading.Tasks;
+using Amazon.EC2;
+using Amazon.EC2.Model;
+using Amazon.Runtime;
+using Horde.Build.Agents.Pools;
+using Microsoft.Extensions.Logging;
+using OpenTracing;
+using OpenTracing.Util;
+
+namespace Horde.Build.Agents.Fleet.Providers;
+
+/// <summary>
+/// Settings for AWS fleet manager that reuses instances
+/// </summary>
+public class AwsRecyclingFleetManagerSettings
+{
+	/// <summary>
+	/// Prioritized list of instance types to use waking up an instance
+	/// If set to null, the instance will be awakened with its last instance type.
+	/// </summary>
+	public List<string>? InstanceTypes { get; }
+
+	/// <summary>
+	/// Constructor
+	/// </summary>
+	/// <param name="instanceTypes"></param>
+	public AwsRecyclingFleetManagerSettings(List<string>? instanceTypes = null)
+	{
+		InstanceTypes = instanceTypes;
+	}
+}
+
+/// <summary>
+/// Base exception for AWS fleet managers
+/// </summary>
+public class AwsFleetManagerException : Exception
+{
+	/// <inheritdoc />
+	public AwsFleetManagerException() { }
+
+	/// <inheritdoc />
+	public AwsFleetManagerException(string message) : base(message) { }
+}
+
+/// <summary>
+/// Exception when AWS API returns insufficient capacity
+/// </summary>
+public class AwsInsufficientCapacityException : AwsFleetManagerException { }
+
+/// <summary>
+/// Generic HTTP response error from AWS API
+/// </summary>
+public class AwsBadHttpResponseException : AwsFleetManagerException
+{
+	/// <summary>
+	/// Constructor
+	/// </summary>
+	/// <param name="statusCode"></param>
+	public AwsBadHttpResponseException(HttpStatusCode statusCode) : base($"HTTP status code {statusCode}") { }
+}
+
+/// <summary>
+/// Fleet manager for handling AWS EC2 instances
+/// Will start already existing but stopped instances to reuse existing EBS disks.
+/// See <see cref="AwsFleetManager" /> for creating and terminating instances from scratch.
+/// </summary>
+public sealed class AwsRecyclingFleetManager : IFleetManager
+{
+	private static readonly Random s_random = new ();
+	
+	/// <summary>
+	/// Settings for fleet manager
+	/// </summary>
+	public AwsRecyclingFleetManagerSettings Settings { get; }
+	
+	private readonly IAmazonEC2 _ec2;
+	private readonly IAgentCollection _agentCollection;
+	private readonly ILogger _logger;
+
+	/// <summary>
+	/// Constructor
+	/// </summary>
+	public AwsRecyclingFleetManager(IAmazonEC2 ec2, IAgentCollection agentCollection, AwsRecyclingFleetManagerSettings settings, ILogger<AwsReuseFleetManager> logger)
+	{
+		_ec2 = ec2;
+		_agentCollection = agentCollection;
+		_logger = logger;
+		Settings = settings;
+	}
+
+	/// <inheritdoc/>
+	public async Task ExpandPoolAsync(IPool pool, IReadOnlyList<IAgent> agents, int requestedInstancesCount, CancellationToken cancellationToken)
+	{
+		using IScope scope = GlobalTracer.Instance.BuildSpan("ExpandPool").StartActive();
+		scope.Span.SetTag("poolName", pool.Name);
+		scope.Span.SetTag("numAgents", agents.Count);
+		scope.Span.SetTag("count", requestedInstancesCount);
+
+		using IDisposable logScope = _logger.BeginScope(new Dictionary<string, object> { ["PoolId"] = pool.Id });
+
+		Dictionary<string, List<Instance>> candidatesPerAz = await GetCandidateInstancesAsync(pool, cancellationToken);
+		Dictionary<string, int> instanceCountPerAz = candidatesPerAz.ToDictionary(x => x.Key, x => x.Value.Count);
+		Dictionary<string, int> requestCountPerAz = DistributeInstanceRequestsOverAzs(requestedInstancesCount, instanceCountPerAz, out int remainingInstancesCount);
+		Dictionary<string, List<Instance>> stoppedInstancesPerAz = GetInstancesToLaunch(candidatesPerAz, requestCountPerAz);
+		List<InstanceType>? instanceTypePriority = Settings.InstanceTypes?.Select(InstanceType.FindValue).ToList();
+
+		var info = new { MissingInstanceCount = remainingInstancesCount, InstancesToStartCount = requestCountPerAz.Values.Sum(x => x)};
+		_logger.LogWarning("Not enough stopped instances to accommodate the full pool scale-out. {@info}", info);
+
+		foreach ((string az, List<Instance> instances) in stoppedInstancesPerAz)
+		{
+			using (_logger.BeginScope(new Dictionary<string, object> { ["Az"] = az }))
+			{
+				_logger.LogDebug("Trying to start {InstanceCount} instance(s)", instances.Count);
+				await StartInstancesWithRetriesAsync(instances, instanceTypePriority, cancellationToken);					
+			}
+		}
+	}
+
+	private static Dictionary<string, List<Instance>> GetInstancesToLaunch(
+		Dictionary<string, List<Instance>> candidatesPerAz,
+		Dictionary<string, int> requestCountPerAz)
+	{
+		Dictionary<string, List<Instance>> result = new();
+		foreach ((string az, List<Instance> candidates) in candidatesPerAz)
+		{
+			result[az] = candidates.OrderBy(x => x.LaunchTime).Take(requestCountPerAz[az]).ToList();
+		}
+
+		return result;
+	}
+
+	/// <summary>
+	/// Balance number of instance requests evenly over given availability zones (AZs)
+	/// </summary>
+	/// <param name="requestedInstancesCount">Number of instances requested</param>
+	/// <param name="instanceCountPerAz">Number of instances available per AZ</param>
+	/// <param name="remainingInstancesCount">Number of instances that can't fit within the given AZ counts</param>
+	/// <param name="randomizeAzOrdering">Spread requests evenly by randomizing AZ order. Disabled during testing</param>
+	/// <returns>Number of instances to request from each AZ</returns>
+	internal static Dictionary<string, int> DistributeInstanceRequestsOverAzs(
+		int requestedInstancesCount,
+		IReadOnlyDictionary<string, int> instanceCountPerAz,
+		out int remainingInstancesCount,
+		bool randomizeAzOrdering = true)
+	{
+		Dictionary<string, int> availableCountPerAz = new(instanceCountPerAz);
+		Dictionary<string, int> requestCountPerAz = availableCountPerAz.ToDictionary(x => x.Key, _ => 0);
+
+		List<string> availabilityZones = availableCountPerAz.Keys.OrderBy(x => x).ToList();
+		if (randomizeAzOrdering)
+		{
+			availabilityZones = availabilityZones.OrderBy(_ => s_random.Next()).ToList();
+		}
+
+		int instancesAvailableCount = availableCountPerAz.Sum(x => x.Value);
+		int instancesAssignedCount = 0;
+
+		while (instancesAvailableCount > 0 && instancesAssignedCount < requestedInstancesCount)
+		{
+			foreach (string az in availabilityZones)
+			{
+				int count = availableCountPerAz[az];
+				if (count > 0 && instancesAssignedCount < requestedInstancesCount)
+				{
+					requestCountPerAz[az]++;
+					availableCountPerAz[az]--;
+					instancesAssignedCount++;
+					instancesAvailableCount--;
+				}
+			}
+		}
+
+		remainingInstancesCount = requestedInstancesCount - instancesAssignedCount;
+		return requestCountPerAz;
+	}
+
+	/// <summary>
+	/// Wake up already stopped instances
+	/// If an insufficient capacity error, all instance types will be changed according to the priority type list.
+	/// Capacity errors typically occurs per instance type, so retrying using a different one may circumvent the issue. 
+	/// </summary>
+	/// <param name="instances">List of stopped instances to start</param>
+	/// <param name="instanceTypePriority">List of instance types to fallback on</param>
+	/// <param name="cancellationToken">Cancellation token for the call</param>
+	private async Task StartInstancesWithRetriesAsync(List<Instance> instances, IReadOnlyList<InstanceType>? instanceTypePriority, CancellationToken cancellationToken)
+	{
+		if (instances.Count == 0) return;
+		
+		List<InstanceType?> instanceTypes = new();
+		if (instanceTypePriority == null || instanceTypePriority.Count == 0)
+		{
+			instanceTypes.Add(null);
+		}
+		else
+		{
+			instanceTypes.AddRange(instanceTypePriority);
+		}
+		
+		using IDisposable logScope = _logger.BeginScope(new Dictionary<string, object>
+		{
+			["InstanceIds"] = instances.Select(x => x.InstanceId),
+			["InstanceTypes"] = instanceTypes
+		});
+
+		bool success = false;
+		foreach (InstanceType? instanceType in instanceTypes)
+		{
+			try
+			{
+				await StartInstancesAsync(instances, instanceType, cancellationToken);
+				success = true;
+			}
+			catch (AwsInsufficientCapacityException)
+			{
+				_logger.LogDebug("Insufficient capacity for {InstanceType}", instanceType);
+			}
+		}
+
+		if (!success)
+		{
+			_logger.LogError("Unable to start instances. Insufficient capacity for all instance types tried.");
+		}
+	}
+
+	internal async Task StartInstancesAsync(List<Instance> instances, InstanceType? instanceType, CancellationToken cancellationToken)
+	{
+		if (instances.Count == 0) return;
+		if (instanceType != null)
+		{
+			await ChangeInstanceTypeAsync(instances, instanceType, cancellationToken);
+		}
+
+		StartInstancesRequest startRequest = new ();
+		startRequest.InstanceIds.AddRange(instances.Select(x => x.InstanceId));
+
+		StartInstancesResponse startResponse = new();
+		try
+		{
+			using IScope scope = GlobalTracer.Instance.BuildSpan("StartInstances").StartActive();
+			scope.Span.SetTag("req.instanceIds", String.Join(",", startRequest.InstanceIds));
+			startResponse = await _ec2.StartInstancesAsync(startRequest, cancellationToken);
+			scope.Span.SetTag("res.statusCode", (int)startResponse.HttpStatusCode);
+			scope.Span.SetTag("res.numInstances", startResponse.StartingInstances.Count);
+		}
+		catch (AmazonEC2Exception e)
+		{
+			if (e.Message.Contains("insufficient capacity", StringComparison.OrdinalIgnoreCase))
+			{
+				throw new AwsInsufficientCapacityException();
+			}
+		}
+
+		if (!IsResponseSuccessful(startResponse))
+		{
+			throw new AwsBadHttpResponseException(startResponse.HttpStatusCode);
+		}
+		
+		foreach (InstanceStateChange instanceChange in startResponse.StartingInstances)
+		{
+			_logger.LogInformation("Instance {InstanceId} started ({PrevState} -> {CurrentState})", instanceChange.InstanceId, instanceChange.PreviousState.Name, instanceChange.CurrentState.Name);
+		}
+
+		if (startResponse.StartingInstances.Count != instances.Count)
+		{
+			var info = new { NumRequestInstances = instances.Count, NumStartedInstances = startResponse.StartingInstances.Count };
+			_logger.LogError("Num started instances does not match requested instances {@info}", info);
+		}
+	}
+
+	/// <summary>
+	/// Change instance type of given instances
+	/// </summary>
+	/// <param name="instances">Instances to change</param>
+	/// <param name="newInstanceType">New instance type to use</param>
+	/// <param name="cancellationToken">Cancellation token for the call</param>
+	/// <returns>List of instances grouped by availability zone</returns>
+	private async Task ChangeInstanceTypeAsync(List<Instance> instances, InstanceType newInstanceType, CancellationToken cancellationToken)
+	{
+		using IScope scope = GlobalTracer.Instance.BuildSpan("ChangeInstanceType").StartActive();
+		
+		foreach (Instance instance in instances)
+		{
+			if (instance.InstanceType == newInstanceType) { continue; }
+			using IScope modifyScope = GlobalTracer.Instance.BuildSpan("ModifyInstanceAttribute").StartActive();
+
+			ModifyInstanceAttributeRequest request = new () { InstanceId = instance.InstanceId, InstanceType = newInstanceType };
+			ModifyInstanceAttributeResponse response = await _ec2.ModifyInstanceAttributeAsync(request, cancellationToken);
+			if (IsResponseSuccessful(response))
+			{
+				_logger.LogInformation("Instance type for {InstanceId} modified ({PrevInstanceType} -> {CurrentInstanceType})", instance.InstanceId, instance.InstanceType, newInstanceType);
+			}
+			else
+			{
+				_logger.LogError("Failed to modify instance type for {InstanceId}. Wanted {InstanceType}. Status code {StatusCode}", instance.InstanceId, newInstanceType, response.HttpStatusCode);
+			}
+		}
+	}
+	
+	/// <summary>
+	/// Get list of stopped instances from EC2 API that can be used scaling out a given pool
+	/// </summary>
+	/// <param name="pool">Which pool instances must belong to</param>
+	/// <param name="cancellationToken">Cancellation token for the call</param>
+	/// <returns>List of instances grouped by availability zone</returns>
+	private async Task<Dictionary<string, List<Instance>>> GetCandidateInstancesAsync(IPool pool, CancellationToken cancellationToken)
+	{
+		using IScope scope = GlobalTracer.Instance.BuildSpan("DescribeInstances").StartActive();
+		
+		// Find stopped instances in the correct pool
+		DescribeInstancesRequest request = new()
+		{
+			Filters = new List<Filter>
+			{
+				new("instance-state-name", new List<string> { InstanceStateName.Stopped.Value }),
+				new("tag:" + AwsFleetManager.PoolTagName, new List<string> { pool.Name })
+			}
+		};
+		DescribeInstancesResponse response = await _ec2.DescribeInstancesAsync(request, cancellationToken);
+		scope.Span.SetTag("res.statusCode", (int)response.HttpStatusCode);
+		scope.Span.SetTag("res.numReservations", response.Reservations.Count);
+
+		// Group instance objects by availability zone
+		List<IGrouping<string, Instance>> azInstancePairs = response.Reservations
+			.SelectMany(x => x.Instances)
+			.GroupBy(x => x.Placement.AvailabilityZone).ToList();
+
+		Dictionary<string, object> logScope = new()
+		{
+			["InstanceIds"] = azInstancePairs.ToDictionary(
+				x => x.Key,
+				y => y.Select(x => x.InstanceId).ToList()
+			)
+		};
+		using (_logger.BeginScope(logScope))
+		{
+			_logger.LogDebug("Found {NumCandidates} candidate instances", azInstancePairs.Count);
+		}
+		
+		return azInstancePairs.ToDictionary(x => x.Key, y => y.ToList());;
+	}
+
+	/// <inheritdoc/>
+	public Task ShrinkPoolAsync(IPool pool, IReadOnlyList<IAgent> agents, int count, CancellationToken cancellationToken)
+	{
+		return AwsFleetManager.ShrinkPoolViaAgentShutdownRequestAsync(_agentCollection, pool, agents, count, cancellationToken);
+	} 
+
+	/// <inheritdoc/>
+	public async Task<int> GetNumStoppedInstancesAsync(IPool pool, CancellationToken cancellationToken)
+	{
+		// Find all instances in the pool
+		DescribeInstancesRequest describeRequest = new DescribeInstancesRequest();
+		describeRequest.Filters = new List<Filter>();
+		describeRequest.Filters.Add(new Filter("instance-state-name", new List<string> { InstanceStateName.Stopped.Value }));
+		describeRequest.Filters.Add(new Filter("tag:" + AwsFleetManager.PoolTagName, new List<string> { pool.Name }));
+
+		DescribeInstancesResponse describeResponse = await _ec2.DescribeInstancesAsync(describeRequest, cancellationToken);
+		return describeResponse.Reservations.SelectMany(x => x.Instances).Select(x => x.InstanceId).Distinct().Count();
+	}
+
+	private static bool IsResponseSuccessful(AmazonWebServiceResponse response)
+	{
+		return (int)response.HttpStatusCode >= 200 && (int)response.HttpStatusCode <= 299;
+	}
+}
