Horde: Handle empty fleet manager configs

diff --git a/Engine/Source/Programs/Horde/Horde.Build.Tests/Fleet/FleetServiceTest.cs b/Engine/Source/Programs/Horde/Horde.Build.Tests/Fleet/FleetServiceTest.cs
--- a/Engine/Source/Programs/Horde/Horde.Build.Tests/Fleet/FleetServiceTest.cs
+++ b/Engine/Source/Programs/Horde/Horde.Build.Tests/Fleet/FleetServiceTest.cs
@@ -337,7 +337,9 @@ namespace Horde.Build.Tests.Fleet
 		[TestMethod]
 		public async Task EmptyConfigThrowsException()
 		{
-			await Assert.ThrowsExceptionAsync<JsonException>(() => CreateFleetManager(new FleetManagerInfo(FleetManagerType.AwsAsg, null, "")));
+			IFleetManager fm = await CreateFleetManager(new FleetManagerInfo(FleetManagerType.AwsRecycle, null, ""));
+			Assert.AreEqual(typeof(AwsRecyclingFleetManager), fm.GetType());
+			Assert.IsNull(((AwsRecyclingFleetManager)fm).Settings.InstanceTypes);
 		}
 		
 		[TestMethod]
diff --git a/Engine/Source/Programs/Horde/Horde.Build/Agents/Pools/IPool.cs b/Engine/Source/Programs/Horde/Horde.Build/Agents/Pools/IPool.cs
--- a/Engine/Source/Programs/Horde/Horde.Build/Agents/Pools/IPool.cs
+++ b/Engine/Source/Programs/Horde/Horde.Build/Agents/Pools/IPool.cs
@@ -90,8 +90,7 @@ namespace Horde.Build.Agents.Pools
 		[JsonConstructor]
 		public FleetManagerInfo(FleetManagerType type, Condition? condition, string? config)
 		{
-			// Try deserializing to ensure the config is valid JSON
-			if (config == null) throw new ArgumentNullException(nameof(config));
+			config = String.IsNullOrWhiteSpace(config) ? "{}" : config;
 			JsonSerializer.Deserialize<dynamic>(config);
 			
 			Type = type;
