Horde: Fix incorrect logger name for AwsRecyclingFleetManager

diff --git a/Engine/Source/Programs/Horde/Horde.Build.Tests/Fleet/AwsRecyclingFleetManagerTest.cs b/Engine/Source/Programs/Horde/Horde.Build.Tests/Fleet/AwsRecyclingFleetManagerTest.cs
--- a/Engine/Source/Programs/Horde/Horde.Build.Tests/Fleet/AwsRecyclingFleetManagerTest.cs
+++ b/Engine/Source/Programs/Horde/Horde.Build.Tests/Fleet/AwsRecyclingFleetManagerTest.cs
@@ -183,7 +183,7 @@ namespace Horde.Build.Tests.Fleet
 					builder.AddSimpleConsole(options => { options.SingleLine = true; });
 				});
 
-			ILogger<AwsReuseFleetManager> logger = loggerFactory.CreateLogger<AwsReuseFleetManager>();
+			ILogger<AwsRecyclingFleetManager> logger = loggerFactory.CreateLogger<AwsRecyclingFleetManager>();
 			IPool pool = await PoolService.CreatePoolAsync("testPool", null, true, 0, 0, sizeStrategy: PoolSizeStrategy.NoOp);
 			AwsRecyclingFleetManager manager = new (ec2, AgentCollection, settings, logger);
 			await manager.ExpandPoolAsync(pool, new List<IAgent>(), numRequestedInstances, CancellationToken.None);
diff --git a/Engine/Source/Programs/Horde/Horde.Build/Agents/Fleet/FleetManagerFactory.cs b/Engine/Source/Programs/Horde/Horde.Build/Agents/Fleet/FleetManagerFactory.cs
--- a/Engine/Source/Programs/Horde/Horde.Build/Agents/Fleet/FleetManagerFactory.cs
+++ b/Engine/Source/Programs/Horde/Horde.Build/Agents/Fleet/FleetManagerFactory.cs
@@ -64,7 +64,7 @@ public sealed class FleetManagerFactory : IFleetManagerFactory
 			FleetManagerType.AwsReuse =>
 				new AwsReuseFleetManager(GetAwsEc2(type), _agentCollection, DeserializeSettings<AwsReuseFleetManagerSettings>(config), _loggerFactory.CreateLogger<AwsReuseFleetManager>()),
 			FleetManagerType.AwsRecycle =>
-				new AwsRecyclingFleetManager(GetAwsEc2(type), _agentCollection, DeserializeSettings<AwsRecyclingFleetManagerSettings>(config), _loggerFactory.CreateLogger<AwsReuseFleetManager>()),
+				new AwsRecyclingFleetManager(GetAwsEc2(type), _agentCollection, DeserializeSettings<AwsRecyclingFleetManagerSettings>(config), _loggerFactory.CreateLogger<AwsRecyclingFleetManager>()),
 			FleetManagerType.AwsAsg =>
 				new AwsAsgFleetManager(GetAwsAutoScaling(type), DeserializeSettings<AwsAsgSettings>(config), _loggerFactory.CreateLogger<AwsAsgFleetManager>()),
 			_ => throw new ArgumentException("Unknown fleet manager type " + type)
diff --git a/Engine/Source/Programs/Horde/Horde.Build/Agents/Fleet/Providers/AwsRecyclingFleetManager.cs b/Engine/Source/Programs/Horde/Horde.Build/Agents/Fleet/Providers/AwsRecyclingFleetManager.cs
--- a/Engine/Source/Programs/Horde/Horde.Build/Agents/Fleet/Providers/AwsRecyclingFleetManager.cs
+++ b/Engine/Source/Programs/Horde/Horde.Build/Agents/Fleet/Providers/AwsRecyclingFleetManager.cs
@@ -87,7 +87,7 @@ public sealed class AwsRecyclingFleetManager : IFleetManager
 	/// <summary>
 	/// Constructor
 	/// </summary>
-	public AwsRecyclingFleetManager(IAmazonEC2 ec2, IAgentCollection agentCollection, AwsRecyclingFleetManagerSettings settings, ILogger<AwsReuseFleetManager> logger)
+	public AwsRecyclingFleetManager(IAmazonEC2 ec2, IAgentCollection agentCollection, AwsRecyclingFleetManagerSettings settings, ILogger<AwsRecyclingFleetManager> logger)
 	{
 		_ec2 = ec2;
 		_agentCollection = agentCollection;
