Horde: Limit the number of people notified for a build health issue.

diff --git a/Engine/Source/Programs/Horde/Horde.Build/Notifications/Sinks/SlackNotificationSink.cs b/Engine/Source/Programs/Horde/Horde.Build/Notifications/Sinks/SlackNotificationSink.cs
--- a/Engine/Source/Programs/Horde/Horde.Build/Notifications/Sinks/SlackNotificationSink.cs
+++ b/Engine/Source/Programs/Horde/Horde.Build/Notifications/Sinks/SlackNotificationSink.cs
@@ -771,7 +771,14 @@ namespace Horde.Build.Notifications.Sinks
 			HashSet<UserId> userIds = new HashSet<UserId>();
 			if (notifySuspects)
 			{
-				userIds.UnionWith(details.Suspects.Select(x => x.AuthorId));
+				if (details.Suspects.Count > 5)
+				{
+					_issueService.Collection.GetLogger(issueId).LogInformation("Not notifying suspects for issue; too many users ({Count}).", details.Suspects.Count);
+				}
+				else
+				{
+					userIds.UnionWith(details.Suspects.Select(x => x.AuthorId));
+				}
 			}
 			if (notifyOwner && issue.OwnerId.HasValue)
 			{
