Horde: Add more condition parameters to fleet managers/pool size strategies

Both use this condition to check if they should be enabled.

diff --git a/Engine/Source/Programs/Horde/Horde.Build.Tests/Fleet/FleetServiceTest.cs b/Engine/Source/Programs/Horde/Horde.Build.Tests/Fleet/FleetServiceTest.cs
--- a/Engine/Source/Programs/Horde/Horde.Build.Tests/Fleet/FleetServiceTest.cs
+++ b/Engine/Source/Programs/Horde/Horde.Build.Tests/Fleet/FleetServiceTest.cs
@@ -269,16 +269,70 @@ namespace Horde.Build.Tests.Fleet
 		}
 		
 		[TestMethod]
-		public async Task ConditionDayOfWeek()
+		public void ConditionYear()
 		{
-			Clock.UtcNow = new DateTime(2022, 9, 5, 15, 0, 0, DateTimeKind.Utc); // A Monday
-			IPoolSizeStrategy s1 = await CreateStrategy(new PoolSizeStrategyInfo(PoolSizeStrategy.LeaseUtilization, "dayOfWeek == 'monday'", "{}"));
-			
-			Clock.UtcNow = new DateTime(2022, 9, 6, 15, 0, 0, DateTimeKind.Utc); // A Tuesday
-			IPoolSizeStrategy s2 = await CreateStrategy(new PoolSizeStrategyInfo(PoolSizeStrategy.LeaseUtilization, "dayOfWeek == 'monday'", "{}"));
-			
-			Assert.AreEqual(typeof(LeaseUtilizationStrategy), s1.GetType());
-			Assert.AreEqual(typeof(NoOpPoolSizeStrategy), s2.GetType());
+			Assert.IsTrue(EvalCondition(new DateTime(2022, 9, 1, 12, 0, 0, DateTimeKind.Utc), "timeUtcYear == 2022"));
+			Assert.IsTrue(EvalCondition(new DateTime(1900, 1, 1, 12, 0, 0, DateTimeKind.Utc), "timeUtcYear == 1900"));
+			Assert.IsFalse(EvalCondition(new DateTime(2050, 5, 1, 12, 0, 0, DateTimeKind.Utc), "timeUtcYear == 3"));
+		}
+		
+		[TestMethod]
+		public void ConditionMonth()
+		{
+			Assert.IsTrue(EvalCondition(new DateTime(2022, 9, 1, 12, 0, 0, DateTimeKind.Utc), "timeUtcMonth == 9"));
+			Assert.IsTrue(EvalCondition(new DateTime(2022, 1, 1, 12, 0, 0, DateTimeKind.Utc), "timeUtcMonth == 1"));
+			Assert.IsFalse(EvalCondition(new DateTime(2022, 5, 1, 12, 0, 0, DateTimeKind.Utc), "timeUtcMonth == 3"));
+		}
+		
+		[TestMethod]
+		public void ConditionDay()
+		{
+			Assert.IsTrue(EvalCondition(new DateTime(2022, 9, 1, 12, 0, 0, DateTimeKind.Utc), "timeUtcDay == 1"));
+			Assert.IsTrue(EvalCondition(new DateTime(2022, 1, 31, 12, 0, 0, DateTimeKind.Utc), "timeUtcDay == 31"));
+			Assert.IsFalse(EvalCondition(new DateTime(2022, 5, 5, 12, 0, 0, DateTimeKind.Utc), "timeUtcDay == 3"));
+		}
+		
+		[TestMethod]
+		public void ConditionDayOfWeek()
+		{
+			Assert.IsTrue(EvalCondition(new DateTime(2022, 9, 5, 15, 0, 0, DateTimeKind.Utc), "dayOfWeek == 'monday'")); // A Monday
+			Assert.IsTrue(EvalCondition(new DateTime(2022, 9, 5, 15, 0, 0, DateTimeKind.Utc), "dayOfWeek == 'Monday'")); // A Monday
+			Assert.IsFalse(EvalCondition(new DateTime(2022, 9, 6, 15, 0, 0, DateTimeKind.Utc), "dayOfWeek == 'monday'")); // A Tuesday
+
+			Assert.IsTrue(EvalCondition(new DateTime(2022, 9, 5, 15, 0, 0, DateTimeKind.Utc), "timeUtcDayOfWeek == 'monday'")); // A Monday
+			Assert.IsTrue(EvalCondition(new DateTime(2022, 9, 5, 15, 0, 0, DateTimeKind.Utc), "timeUtcDayOfWeek == 'Monday'")); // A Monday
+			Assert.IsFalse(EvalCondition(new DateTime(2022, 9, 6, 15, 0, 0, DateTimeKind.Utc), "timeUtcDayOfWeek == 'monday'")); // A Tuesday
+		}
+		
+		[TestMethod]
+		public void ConditionHour()
+		{
+			Assert.IsTrue(EvalCondition(new DateTime(2022, 9, 5, 15, 0, 0, DateTimeKind.Utc), "timeUtcHour == 15"));
+			Assert.IsTrue(EvalCondition(new DateTime(2022, 9, 5, 7, 0, 0, DateTimeKind.Utc), "timeUtcHour == 7"));
+			Assert.IsFalse(EvalCondition(new DateTime(2022, 9, 5, 15, 0, 0, DateTimeKind.Utc), "timeUtcHour == 3"));
+		}
+		
+		[TestMethod]
+		public void ConditionMin()
+		{
+			Assert.IsTrue(EvalCondition(new DateTime(2022, 9, 5, 15, 55, 0, DateTimeKind.Utc), "timeUtcMin == 55"));
+			Assert.IsTrue(EvalCondition(new DateTime(2022, 9, 5, 7, 0, 0, DateTimeKind.Utc), "timeUtcMin == 0"));
+			Assert.IsFalse(EvalCondition(new DateTime(2022, 9, 5, 15, 1, 0, DateTimeKind.Utc), "timeUtcMin == 2"));
+		}
+		
+		[TestMethod]
+		public void ConditionSec()
+		{
+			Assert.IsTrue(EvalCondition(new DateTime(2022, 9, 5, 15, 0, 0, DateTimeKind.Utc), "timeUtcSec == 0"));
+			Assert.IsTrue(EvalCondition(new DateTime(2022, 9, 5, 7, 0, 1, DateTimeKind.Utc), "timeUtcSec == 1"));
+			Assert.IsTrue(EvalCondition(new DateTime(2022, 9, 5, 7, 0, 15, DateTimeKind.Utc), "timeUtcSec > 5"));
+			Assert.IsFalse(EvalCondition(new DateTime(2022, 9, 5, 15, 0, 55, DateTimeKind.Utc), "timeUtcSec == 3"));
+		}
+
+		private bool EvalCondition(DateTime now, string condition)
+		{
+			Clock.UtcNow = now;
+			return EpicGames.Horde.Common.Condition.Parse(condition).Evaluate(FleetService.GetPropValues);
 		}
 		
 		private async Task<IPoolSizeStrategy> CreateStrategy(params PoolSizeStrategyInfo[] infos)
diff --git a/Engine/Source/Programs/Horde/Horde.Build/Agents/Fleet/FleetService.cs b/Engine/Source/Programs/Horde/Horde.Build/Agents/Fleet/FleetService.cs
--- a/Engine/Source/Programs/Horde/Horde.Build/Agents/Fleet/FleetService.cs
+++ b/Engine/Source/Programs/Horde/Horde.Build/Agents/Fleet/FleetService.cs
@@ -255,11 +255,20 @@ namespace Horde.Build.Agents.Fleet
 			});
 		}
 		
-		private IEnumerable<string> GetPropValues(string name)
+		internal IEnumerable<string> GetPropValues(string name)
 		{
+			DateTime now = _clock.UtcNow;
 			return name switch
 			{
-				"dayOfWeek" => new List<string> { _clock.UtcNow.DayOfWeek.ToString().ToLower() },
+				"timeUtcYear" => new List<string> { now.Year.ToString() }, // as 1 to 9999
+				"timeUtcMonth" => new List<string> { now.Month.ToString() }, // as 1 to 12
+				"timeUtcDay" => new List<string> { now.Day.ToString() }, // as 1 to 31
+				"timeUtcDayOfWeek" => new List<string> { now.DayOfWeek.ToString(), now.DayOfWeek.ToString().ToLower() },
+				"timeUtcHour" => new List<string> { now.Hour.ToString() }, // as 0 to 23
+				"timeUtcMin" => new List<string> { now.Minute.ToString() }, // as 0 to 59
+				"timeUtcSec" => new List<string> { now.Second.ToString() }, // as 0 to 59
+				
+				"dayOfWeek" => new List<string> { now.DayOfWeek.ToString().ToLower() }, // Deprecated, use timeUtcDayOfWeek
 				_ => Array.Empty<string>()
 			};
 		}
