Horde: Add trace spans for Perforce conforms on agents

diff --git a/Engine/Source/Programs/Horde/Horde.Agent/Execution/PerforceExecutor.cs b/Engine/Source/Programs/Horde/Horde.Agent/Execution/PerforceExecutor.cs
--- a/Engine/Source/Programs/Horde/Horde.Agent/Execution/PerforceExecutor.cs
+++ b/Engine/Source/Programs/Horde/Horde.Agent/Execution/PerforceExecutor.cs
@@ -224,6 +224,9 @@ namespace Horde.Agent.Execution
 
 		public static async Task ConformAsync(DirectoryReference rootDir, IList<AgentWorkspace> pendingWorkspaces, bool removeUntrackedFiles, ILogger logger, CancellationToken cancellationToken)
 		{
+			using IScope scope = GlobalTracer.Instance.BuildSpan("Conform").StartActive();
+			scope.Span.SetTag("workspaces", String.Join(',', pendingWorkspaces.Select(x => x.Identifier)));
+			
 			// Print out all the workspaces we're going to sync
 			logger.LogInformation("Workspaces:");
 			foreach (AgentWorkspace pendingWorkspace in pendingWorkspaces)
