Horde: Allow overriding the data directory for the server via the DataDir config file entry. (H13)

diff --git a/Engine/Source/Programs/Horde/Horde.Build/Program.cs b/Engine/Source/Programs/Horde/Horde.Build/Program.cs
--- a/Engine/Source/Programs/Horde/Horde.Build/Program.cs
+++ b/Engine/Source/Programs/Horde/Horde.Build/Program.cs
@@ -104,11 +104,13 @@ namespace Horde.Build
 	{
 		public static SemVer Version => s_version;
 
+		public static string DeploymentEnvironment { get; } = GetEnvironment();
+
 		public static DirectoryReference AppDir { get; } = GetAppDir();
 
-		public static DirectoryReference DataDir { get; } = GetDefaultDataDir();
+		public static DirectoryReference DataDir { get; } = GetDataDir();
 
-		public static FileReference UserConfigFile { get; } = FileReference.Combine(GetDefaultDataDir(), "Horde.json");
+		public static FileReference UserConfigFile { get; } = FileReference.Combine(DataDir, "Horde.json");
 
 		public static Type[] ConfigSchemas = FindSchemaTypes();
 
@@ -141,21 +143,7 @@ namespace Horde.Build
 
 			CommandLineArguments arguments = new CommandLineArguments(args);
 
-			string? environment = Environment.GetEnvironmentVariable("ASPNETCORE_ENVIRONMENT");
-			if (String.IsNullOrEmpty(environment))
-			{
-				environment = "Production";
-			}
-
-			IConfiguration config = new ConfigurationBuilder()
-				.SetBasePath(AppDir.FullName)
-				.AddJsonFile("appsettings.json", optional: false) 
-				.AddJsonFile("appsettings.Build.json", optional: true) // specific settings for builds (installer/dockerfile)
-				.AddJsonFile($"appsettings.{environment}.json", optional: true) // environment variable overrides, also used in k8s setups with Helm
-				.AddJsonFile("appsettings.User.json", optional: true)
-				.AddJsonFile(UserConfigFile.FullName, optional: true, reloadOnChange: true)
-				.AddEnvironmentVariables()
-				.Build();
+			IConfiguration config = CreateConfig(UserConfigFile);
 
 			ServerSettings hordeSettings = new ServerSettings();
 			config.GetSection("Horde").Bind(hordeSettings);
@@ -200,6 +188,20 @@ namespace Horde.Build
 		// Used by WebApplicationFactory in controller tests. Uses reflection to call this exact function signature.
 		public static IHostBuilder CreateHostBuilder(string[] args) => ServerCommand.CreateHostBuilderForTesting(args);
 
+		/// <summary>
+		/// Gets the current environment
+		/// </summary>
+		/// <returns></returns>
+		static string GetEnvironment()
+		{
+			string? environment = Environment.GetEnvironmentVariable("ASPNETCORE_ENVIRONMENT");
+			if (String.IsNullOrEmpty(environment))
+			{
+				environment = "Production";
+			}
+			return environment;
+		}
+
 		/// <summary>
 		/// Get the application directory
 		/// </summary>
@@ -213,8 +215,16 @@ namespace Horde.Build
 		/// Gets the default directory for storing application data
 		/// </summary>
 		/// <returns>The default data directory</returns>
-		static DirectoryReference GetDefaultDataDir()
+		static DirectoryReference GetDataDir()
 		{
+			IConfiguration config = CreateConfig(null);
+
+			string? dataDir = config.GetSection("Horde").GetValue(typeof(string), nameof(ServerSettings.DataDir)) as string;
+			if (dataDir != null)
+			{
+				return DirectoryReference.Combine(GetAppDir(), dataDir);
+			}
+
 			if (RuntimeInformation.IsOSPlatform(OSPlatform.Windows))
 			{
 				DirectoryReference? dir = DirectoryReference.GetSpecialFolder(Environment.SpecialFolder.CommonApplicationData);
@@ -226,6 +236,28 @@ namespace Horde.Build
 			return DirectoryReference.Combine(GetAppDir(), "Data");
 		}
 
+		/// <summary>
+		/// Constructs a configuration object for the current environment
+		/// </summary>
+		/// <param name="userConfigFile"></param>
+		/// <returns></returns>
+		static IConfiguration CreateConfig(FileReference? userConfigFile)
+		{
+			IConfigurationBuilder builder = new ConfigurationBuilder()
+				.SetBasePath(AppDir.FullName)
+				.AddJsonFile("appsettings.json", optional: false)
+				.AddJsonFile("appsettings.Build.json", optional: true) // specific settings for builds (installer/dockerfile)
+				.AddJsonFile($"appsettings.{DeploymentEnvironment}.json", optional: true) // environment variable overrides, also used in k8s setups with Helm
+				.AddJsonFile("appsettings.User.json", optional: true);
+
+			if (userConfigFile != null)
+			{
+				builder = builder.AddJsonFile(userConfigFile.FullName, optional: true, reloadOnChange: true);
+			}
+
+			return builder.AddEnvironmentVariables().Build();
+		}
+
 		/// <summary>
 		/// Handles bootstrapping of defaults for local servers, which can't be generated during build/installation process (or are better handled here where they can be updated)
 		/// This stuff will change as we get settings into database and could be considered discovery for installer/dockerfile builds 
diff --git a/Engine/Source/Programs/Horde/Horde.Build/ServerSettings.cs b/Engine/Source/Programs/Horde/Horde.Build/ServerSettings.cs
--- a/Engine/Source/Programs/Horde/Horde.Build/ServerSettings.cs
+++ b/Engine/Source/Programs/Horde/Horde.Build/ServerSettings.cs
@@ -215,6 +215,11 @@ namespace Horde.Build
 		/// <inheritdoc cref="RunMode" />
 		public RunMode[]? RunModes { get; set; } = null;
 
+		/// <summary>
+		/// Override the data directory used by Horde. Defaults to C:\ProgramData\HordeServer on Windows, {AppDir}/Data on other platforms.
+		/// </summary>
+		public string? DataDir { get; set; } = null;
+
 		/// <summary>
 		/// Output level for console
 		/// </summary>
