Horde: Remove hack to check that usernames do not contain spaces (true only at Epic). (H13)

diff --git a/Engine/Source/Programs/Horde/Horde.Build/Users/UserCollectionV2.cs b/Engine/Source/Programs/Horde/Horde.Build/Users/UserCollectionV2.cs
--- a/Engine/Source/Programs/Horde/Horde.Build/Users/UserCollectionV2.cs
+++ b/Engine/Source/Programs/Horde/Horde.Build/Users/UserCollectionV2.cs
@@ -256,11 +256,6 @@ namespace Horde.Build.Users
 		/// <inheritdoc/>
 		public async Task<IUser> FindOrAddUserByLoginAsync(string login, string? name, string? email)
 		{
-			if (login.Contains(' ', StringComparison.Ordinal))
-			{
-				_logger.LogWarning("Potentially invalid login name: {Login} (from {Trace})", login, Environment.StackTrace);
-			}
-
 			UserId newUserId = UserId.GenerateNewId();
 			UpdateDefinition<UserDocument> update = Builders<UserDocument>.Update.SetOnInsert(x => x.Id, newUserId).SetOnInsert(x => x.Login, login).Unset(x => x.Hidden);
 
