Horde: Upgrade NuGet packages to latest.

diff --git a/Engine/Source/Programs/Horde/Horde.Agent.Tests/Horde.Agent.Tests.csproj b/Engine/Source/Programs/Horde/Horde.Agent.Tests/Horde.Agent.Tests.csproj
--- a/Engine/Source/Programs/Horde/Horde.Agent.Tests/Horde.Agent.Tests.csproj
+++ b/Engine/Source/Programs/Horde/Horde.Agent.Tests/Horde.Agent.Tests.csproj
@@ -10,8 +10,8 @@
 
   <ItemGroup>
 
-    <PackageReference Include="Microsoft.Extensions.Logging.Abstractions" Version="6.0.1" />
-    <PackageReference Include="Microsoft.NET.Test.Sdk" Version="17.3.0" />
+    <PackageReference Include="Microsoft.Extensions.Logging.Abstractions" Version="6.0.2" />
+    <PackageReference Include="Microsoft.NET.Test.Sdk" Version="17.3.2" />
     <PackageReference Include="Moq" Version="4.18.2" />
     <PackageReference Include="MSTest.TestAdapter" Version="2.2.10" />
     <PackageReference Include="MSTest.TestFramework" Version="2.2.10" />
diff --git a/Engine/Source/Programs/Horde/Horde.Agent/Horde.Agent.csproj b/Engine/Source/Programs/Horde/Horde.Agent/Horde.Agent.csproj
--- a/Engine/Source/Programs/Horde/Horde.Agent/Horde.Agent.csproj
+++ b/Engine/Source/Programs/Horde/Horde.Agent/Horde.Agent.csproj
@@ -22,30 +22,30 @@
   </ItemGroup>
 
   <ItemGroup>
-    <PackageReference Include="AWSSDK.Core" Version="3.7.12.25" />
-    <PackageReference Include="AWSSDK.EC2" Version="3.7.83.1" />
-    <PackageReference Include="AWSSDK.SimpleSystemsManagement" Version="3.7.17.42" />
+    <PackageReference Include="AWSSDK.Core" Version="3.7.100" />
+    <PackageReference Include="AWSSDK.EC2" Version="3.7.100" />
+    <PackageReference Include="AWSSDK.SimpleSystemsManagement" Version="3.7.100" />
     <PackageReference Include="Datadog.Trace.OpenTracing" Version="2.14.0" />
-    <PackageReference Include="Google.Protobuf" Version="3.21.5" />
-    <PackageReference Include="Grpc.Net.Client" Version="2.48.0" />
-    <PackageReference Include="Grpc.Tools" Version="2.48.0">
+    <PackageReference Include="Google.Protobuf" Version="3.21.8" />
+    <PackageReference Include="Grpc.Net.Client" Version="2.49.0" />
+    <PackageReference Include="Grpc.Tools" Version="2.50.0">
       <PrivateAssets>all</PrivateAssets>
       <IncludeAssets>runtime; build; native; contentfiles; analyzers; buildtransitive</IncludeAssets>
     </PackageReference>
     <PackageReference Include="Microsoft.Extensions.Caching.Memory" Version="6.0.1" />
     <PackageReference Include="Microsoft.Extensions.Hosting" Version="6.0.1" />
     <PackageReference Include="Microsoft.Extensions.Hosting.Abstractions" Version="6.0.0" />
-    <PackageReference Include="Microsoft.Extensions.Hosting.WindowsServices" Version="6.0.0" />
+    <PackageReference Include="Microsoft.Extensions.Hosting.WindowsServices" Version="6.0.1" />
     <PackageReference Include="Microsoft.Extensions.Http" Version="6.0.0" />
-    <PackageReference Include="Microsoft.Extensions.Http.Polly" Version="6.0.8" />
-    <PackageReference Include="Microsoft.Extensions.Logging.Abstractions" Version="6.0.1" />
+    <PackageReference Include="Microsoft.Extensions.Http.Polly" Version="6.0.10" />
+    <PackageReference Include="Microsoft.Extensions.Logging.Abstractions" Version="6.0.2" />
     <PackageReference Include="Microsoft.Extensions.Options.DataAnnotations" Version="6.0.0" />
     <PackageReference Include="Microsoft.Management.Infrastructure" Version="2.0.0" />
     <PackageReference Include="OpenTracing" Version="0.12.1" />
-    <PackageReference Include="Serilog" Version="2.11.0" />
+    <PackageReference Include="Serilog" Version="2.12.0" />
     <PackageReference Include="Serilog.Extensions.Hosting" Version="5.0.1" />
     <PackageReference Include="Serilog.Extensions.Logging" Version="3.1.0" />
-    <PackageReference Include="Serilog.Sinks.Console" Version="4.0.1" />
+    <PackageReference Include="Serilog.Sinks.Console" Version="4.1.0" />
     <PackageReference Include="Serilog.Sinks.File" Version="5.0.0" />
     <PackageReference Include="System.ServiceProcess.ServiceController" Version="6.0.0" />
   </ItemGroup>
diff --git a/Engine/Source/Programs/Horde/Horde.Build.Tests/Horde.Build.Tests.csproj b/Engine/Source/Programs/Horde/Horde.Build.Tests/Horde.Build.Tests.csproj
--- a/Engine/Source/Programs/Horde/Horde.Build.Tests/Horde.Build.Tests.csproj
+++ b/Engine/Source/Programs/Horde/Horde.Build.Tests/Horde.Build.Tests.csproj
@@ -18,8 +18,8 @@
   </ItemGroup>
 
   <ItemGroup>
-    <PackageReference Include="Microsoft.AspNetCore.Mvc.Testing" Version="6.0.8" />
-    <PackageReference Include="Microsoft.NET.Test.Sdk" Version="17.3.0" />
+    <PackageReference Include="Microsoft.AspNetCore.Mvc.Testing" Version="6.0.10" />
+    <PackageReference Include="Microsoft.NET.Test.Sdk" Version="17.3.2" />
     <PackageReference Include="Moq" Version="4.18.2" />
     <PackageReference Include="MSTest.TestAdapter" Version="2.2.10" />
     <PackageReference Include="MSTest.TestFramework" Version="2.2.10" />
diff --git a/Engine/Source/Programs/Horde/Horde.Build/Horde.Build.csproj b/Engine/Source/Programs/Horde/Horde.Build/Horde.Build.csproj
--- a/Engine/Source/Programs/Horde/Horde.Build/Horde.Build.csproj
+++ b/Engine/Source/Programs/Horde/Horde.Build/Horde.Build.csproj
@@ -19,30 +19,30 @@
   </PropertyGroup>
 
   <ItemGroup>
-    <PackageReference Include="AWSSDK.AutoScaling" Version="3.7.10.48" />
-    <PackageReference Include="AWSSDK.CloudWatch" Version="3.7.7.3" />
-    <PackageReference Include="AWSSDK.EC2" Version="3.7.83.1" />
+    <PackageReference Include="AWSSDK.AutoScaling" Version="3.7.100" />
+    <PackageReference Include="AWSSDK.CloudWatch" Version="3.7.100" />
+    <PackageReference Include="AWSSDK.EC2" Version="3.7.100" />
     <PackageReference Include="AWSSDK.Extensions.NETCore.Setup" Version="3.7.2" />
-    <PackageReference Include="AWSSDK.S3" Version="3.7.9.43" />
-    <PackageReference Include="AWSSDK.SecurityToken" Version="3.7.1.189" />
+    <PackageReference Include="AWSSDK.S3" Version="3.7.100" />
+    <PackageReference Include="AWSSDK.SecurityToken" Version="3.7.100" />
     <PackageReference Include="Datadog.Trace.OpenTracing" Version="2.14.0" />
     <PackageReference Include="DogStatsD-CSharp-Client" Version="7.0.1" />
-    <PackageReference Include="Grpc.AspNetCore" Version="2.48.0" />
-    <PackageReference Include="Grpc.AspNetCore.Server.Reflection" Version="2.48.0" />
-    <PackageReference Include="Grpc.Net.Client" Version="2.48.0" />
-    <PackageReference Include="Grpc.Tools" Version="2.48.0">
+    <PackageReference Include="Grpc.AspNetCore" Version="2.49.0" />
+    <PackageReference Include="Grpc.AspNetCore.Server.Reflection" Version="2.49.0" />
+    <PackageReference Include="Grpc.Net.Client" Version="2.49.0" />
+    <PackageReference Include="Grpc.Tools" Version="2.50.0">
       <PrivateAssets>all</PrivateAssets>
       <IncludeAssets>runtime; build; native; contentfiles; analyzers; buildtransitive</IncludeAssets>
     </PackageReference>
     <PackageReference Include="JetBrains.Profiler.SelfApi" Version="2.3.0" />
     <PackageReference Include="Microsoft.AspNetCore.Authentication.Cookies" Version="2.2.0" />
-    <PackageReference Include="Microsoft.AspNetCore.Authentication.JwtBearer" Version="6.0.8" />
-    <PackageReference Include="Microsoft.AspNetCore.Authentication.OpenIdConnect" Version="6.0.8" />
-    <PackageReference Include="Microsoft.AspNetCore.DataProtection.StackExchangeRedis" Version="6.0.8" />
-    <PackageReference Include="Microsoft.AspNetCore.SpaServices.Extensions" Version="6.0.8" />
+    <PackageReference Include="Microsoft.AspNetCore.Authentication.JwtBearer" Version="6.0.10" />
+    <PackageReference Include="Microsoft.AspNetCore.Authentication.OpenIdConnect" Version="6.0.10" />
+    <PackageReference Include="Microsoft.AspNetCore.DataProtection.StackExchangeRedis" Version="6.0.10" />
+    <PackageReference Include="Microsoft.AspNetCore.SpaServices.Extensions" Version="6.0.10" />
     <PackageReference Include="Microsoft.Extensions.Caching.Memory" Version="6.0.1" />
     <PackageReference Include="Microsoft.Extensions.Http" Version="6.0.0" />
-    <PackageReference Include="MongoDB.Driver" Version="2.17.1" />
+    <PackageReference Include="MongoDB.Driver" Version="2.18.0" />
     <PackageReference Include="OpenTracing" Version="0.12.1" />
     <PackageReference Include="OpenTracing.Contrib.Grpc" Version="0.4.0" />
     <PackageReference Include="Polly" Version="7.2.3" />
@@ -50,14 +50,14 @@
     <PackageReference Include="Serilog.AspNetCore" Version="6.0.1" />
     <PackageReference Include="Serilog.Enrichers.OpenTracing" Version="0.0.2" />
     <PackageReference Include="Serilog.Extensions.Hosting" Version="5.0.1" />
-    <PackageReference Include="Serilog.Sinks.Console" Version="4.0.1" />
+    <PackageReference Include="Serilog.Sinks.Console" Version="4.1.0" />
     <PackageReference Include="Serilog.Sinks.File" Version="5.0.0" />
     <PackageReference Include="SerilogAnalyzer" Version="0.15.0" />
-    <PackageReference Include="SharpZipLib" Version="1.3.3" />
-    <PackageReference Include="StackExchange.Redis" Version="2.6.48" />
+    <PackageReference Include="SharpZipLib" Version="1.4.0" />
+    <PackageReference Include="StackExchange.Redis" Version="2.6.70" />
     <PackageReference Include="Swashbuckle.AspNetCore" Version="6.4.0" />
     <PackageReference Include="System.Linq.Async" Version="6.0.1" />
-    <PackageReference Include="TimeZoneConverter" Version="5.0.0" />
+    <PackageReference Include="TimeZoneConverter" Version="6.0.1" />
   </ItemGroup>
 
   <ItemGroup>
