Horde: Await completion of issue logging. Async void methods are meant for event handlers, and will run asynchronously in the background. Also used existing user cache to format names.

diff --git a/Engine/Source/Programs/Horde/Horde.Build/Issues/IssueCollection.cs b/Engine/Source/Programs/Horde/Horde.Build/Issues/IssueCollection.cs
--- a/Engine/Source/Programs/Horde/Horde.Build/Issues/IssueCollection.cs
+++ b/Engine/Source/Programs/Horde/Horde.Build/Issues/IssueCollection.cs
@@ -520,60 +520,24 @@ namespace Horde.Build.Issues
 			return newIssue;
 		}
 
-		static void GetRelevantUserIds(List<IIssue> issues, HashSet<UserId> ids)
+		async ValueTask<string> GetUserNameAsync(UserId? UserId)
 		{
-			foreach (IIssue issue in issues)
+			if (UserId == null)
 			{
-				if (issue.OwnerId != null)
-				{
-					ids.Add(issue.OwnerId.Value);
-				}
-				if (issue.NominatedById != null)
-				{
-					ids.Add(issue.NominatedById.Value);
-				}
-				if (issue.ResolvedById != null)
-				{
-					ids.Add(issue.ResolvedById.Value);
-				}
-				if (issue.QuarantinedByUserId != null)
-				{
-					ids.Add(issue.QuarantinedByUserId.Value);
-				}
-				if (issue.ForceClosedByUserId != null)
-				{
-					ids.Add(issue.ForceClosedByUserId.Value);
-				}
+				return "null";
 			}
-		}
 
-		async void LogIssueChanges(ILogger issueLogger, Issue oldIssue, Issue newIssue)
-		{
-			HashSet<UserId> ids = new HashSet<UserId>();
-			GetRelevantUserIds(new List<IIssue>() { oldIssue, newIssue }, ids);
-
-			Dictionary<UserId, string> userNames = new Dictionary<UserId, string>() { [IIssue.ResolvedByUnknownId] = "Horde", [IIssue.ResolvedByTimeoutId] = "Horde"};
-
-			if (ids.Count > 0)
+			IUser? user = await _userCollection.GetCachedUserAsync(UserId);
+			if (user == null)
 			{
-				List<IUser> users = await _userCollection.FindUsersAsync(ids);
-				users.ForEach(x => { userNames[x.Id] = x.Name; });
+				return "Unknown user";
 			}
 
-			string GetUserName(UserId? UserId)
-			{
-				if (UserId == null)
-				{
-					return "null";
-				}
-				string? name;
-				if (userNames.TryGetValue(UserId.Value, out name))
-				{
-					return name;
-				}
-				return "Unknown User";
-			}
+			return user.Name;
+		}
 
+		async Task LogIssueChangesAsync(ILogger issueLogger, Issue oldIssue, Issue newIssue)
+		{
 			if (newIssue.Severity != oldIssue.Severity)
 			{
 				issueLogger.LogInformation("Changed severity to {Severity}", newIssue.Severity);
@@ -594,22 +558,22 @@ namespace Horde.Build.Issues
 			{
 				if (newIssue.NominatedById != null)
 				{					
-					issueLogger.LogInformation("User {UserName} ({UserId}) was nominated by {NominatedByUserName} ({NominatedByUserId})", GetUserName(newIssue.OwnerId), newIssue.OwnerId, GetUserName(newIssue.NominatedById), newIssue.NominatedById);
+					issueLogger.LogInformation("User {UserName} ({UserId}) was nominated by {NominatedByUserName} ({NominatedByUserId})", await GetUserNameAsync(newIssue.OwnerId), newIssue.OwnerId, await GetUserNameAsync(newIssue.NominatedById), newIssue.NominatedById);
 				}
 				else
 				{
-					issueLogger.LogInformation("User {UserName} ({UserId}) was nominated by default", GetUserName(newIssue.OwnerId), newIssue.OwnerId);
+					issueLogger.LogInformation("User {UserName} ({UserId}) was nominated by default", await GetUserNameAsync(newIssue.OwnerId), newIssue.OwnerId);
 				}
 			}
 			if (newIssue.AcknowledgedAt != oldIssue.AcknowledgedAt)
 			{
 				if (newIssue.AcknowledgedAt == null)
 				{
-					issueLogger.LogInformation("Issue was un-acknowledged by {UserName} ({UserId})", GetUserName(oldIssue.OwnerId), oldIssue.OwnerId);
+					issueLogger.LogInformation("Issue was un-acknowledged by {UserName} ({UserId})", await GetUserNameAsync(oldIssue.OwnerId), oldIssue.OwnerId);
 				}
 				else
 				{
-					issueLogger.LogInformation("Issue was acknowledged by {UserName} ({UserId})", GetUserName(newIssue.OwnerId), newIssue.OwnerId);
+					issueLogger.LogInformation("Issue was acknowledged by {UserName} ({UserId})", await GetUserNameAsync(newIssue.OwnerId), newIssue.OwnerId);
 				}
 			}
 			if (newIssue.FixChange != oldIssue.FixChange)
@@ -631,7 +595,7 @@ namespace Horde.Build.Issues
 				}
 				else
 				{
-					issueLogger.LogInformation("Resolved by {UserName} ({UserId})", GetUserName(newIssue.ResolvedById), newIssue.ResolvedById);
+					issueLogger.LogInformation("Resolved by {UserName} ({UserId})", await GetUserNameAsync(newIssue.ResolvedById), newIssue.ResolvedById);
 				}
 			}
 
@@ -675,7 +639,7 @@ namespace Horde.Build.Issues
 			{
 				if (newIssue.QuarantinedByUserId != null)
 				{
-					issueLogger.LogInformation("Quarantined by {UserName} ({UserId})", GetUserName(newIssue.QuarantinedByUserId), newIssue.QuarantinedByUserId);
+					issueLogger.LogInformation("Quarantined by {UserName} ({UserId})", await GetUserNameAsync(newIssue.QuarantinedByUserId), newIssue.QuarantinedByUserId);
 				}
 				else
 				{
@@ -687,7 +651,7 @@ namespace Horde.Build.Issues
 			{
 				if (newIssue.ForceClosedByUserId != null)
 				{
-					issueLogger.LogInformation("Forced closed by {UserName} ({UserId})", GetUserName(newIssue.ForceClosedByUserId), newIssue.ForceClosedByUserId);
+					issueLogger.LogInformation("Forced closed by {UserName} ({UserId})", await GetUserNameAsync(newIssue.ForceClosedByUserId), newIssue.ForceClosedByUserId);
 				}
 				else
 				{
@@ -1105,7 +1069,7 @@ namespace Horde.Build.Issues
 			}
 
 			ILogger issueLogger = GetLogger(issue.Id);
-			LogIssueChanges(issueLogger, issueDocument, newIssue);
+			await LogIssueChangesAsync(issueLogger, issueDocument, newIssue);
 			return newIssue;
 		}
 
@@ -1228,7 +1192,7 @@ namespace Horde.Build.Issues
 			if(newIssue != null)
 			{
 				ILogger issueLogger = GetLogger(issue.Id);
-				LogIssueChanges(GetLogger(issue.Id), issueImpl, newIssue);
+				await LogIssueChangesAsync(GetLogger(issue.Id), issueImpl, newIssue);
 				LogIssueSuspectChanges(issueLogger, oldSuspectImpls, newSuspectImpls);
 				return newIssue;
 			}
