Horde: Ensure the user initiating an issue change is named in the log.

diff --git a/Engine/Source/Programs/Horde/Horde.Build/Issues/External/JiraService.cs b/Engine/Source/Programs/Horde/Horde.Build/Issues/External/JiraService.cs
--- a/Engine/Source/Programs/Horde/Horde.Build/Issues/External/JiraService.cs
+++ b/Engine/Source/Programs/Horde/Horde.Build/Issues/External/JiraService.cs
@@ -423,7 +423,7 @@ namespace Horde.Build.Issues.External
 				throw new Exception($"Unable to parse returned jira json: {responseBody}");
 			}
 
-			await _issueService.UpdateIssueAsync(issueId, externalIssueKey: jiraResponse.Key);
+			await _issueService.UpdateIssueAsync(issueId, externalIssueKey: jiraResponse.Key, initiatedById: user.Id);
 
 			// add the user as a watcher to the newly created issue
 			if (externalIssueUser != null)
diff --git a/Engine/Source/Programs/Horde/Horde.Build/Issues/IIssueCollection.cs b/Engine/Source/Programs/Horde/Horde.Build/Issues/IIssueCollection.cs
--- a/Engine/Source/Programs/Horde/Horde.Build/Issues/IIssueCollection.cs
+++ b/Engine/Source/Programs/Horde/Horde.Build/Issues/IIssueCollection.cs
@@ -450,6 +450,7 @@ namespace Horde.Build.Issues
 		/// Try to update the state of an issue
 		/// </summary>
 		/// <param name="issue">The issue to update</param>
+		/// <param name="initiatedByUserId">User that is making the change</param>
 		/// <param name="newSeverity">New severity for the issue</param>
 		/// <param name="newSummary">New summary for the issue</param>
 		/// <param name="newUserSummary">New user summary for the issue</param>
@@ -467,7 +468,7 @@ namespace Horde.Build.Issues
 		/// <param name="newQuarantinedById">The user that quarantined the issue</param>
 		/// <param name="newForceClosedById">The user that force closed the issue</param>
 		/// <returns>True if the issue was updated</returns>
-		Task<IIssue?> TryUpdateIssueAsync(IIssue issue, IssueSeverity? newSeverity = null, string? newSummary = null, string? newUserSummary = null, string? newDescription = null, bool? newPromoted = null, UserId? newOwnerId = null, UserId? newNominatedById = null, bool? newAcknowledged = null, UserId? newDeclinedById = null, int? newFixChange = null, UserId? newResolvedById = null, List<ObjectId>? newExcludeSpanIds = null, DateTime? newLastSeenAt = null, string? newExternalIssueKey = null, UserId? newQuarantinedById = null, UserId? newForceClosedById = null);
+		Task<IIssue?> TryUpdateIssueAsync(IIssue issue, UserId? initiatedByUserId, IssueSeverity? newSeverity = null, string? newSummary = null, string? newUserSummary = null, string? newDescription = null, bool? newPromoted = null, UserId? newOwnerId = null, UserId? newNominatedById = null, bool? newAcknowledged = null, UserId? newDeclinedById = null, int? newFixChange = null, UserId? newResolvedById = null, List<ObjectId>? newExcludeSpanIds = null, DateTime? newLastSeenAt = null, string? newExternalIssueKey = null, UserId? newQuarantinedById = null, UserId? newForceClosedById = null);
 
 		/// <summary>
 		/// Updates derived data for an issue (ie. data computed from the spans attached to it). Also clears the issue's 'modified' state.
diff --git a/Engine/Source/Programs/Horde/Horde.Build/Issues/IssueCollection.cs b/Engine/Source/Programs/Horde/Horde.Build/Issues/IssueCollection.cs
--- a/Engine/Source/Programs/Horde/Horde.Build/Issues/IssueCollection.cs
+++ b/Engine/Source/Programs/Horde/Horde.Build/Issues/IssueCollection.cs
@@ -536,7 +536,14 @@ namespace Horde.Build.Issues
 			return user.Name;
 		}
 
-		async Task LogIssueChangesAsync(ILogger issueLogger, Issue oldIssue, Issue newIssue)
+		async Task LogIssueChangesAsync(UserId? initiatedByUserId, Issue oldIssue, Issue newIssue)
+		{
+			ILogger issueLogger = GetLogger(oldIssue.Id);
+			using IDisposable scope = issueLogger.BeginScope("User {UserName} ({UserId})", await GetUserNameAsync(initiatedByUserId), initiatedByUserId ?? UserId.Empty);
+			await LogIssueChangesImplAsync(issueLogger, oldIssue, newIssue);
+		}
+
+		async Task LogIssueChangesImplAsync(ILogger issueLogger, Issue oldIssue, Issue newIssue)
 		{
 			if (newIssue.Severity != oldIssue.Severity)
 			{
@@ -678,28 +685,28 @@ namespace Horde.Build.Issues
 			}
 		}
 
-		static void LogIssueSuspectChanges(ILogger issueLogger, List<IssueSuspect> oldIssueSuspects, List<IssueSuspect> newIssueSuspects)
+		async Task LogIssueSuspectChangesAsync(ILogger issueLogger, List<IssueSuspect> oldIssueSuspects, List<IssueSuspect> newIssueSuspects)
 		{
 			HashSet<(UserId, int)> oldSuspects = new HashSet<(UserId, int)>(oldIssueSuspects.Select(x => (x.AuthorId, x.Change)));
 			HashSet<(UserId, int)> newSuspects = new HashSet<(UserId, int)>(newIssueSuspects.Select(x => (x.AuthorId, x.Change)));
 			foreach ((UserId userId, int change) in newSuspects.Where(x => !oldSuspects.Contains(x)))
 			{
-				issueLogger.LogInformation("Added suspect {UserId} for change {Change}", userId, change);
+				issueLogger.LogInformation("Added suspect {UserName} ({UserId}) for change {Change}", await GetUserNameAsync(userId), userId, change);
 			}
 			foreach ((UserId userId, int change) in oldSuspects.Where(x => !newSuspects.Contains(x)))
 			{
-				issueLogger.LogInformation("Removed suspect {UserId} for change {Change}", userId, change);
+				issueLogger.LogInformation("Removed suspect {UserName} ({UserId}) for change {Change}", await GetUserNameAsync(userId), userId, change);
 			}
 
 			HashSet<UserId> oldDeclinedBy = new HashSet<UserId>(oldIssueSuspects.Where(x => x.DeclinedAt != null).Select(x => x.AuthorId));
 			HashSet<UserId> newDeclinedBy = new HashSet<UserId>(newIssueSuspects.Where(x => x.DeclinedAt != null).Select(x => x.AuthorId));
 			foreach (UserId addDeclinedBy in newDeclinedBy.Where(x => !oldDeclinedBy.Contains(x)))
 			{
-				issueLogger.LogInformation("Declined by {UserId}", addDeclinedBy);
+				issueLogger.LogInformation("Declined by {UserName} ({UserId})", await GetUserNameAsync(addDeclinedBy), addDeclinedBy);
 			}
 			foreach (UserId removeDeclinedBy in oldDeclinedBy.Where(x => !newDeclinedBy.Contains(x)))
 			{
-				issueLogger.LogInformation("Un-declined by {UserId}", removeDeclinedBy);
+				issueLogger.LogInformation("Un-declined by {UserName} ({UserId})", await GetUserNameAsync(removeDeclinedBy), removeDeclinedBy);
 			}
 		}
 
@@ -876,7 +883,7 @@ namespace Horde.Build.Issues
 		}
 
 		/// <inheritdoc/>
-		public async Task<IIssue?> TryUpdateIssueAsync(IIssue issue, IssueSeverity? newSeverity = null, string? newSummary = null, string? newUserSummary = null, string? newDescription = null, bool? newManuallyPromoted = null, UserId? newOwnerId = null, UserId? newNominatedById = null, bool? newAcknowledged = null, UserId? newDeclinedById = null, int? newFixChange = null, UserId? newResolvedById = null, List<ObjectId>? newExcludeSpanIds = null, DateTime? newLastSeenAt = null, string? newExternaIssueKey = null, UserId? newQuarantinedById = null, UserId? newForceClosedById = null)
+		public async Task<IIssue?> TryUpdateIssueAsync(IIssue issue, UserId? initiatedByUserId, IssueSeverity? newSeverity = null, string? newSummary = null, string? newUserSummary = null, string? newDescription = null, bool? newManuallyPromoted = null, UserId? newOwnerId = null, UserId? newNominatedById = null, bool? newAcknowledged = null, UserId? newDeclinedById = null, int? newFixChange = null, UserId? newResolvedById = null, List<ObjectId>? newExcludeSpanIds = null, DateTime? newLastSeenAt = null, string? newExternaIssueKey = null, UserId? newQuarantinedById = null, UserId? newForceClosedById = null)
 		{
 			Issue issueDocument = (Issue)issue;
 
@@ -1051,7 +1058,6 @@ namespace Horde.Build.Issues
 				}
 			}
 
-
 			if (newExternaIssueKey != null)
 			{
 				updates.Add(Builders<Issue>.Update.Set(x => x.ExternalIssueKey, newExternaIssueKey.Length == 0 ? null : newExternaIssueKey));
@@ -1068,8 +1074,7 @@ namespace Horde.Build.Issues
 				return null;
 			}
 
-			ILogger issueLogger = GetLogger(issue.Id);
-			await LogIssueChangesAsync(issueLogger, issueDocument, newIssue);
+			await LogIssueChangesAsync(initiatedByUserId, issueDocument, newIssue);
 			return newIssue;
 		}
 
@@ -1191,9 +1196,8 @@ namespace Horde.Build.Issues
 			Issue? newIssue = await TryUpdateIssueAsync(issue, Builders<Issue>.Update.Combine(updates));
 			if(newIssue != null)
 			{
-				ILogger issueLogger = GetLogger(issue.Id);
-				await LogIssueChangesAsync(GetLogger(issue.Id), issueImpl, newIssue);
-				LogIssueSuspectChanges(issueLogger, oldSuspectImpls, newSuspectImpls);
+				await LogIssueChangesAsync(null, issueImpl, newIssue);
+				await LogIssueSuspectChangesAsync(GetLogger(issue.Id), oldSuspectImpls, newSuspectImpls);
 				return newIssue;
 			}
 			return null;
diff --git a/Engine/Source/Programs/Horde/Horde.Build/Issues/IssueService.cs b/Engine/Source/Programs/Horde/Horde.Build/Issues/IssueService.cs
--- a/Engine/Source/Programs/Horde/Horde.Build/Issues/IssueService.cs
+++ b/Engine/Source/Programs/Horde/Horde.Build/Issues/IssueService.cs
@@ -287,7 +287,7 @@ namespace Horde.Build.Issues
 
 				if (openIssue.LastSeenAt < utcNow - TimeSpan.FromDays(7.0))
 				{
-					await _issueCollection.TryUpdateIssueAsync(openIssue, newResolvedById: IIssue.ResolvedByTimeoutId);
+					await _issueCollection.TryUpdateIssueAsync(openIssue, null, newResolvedById: IIssue.ResolvedByTimeoutId);
 					openIssues.RemoveAt(idx--);
 					continue;
 				}
@@ -389,7 +389,7 @@ namespace Horde.Build.Issues
 		/// <summary>
 		/// Updates the state of an issue
 		/// </summary>
-		/// <param name="id">The current issue id</param>s
+		/// <param name="id">The current issue id</param>
 		/// <param name="summary">New summary for the issue</param>
 		/// <param name="description">New description for the issue</param>
 		/// <param name="promoted">Whether the issue should be set as promoted</param>
@@ -404,8 +404,9 @@ namespace Horde.Build.Issues
 		/// <param name="externalIssueKey">Key for external issue tracking</param>
 		/// <param name="quarantinedById">User who has quarantined the issue</param>
 		/// <param name="forceClosedById">User who has force closed the issue</param>
+		/// <param name="initiatedById">User initiating the changes, for auditing purposes</param>
 		/// <returns>True if the issue was updated</returns>
-		public async Task<bool> UpdateIssueAsync(int id, string? summary = null, string? description = null, bool? promoted = null, UserId? ownerId = null, UserId? nominatedById = null, bool? acknowledged = null, UserId? declinedById = null, int? fixChange = null, UserId? resolvedById = null, List<ObjectId>? addSpanIds = null, List<ObjectId>? removeSpanIds = null, string? externalIssueKey = null, UserId? quarantinedById = null, UserId? forceClosedById = null)
+		public async Task<bool> UpdateIssueAsync(int id, string? summary = null, string? description = null, bool? promoted = null, UserId? ownerId = null, UserId? nominatedById = null, bool? acknowledged = null, UserId? declinedById = null, int? fixChange = null, UserId? resolvedById = null, List<ObjectId>? addSpanIds = null, List<ObjectId>? removeSpanIds = null, string? externalIssueKey = null, UserId? quarantinedById = null, UserId? forceClosedById = null, UserId? initiatedById = null)
 		{
 			IIssue? issue;
 			for (; ; )
@@ -416,7 +417,7 @@ namespace Horde.Build.Issues
 					return false;
 				}
 
-				issue = await _issueCollection.TryUpdateIssueAsync(issue, newUserSummary: summary, newDescription: description, newPromoted: promoted, newOwnerId: ownerId ?? resolvedById, newNominatedById: nominatedById, newDeclinedById: declinedById, newAcknowledged: acknowledged, newFixChange: fixChange, newResolvedById: resolvedById, newExcludeSpanIds: removeSpanIds, newExternalIssueKey: externalIssueKey, newQuarantinedById: quarantinedById, newForceClosedById: forceClosedById);
+				issue = await _issueCollection.TryUpdateIssueAsync(issue, initiatedById, newUserSummary: summary, newDescription: description, newPromoted: promoted, newOwnerId: ownerId ?? resolvedById, newNominatedById: nominatedById, newDeclinedById: declinedById, newAcknowledged: acknowledged, newFixChange: fixChange, newResolvedById: resolvedById, newExcludeSpanIds: removeSpanIds, newExternalIssueKey: externalIssueKey, newQuarantinedById: quarantinedById, newForceClosedById: forceClosedById);
 				if (issue != null)
 				{
 					break;
diff --git a/Engine/Source/Programs/Horde/Horde.Build/Issues/IssuesController.cs b/Engine/Source/Programs/Horde/Horde.Build/Issues/IssuesController.cs
--- a/Engine/Source/Programs/Horde/Horde.Build/Issues/IssuesController.cs
+++ b/Engine/Source/Programs/Horde/Horde.Build/Issues/IssuesController.cs
@@ -315,33 +315,6 @@ namespace Horde.Build.Issues
 			await _issueCollection.GetLogger(issueId).FindAsync(Response.BodyWriter, minTime, maxTime, index, count);
 		}
 
-		/// <summary>
-		/// Hook to allow a Perforce trigger to mark an issue as fixed, via a tag in the changelist description.
-		/// </summary>
-		/// <param name="issueId">Id of the agent to get information about</param>
-		/// <param name="request">Request body</param>
-		/// <returns>Information about the requested agent</returns>
-		[HttpPost]
-		[Route("/api/v1/issues/{issueId}/p4fix")]
-		public async Task<ActionResult> MarkIssueAsFixedViaPerforceAsync(int issueId, [FromBody] MarkFixedViaPerforceRequest request)
-		{
-			if (!await _aclService.AuthorizeAsync(AclAction.IssueFixViaPerforce, User))
-			{
-				return Forbid();
-			}
-
-			IIssueDetails? issue = await _issueService.GetIssueDetailsAsync(issueId);
-			if (issue == null)
-			{
-				return NotFound();
-			}
-
-			IUser user = await _userCollection.FindOrAddUserByLoginAsync(request.UserName);
-			await _issueService.UpdateIssueAsync(issueId, fixChange: request.FixChange, resolvedById: user.Id);
-
-			return Ok();
-		}
-
 		/// <summary>
 		/// Create an issue response object
 		/// </summary>
@@ -625,12 +598,9 @@ namespace Horde.Build.Issues
 				removeSpans = request.RemoveSpans.ConvertAll(x => ObjectId.Parse(x));
 			}
 
-			using (IDisposable scope = _issueCollection.GetLogger(issueId).BeginScope("User {UserId}", User.GetUserId() ?? UserId.Empty))
+			if (!await _issueService.UpdateIssueAsync(issueId, request.Summary, request.Description, request.Promoted, newOwnerId, newNominatedById, request.Acknowledged, newDeclinedById, request.FixChange, newResolvedById, addSpans, removeSpans, request.ExternalIssueKey, newQuarantinedById, newForceClosedById, initiatedById: User.GetUserId()))
 			{
-				if (!await _issueService.UpdateIssueAsync(issueId, request.Summary, request.Description, request.Promoted, newOwnerId, newNominatedById, request.Acknowledged, newDeclinedById, request.FixChange, newResolvedById, addSpans, removeSpans, request.ExternalIssueKey, newQuarantinedById, newForceClosedById))
-				{
-					return NotFound();
-				}
+				return NotFound();
 			}
 			return Ok();
 		}
diff --git a/Engine/Source/Programs/Horde/Horde.Build/Notifications/Sinks/SlackNotificationSink.cs b/Engine/Source/Programs/Horde/Horde.Build/Notifications/Sinks/SlackNotificationSink.cs
--- a/Engine/Source/Programs/Horde/Horde.Build/Notifications/Sinks/SlackNotificationSink.cs
+++ b/Engine/Source/Programs/Horde/Horde.Build/Notifications/Sinks/SlackNotificationSink.cs
@@ -2476,7 +2476,7 @@ namespace Horde.Build.Notifications.Sinks
 									return new { response_action = "errors", errors };
 								}
 
-								await _issueService.UpdateIssueAsync(issueId, fixChange: fixChange, resolvedById: userId);
+								await _issueService.UpdateIssueAsync(issueId, fixChange: fixChange, resolvedById: userId, initiatedById: userId);
 								_logger.LogInformation("Marked issue {IssueId} fixed by user {UserId} in {Change}", issueId, userId, fixChange);
 							}
 						}
@@ -2484,7 +2484,7 @@ namespace Horde.Build.Notifications.Sinks
 						{
 							int issueId = Int32.Parse(match.Groups[1].Value, NumberStyles.Integer, CultureInfo.InvariantCulture);
 							UserId userId = match.Groups[2].Value.ToObjectId<IUser>();
-							await _issueService.UpdateIssueAsync(issueId, acknowledged: true, ownerId: userId);
+							await _issueService.UpdateIssueAsync(issueId, acknowledged: true, ownerId: userId, initiatedById: userId);
 						}
 					}
 				}
@@ -2513,15 +2513,15 @@ namespace Horde.Build.Notifications.Sinks
 
 			if (String.Equals(verb, "ack", StringComparison.Ordinal))
 			{
-				await _issueService.UpdateIssueAsync(issueId, acknowledged: true);
+				await _issueService.UpdateIssueAsync(issueId, acknowledged: true, initiatedById: userId);
 			}
 			else if (String.Equals(verb, "accept", StringComparison.Ordinal))
 			{
-				await _issueService.UpdateIssueAsync(issueId, ownerId: userId, nominatedById: userId, acknowledged: true);
+				await _issueService.UpdateIssueAsync(issueId, ownerId: userId, nominatedById: userId, acknowledged: true, initiatedById: userId);
 			}
 			else if (String.Equals(verb, "decline", StringComparison.Ordinal))
 			{
-				await _issueService.UpdateIssueAsync(issueId, declinedById: userId);
+				await _issueService.UpdateIssueAsync(issueId, declinedById: userId, initiatedById: userId);
 			}
 
 			IIssue? newIssue = await _issueService.Collection.GetIssueAsync(issueId);
@@ -2583,7 +2583,7 @@ namespace Horde.Build.Notifications.Sinks
 					}
 				}
 
-				await _issueService.UpdateIssueAsync(issueId, acknowledged: true, ownerId: user.Id, nominatedById: user.Id);
+				await _issueService.UpdateIssueAsync(issueId, acknowledged: true, ownerId: user.Id, nominatedById: user.Id, initiatedById: user.Id);
 			}
 			else if (String.Equals(verb, "decline", StringComparison.Ordinal))
 			{
@@ -2608,7 +2608,7 @@ namespace Horde.Build.Notifications.Sinks
 					}
 				}
 
-				await _issueService.UpdateIssueAsync(issueId, declinedById: user.Id);
+				await _issueService.UpdateIssueAsync(issueId, declinedById: user.Id, initiatedById: user.Id);
 			}
 			else if (String.Equals(verb, "markfixed", StringComparison.Ordinal))
 			{
diff --git a/Engine/Source/Programs/Horde/Horde.Build/Ugs/UgsController.cs b/Engine/Source/Programs/Horde/Horde.Build/Ugs/UgsController.cs
--- a/Engine/Source/Programs/Horde/Horde.Build/Ugs/UgsController.cs
+++ b/Engine/Source/Programs/Horde/Horde.Build/Ugs/UgsController.cs
@@ -373,7 +373,7 @@ namespace Horde.Build.Ugs
 				newResolvedById = request.Resolved.Value ? IIssue.ResolvedByUnknownId : UserId.Empty;
 			}
 
-			if (!await _issueService.UpdateIssueAsync(issueId, ownerId: newOwnerId, nominatedById: newNominatedById, acknowledged: request.Acknowledged, declinedById: newDeclinedById, fixChange: request.FixChange, resolvedById: newResolvedById))
+			if (!await _issueService.UpdateIssueAsync(issueId, ownerId: newOwnerId, nominatedById: newNominatedById, acknowledged: request.Acknowledged, declinedById: newDeclinedById, fixChange: request.FixChange, resolvedById: newResolvedById, initiatedById: User.GetUserId()))
 			{
 				return NotFound();
 			}
