Horde: Add garbage collection for managed namespaces.

* Frequency can be configured via the GcFrequencyHrs setting for the namespace (defaults to every 2 hours).
* Reachable set of nodes is stored in memory, but queue of nodes to process are paged out to the database to allow other instances to resume an existing GC cycle if a server shuts down.
* Import metadata is also cached in the database, to avoid having to read header data for previously visited nodes.
* Still WIP: does not synchronize newly uploaded nodes with capturing of root set.

diff --git a/Engine/Source/Programs/Horde/Horde.Build.Tests/GcServiceTests.cs b/Engine/Source/Programs/Horde/Horde.Build.Tests/GcServiceTests.cs
new file mode 100644
--- /dev/null
+++ b/Engine/Source/Programs/Horde/Horde.Build.Tests/GcServiceTests.cs
@@ -0,0 +1,117 @@
+// Copyright Epic Games, Inc. All Rights Reserved.
+
+using System;
+using System.Collections.Generic;
+using System.Linq;
+using System.Threading.Tasks;
+using EpicGames.Core;
+using Microsoft.VisualStudio.TestTools.UnitTesting;
+using EpicGames.Horde.Storage;
+using Horde.Build.Storage;
+using System.Threading;
+using Horde.Build.Server;
+using Horde.Build.Utilities;
+using Microsoft.Extensions.DependencyInjection;
+
+namespace Horde.Build.Tests
+{
+	using BackendId = StringId<IStorageBackend>;
+
+	[TestClass]
+	public sealed class GcServiceTests : TestSetup
+	{
+		async Task SetupNamespaceAsync()
+		{
+			GlobalConfig globalConfig = new GlobalConfig();
+			globalConfig.Storage.Backends.Add(new BackendConfig { Id = new BackendId("default-backend"), Type = StorageBackendType.Memory });
+			globalConfig.Storage.Namespaces.Add(new NamespaceConfig { Id = new NamespaceId("default"), Backend = new BackendId("default-backend") });
+			await ConfigCollection.AddConfigAsync("globals", globalConfig);
+			Assert.IsNotNull(await GlobalsService.TryUpdateAsync(await GlobalsService.GetAsync(), "globals"));
+		}
+
+		[TestMethod]
+		public async Task CreateBasicTree()
+		{
+			BlobId.UseDeterministicIds();
+			await SetupNamespaceAsync();
+			IStorageClientImpl store = await StorageService.GetClientAsync(new NamespaceId("default"), CancellationToken.None);
+
+			Random random = new Random(0);
+			BlobLocator[] blobs = await CreateTestDataAsync(store, 2000, 4000, 2000, 10, random);
+
+			HashSet<BlobLocator> roots = new HashSet<BlobLocator>();
+			for (int idx = 0; idx < 10; idx++)
+			{
+				int blobIdx = (int)(random.NextDouble() * blobs.Length);
+				if (roots.Add(blobs[blobIdx]))
+				{
+					await store.WriteRefTargetAsync(new RefName($"ref-{idx}"), new NodeLocator(blobs[blobIdx], 0));
+				}
+			}
+
+			HashSet<BlobLocator> nodes = await FindNodes(store, roots);
+
+			GcService gc = ServiceProvider.GetRequiredService<GcService>();
+			await gc.RunAsync(new NamespaceId("default"), CancellationToken.None);
+
+			BlobLocator[] remaining = await store.Backend.EnumerateAsync().Select(x => new BlobLocator(HostId.Empty, StorageBackend.GetBlobIdFromPath(x))).ToArrayAsync();
+			Assert.AreEqual(nodes.Count, remaining.Length);
+			Assert.IsTrue(remaining.All(x => nodes.Contains(x)));
+		}
+
+		async Task<HashSet<BlobLocator>> FindNodes(IStorageClient store, IEnumerable<BlobLocator> roots)
+		{
+			HashSet<BlobLocator> nodes = new HashSet<BlobLocator>();
+			await FindNodes(store, roots, nodes);
+			return nodes;
+		}
+
+		async Task FindNodes(IStorageClient store, IEnumerable<BlobLocator> roots, HashSet<BlobLocator> nodes)
+		{
+			foreach (BlobLocator root in roots)
+			{
+				if (nodes.Add(root))
+				{
+					Bundle bundle = await store.ReadBundleAsync(root);
+					await FindNodes(store, bundle.Header.Imports.Select(x => x.Locator), nodes);
+				}
+			}
+		}
+
+		static async ValueTask<BlobLocator[]> CreateTestDataAsync(IStorageClient store, int numRoots, int numInterior, int numLeaves, int avgChildren, Random random)
+		{
+			int firstRoot = 0;
+			int firstInterior = firstRoot + numRoots;
+			int firstLeaf = firstInterior + numInterior;
+			int numNodes = firstLeaf + numLeaves;
+
+			List<int>[] children = new List<int>[numNodes];
+			for (int idx = 0; idx < numNodes; idx++)
+			{
+				children[idx] = new List<int>();
+			}
+
+			double maxParents = ((numRoots + numInterior) * avgChildren) / (numInterior + numLeaves);
+			for (int idx = numRoots; idx < numNodes; idx++)
+			{
+				int numParents = 1 + (int)(random.NextDouble() * maxParents);
+				for (; numParents > 0; numParents--)
+				{
+					int parentIdx = (int)(random.NextDouble() * Math.Min(idx, numRoots + numInterior - 1));
+					children[parentIdx].Add(idx);
+				}
+			}
+
+			BlobLocator[] locators = new BlobLocator[children.Length];
+			for (int idx = numNodes - 1; idx >= 0; idx--)
+			{
+				List<BundleImport> imports = children[idx].ConvertAll(x => new BundleImport(locators[x], Array.Empty<(int, IoHash)>()));
+				BundleHeader header = new BundleHeader(BundleCompressionFormat.None, imports, Array.Empty<BundleExport>(), Array.Empty<BundlePacket>());
+				Bundle bundle = new Bundle(header, Array.Empty<ReadOnlyMemory<byte>>());
+				locators[idx] = await store.WriteBundleAsync(bundle, prefix: "gctest");
+			}
+
+			return locators;
+		}
+	}
+}
diff --git a/Engine/Source/Programs/Horde/Horde.Build.Tests/TestSetup.cs b/Engine/Source/Programs/Horde/Horde.Build.Tests/TestSetup.cs
--- a/Engine/Source/Programs/Horde/Horde.Build.Tests/TestSetup.cs
+++ b/Engine/Source/Programs/Horde/Horde.Build.Tests/TestSetup.cs
@@ -123,6 +123,7 @@ namespace Horde.Build.Tests
 		public LifetimeService LifetimeService => ServiceProvider.GetRequiredService<LifetimeService>();
 		public ScheduleService ScheduleService => ServiceProvider.GetRequiredService<ScheduleService>();
 		public DeviceService DeviceService => ServiceProvider.GetRequiredService<DeviceService>();
+		public StorageService StorageService => ServiceProvider.GetRequiredService<StorageService>();
 
 		public ServerSettings ServerSettings => ServiceProvider.GetRequiredService<IOptions<ServerSettings>>().Value;
 		public IOptionsMonitor<ServerSettings> ServerSettingsMon => ServiceProvider.GetRequiredService<IOptionsMonitor<ServerSettings>>();
@@ -247,6 +248,10 @@ namespace Horde.Build.Tests
 
 			services.AddSingleton<ILegacyStorageClient, BasicStorageClient>();
 
+			services.AddSingleton<IRefCollection, RefCollection>();
+			services.AddSingleton<StorageService>();
+			services.AddSingleton<GcService>();
+
 			services.AddSingleton<ISingletonDocument<AgentSoftwareChannels>>(new SingletonDocumentStub<AgentSoftwareChannels>());
 			services.AddSingleton<ISingletonDocument<DevicePlatformMapV1>>(new SingletonDocumentStub<DevicePlatformMapV1>());
 		}
diff --git a/Engine/Source/Programs/Horde/Horde.Build/Startup.cs b/Engine/Source/Programs/Horde/Horde.Build/Startup.cs
--- a/Engine/Source/Programs/Horde/Horde.Build/Startup.cs
+++ b/Engine/Source/Programs/Horde/Horde.Build/Startup.cs
@@ -830,6 +830,21 @@ namespace Horde.Build
 			}
 		}
 
+		public sealed class IoHashBsonSerializer : SerializerBase<IoHash>
+		{
+			/// <inheritdoc/>
+			public override IoHash Deserialize(BsonDeserializationContext context, BsonDeserializationArgs args)
+			{
+				return IoHash.Parse(context.Reader.ReadString());
+			}
+
+			/// <inheritdoc/>
+			public override void Serialize(BsonSerializationContext context, BsonSerializationArgs args, IoHash value)
+			{
+				context.Writer.WriteString(value.ToString());
+			}
+		}
+
 		public sealed class NamespaceIdBsonSerializer : SerializerBase<NamespaceId>
 		{
 			/// <inheritdoc/>
@@ -861,6 +876,7 @@ namespace Horde.Build
 				BsonSerializer.RegisterSerializer(new BlobIdBsonSerializer());
 				BsonSerializer.RegisterSerializer(new RefIdBsonSerializer());
 				BsonSerializer.RegisterSerializer(new RefNameBsonSerializer());
+				BsonSerializer.RegisterSerializer(new IoHashBsonSerializer());
 				BsonSerializer.RegisterSerializer(new NamespaceIdBsonSerializer());
 				BsonSerializer.RegisterSerializer(new ConditionSerializer());
 				BsonSerializer.RegisterSerializationProvider(new BsonSerializationProvider());
diff --git a/Engine/Source/Programs/Horde/Horde.Build/Storage/GcService.cs b/Engine/Source/Programs/Horde/Horde.Build/Storage/GcService.cs
new file mode 100644
--- /dev/null
+++ b/Engine/Source/Programs/Horde/Horde.Build/Storage/GcService.cs
@@ -0,0 +1,500 @@
+// Copyright Epic Games, Inc. All Rights Reserved.
+
+using System;
+using System.Collections.Generic;
+using System.Linq;
+using System.Threading;
+using System.Threading.Tasks;
+using EpicGames.Core;
+using EpicGames.Horde.Storage;
+using Microsoft.Extensions.Logging;
+using Horde.Build.Utilities;
+using MongoDB.Bson;
+using MongoDB.Bson.Serialization.Attributes;
+using Microsoft.Extensions.Hosting;
+using Horde.Build.Server;
+using HordeCommon;
+using EpicGames.Redis.Utility;
+using StackExchange.Redis;
+using MongoDB.Driver;
+using System.Runtime.CompilerServices;
+using System.IO;
+using System.Diagnostics;
+
+namespace Horde.Build.Storage
+{
+	[DebuggerDisplay("Index={BaseIndex}, Count={Hashes.Count}")]
+	class GcPage
+	{
+		public const int MaxItems = 1000;
+
+		[BsonId, BsonIgnoreIfDefault]
+		public ObjectId Id { get; set; }
+
+		[BsonElement("cyc")]
+		public int Cycle { get; set; }
+
+		[BsonElement("base")]
+		public int BaseIndex { get; set; }
+
+		[BsonElement("read")]
+		public int ReadIndex { get; set; }
+
+		[BsonElement("hash")]
+		public List<IoHash> Hashes { get; set; } = new List<IoHash>();
+
+		public int HeadIndex => BaseIndex + Hashes.Count;
+
+		public GcPage()
+		{
+		}
+
+		public GcPage(int cycle, int baseIndex, int readIndex)
+		{
+			Cycle = cycle;
+			BaseIndex = baseIndex;
+			ReadIndex = readIndex;
+		}
+	}
+
+	[DebuggerDisplay("{Id}: {Locator}")]
+	class GcNode
+	{
+		// Hash of the namespace id and blob id, as a string.
+		public IoHash Id { get; set; }
+
+		[BsonElement("l")]
+		public BlobLocator Locator { get; set; }
+
+		[BsonElement("r")]
+		public List<IoHash>? Imports { get; set; }
+	}
+
+	class GcSession
+	{
+		readonly IStorageClientImpl _store;
+		readonly NamespaceId _namespaceId;
+		readonly Utf8String _namespaceIdUtf8;
+		readonly IMongoCollection<GcPage> _pageCollection;
+		readonly IMongoCollection<GcNode> _nodeCollection;
+		readonly int _cycle;
+		readonly ILogger _logger;
+		GcPage _writePage = new GcPage();
+		readonly HashSet<IoHash> _reachable = new HashSet<IoHash>();
+
+		public GcSession(IStorageClientImpl store, NamespaceId namespaceId, IMongoCollection<GcPage> pageCollection, IMongoCollection<GcNode> nodeCollection, int cycle, ILogger logger)
+		{
+			_store = store;
+			_namespaceId = namespaceId;
+			_namespaceIdUtf8 = new Utf8String(namespaceId.ToString());
+			_pageCollection = pageCollection;
+			_nodeCollection = nodeCollection;
+			_cycle = cycle;
+			_logger = logger;
+			_writePage = new GcPage(_cycle, 0, -1);
+		}
+
+		public async Task RecoverAsync(CancellationToken cancellationToken)
+		{
+			int expectedBaseIndex = 0;
+			await foreach (GcPage page in _pageCollection.Find(x => x.Cycle == _cycle, new FindOptions { BatchSize = 1 }).SortBy(x => x.BaseIndex).ToAsyncEnumerable(cancellationToken))
+			{
+				if (page.BaseIndex != expectedBaseIndex)
+				{
+					break;
+				}
+
+				_reachable.UnionWith(page.Hashes);
+				_writePage = new GcPage(_cycle, page.BaseIndex + page.Hashes.Count, page.ReadIndex);
+
+				expectedBaseIndex = page.BaseIndex + page.Hashes.Count;
+			}
+		}
+
+		public async Task RunAsync(CancellationToken cancellationToken)
+		{
+			// Attempt to recover the current state
+			await RecoverAsync(cancellationToken);
+
+			// Find the root set
+			if (_writePage.ReadIndex < 0)
+			{
+				// Add all the nodes to the root set
+				await using (IAsyncEnumerator<NodeLocator> enumerator = _store.Refs.EnumerateRefs(_namespaceId, cancellationToken).GetAsyncEnumerator(cancellationToken))
+				{
+					Task<bool> moveNextTask = enumerator.MoveNextAsync(cancellationToken).AsTask();
+					while (await moveNextTask)
+					{
+						List<WriteModel<GcNode>> requests = new List<WriteModel<GcNode>>();
+						while (requests.Count < 20 && await moveNextTask)
+						{
+							NodeLocator locator = enumerator.Current;
+
+							IoHash hash = GetBlobDataId(locator.Blob.BlobId);
+							if (await WriteHashAsync(hash, cancellationToken))
+							{
+								FilterDefinition<GcNode> filter = Builders<GcNode>.Filter.Eq(x => x.Id, hash);
+								UpdateDefinition<GcNode> update = Builders<GcNode>.Update.SetOnInsert(x => x.Locator, locator.Blob);
+								requests.Add(new UpdateOneModel<GcNode>(filter, update) { IsUpsert = true });
+							}
+
+							moveNextTask = enumerator.MoveNextAsync(cancellationToken).AsTask();
+						}
+
+						if (requests.Count > 0)
+						{
+							await _nodeCollection.BulkWriteAsync(requests, new BulkWriteOptions { IsOrdered = false }, cancellationToken);
+						}
+					}
+				}
+
+				// Write the last page with a read index of zero to indicate that the root set is complete
+				_writePage.ReadIndex = 0;
+				await FlushPageAsync(cancellationToken);
+			}
+
+			// Find the reachable set
+			int readIndex = _writePage.ReadIndex;
+			await foreach (ListSegment<GcNode> batch in ReadNodeBatchesAsync(readIndex, 20, cancellationToken).Prefetch(cancellationToken))
+			{
+				if (batch.Count == 0 && readIndex == _writePage.HeadIndex)
+				{
+					Debug.Assert(readIndex == _writePage.HeadIndex);
+					break;
+				}
+
+				foreach (GcNode node in batch)
+				{
+					foreach (IoHash import in node.Imports!)
+					{
+						await WriteHashAsync(import, cancellationToken);
+					}
+					_writePage.ReadIndex = ++readIndex;
+				}
+			}
+			await FlushPageAsync(cancellationToken);
+
+			// Delete any unreferenced nodes
+			await foreach (string path in _store.Backend.EnumerateAsync(cancellationToken))
+			{
+				BlobId blobId;
+				if (StorageBackend.TryGetBlobIdFromPath(path, out blobId))
+				{
+					IoHash hash = GetBlobDataId(blobId);
+					if (!_reachable.Contains(hash))
+					{
+						await _nodeCollection.DeleteOneAsync(x => x.Id == hash, cancellationToken);
+						await _store.Backend.DeleteAsync(path, cancellationToken);
+					}
+				}
+			}
+		}
+
+		public async IAsyncEnumerable<ListSegment<GcNode>> ReadNodeBatchesAsync(int readIndex, int maxBatchCount, [EnumeratorCancellation] CancellationToken cancellationToken)
+		{
+			await foreach (ListSegment<IoHash> batch in ReadHashBatchesAsync(readIndex, maxBatchCount, cancellationToken))
+			{
+				List<GcNode> nodes = await _nodeCollection.Find(Builders<GcNode>.Filter.In(x => x.Id, batch)).ToListAsync(cancellationToken);
+				Dictionary<IoHash, GcNode> hashToNode = nodes.ToDictionary(x => x.Id, x => x);
+
+				List<GcNode> results = new List<GcNode>();
+				foreach (IoHash hash in batch)
+				{
+					GcNode node = hashToNode[hash];
+					if (node.Imports == null)
+					{
+						await FindImportsAsync(node, cancellationToken);
+					}
+					results.Add(node);
+				}
+
+				yield return results;
+			}
+		}
+
+		public async IAsyncEnumerable<ListSegment<IoHash>> ReadHashBatchesAsync(int readIndex, int maxBatchCount, [EnumeratorCancellation] CancellationToken cancellationToken)
+		{
+			for (; ; )
+			{
+				GcPage page = await ReadPageAsync(readIndex, cancellationToken);
+				if (readIndex == page.HeadIndex)
+				{
+					yield return ListSegment.Empty<IoHash>();
+				}
+				else
+				{
+					while (readIndex < page.HeadIndex)
+					{
+						int start = readIndex - page.BaseIndex;
+						int count = Math.Min(maxBatchCount, page.Hashes.Count - start);
+						yield return page.Hashes.AsSegment(start, count);
+						readIndex += count;
+					}
+				}
+			}
+		}
+
+		async ValueTask<GcPage> ReadPageAsync(int readIndex, CancellationToken cancellationToken)
+		{
+			cancellationToken.ThrowIfCancellationRequested();
+
+			GcPage page;
+			if (readIndex >= _writePage.BaseIndex)
+			{
+				page = _writePage;
+			}
+			else
+			{
+				page = await _pageCollection.Find(x => x.Cycle == _cycle && x.BaseIndex <= readIndex).SortByDescending(x => x.BaseIndex).FirstAsync(cancellationToken);
+			}
+
+			int pageBaseIndex = page.BaseIndex;
+			if (readIndex > page.BaseIndex + page.Hashes.Count)
+			{
+				throw new InvalidDataException();
+			}
+			return page;
+		}
+
+		async Task FindImportsAsync(GcNode node, CancellationToken cancellationToken)
+		{
+			List<WriteModel<GcNode>> updates = new List<WriteModel<GcNode>>();
+			node.Imports = new List<IoHash>();
+
+			BundleHeader? header = await _store.ReadBundleHeaderAsync(node.Locator, cancellationToken);
+			if (header != null)
+			{
+				foreach (BundleImport import in header.Imports)
+				{
+					IoHash importId = GetBlobDataId(import.Locator.BlobId);
+					node.Imports.Add(importId);
+
+					if (!_reachable.Contains(importId))
+					{
+						FilterDefinition<GcNode> filter = Builders<GcNode>.Filter.Eq(x => x.Id, importId);
+						UpdateDefinition<GcNode> update = Builders<GcNode>.Update.SetOnInsert(x => x.Locator, import.Locator);
+						updates.Add(new UpdateOneModel<GcNode>(filter, update) { IsUpsert = true });
+					}
+				}
+			}
+
+			updates.Add(new UpdateOneModel<GcNode>(Builders<GcNode>.Filter.Eq(x => x.Id, node.Id), Builders<GcNode>.Update.Set(x => x.Imports, node.Imports)));
+
+			await _nodeCollection.BulkWriteAsync(updates, new BulkWriteOptions { IsOrdered = false }, cancellationToken);
+		}
+
+		async ValueTask<bool> WriteHashAsync(IoHash hash, CancellationToken cancellationToken)
+		{
+			if (_reachable.Add(hash))
+			{
+				_writePage.Hashes.Add(hash);
+				if (_writePage.Hashes.Count == GcPage.MaxItems)
+				{
+					await FlushPageAsync(cancellationToken);
+				}
+				return true;
+			}
+			return false;
+		}
+
+		async Task FlushPageAsync(CancellationToken cancellationToken)
+		{
+			await _pageCollection.ReplaceOneAsync(x => x.Cycle == _writePage.Cycle && x.BaseIndex == _writePage.BaseIndex, _writePage, new ReplaceOptions { IsUpsert = true }, cancellationToken);
+			if (_writePage.Hashes.Count > 0)
+			{
+				_writePage = new GcPage(_cycle, _writePage.BaseIndex + _writePage.Hashes.Count, _writePage.ReadIndex);
+			}
+		}
+
+		public static IoHash GetBlobDataId(Utf8String namespaceIdUtf8, BlobId blobId)
+		{
+			Span<byte> buffer = stackalloc byte[namespaceIdUtf8.Length + 1 + blobId.Inner.Length];
+			namespaceIdUtf8.Span.CopyTo(buffer);
+			blobId.Inner.Span.CopyTo(buffer.Slice(namespaceIdUtf8.Length + 1));
+			return IoHash.Compute(buffer);
+		}
+
+		IoHash GetBlobDataId(BlobId blobId) => GetBlobDataId(_namespaceIdUtf8, blobId);
+	}
+
+	class GcService : IHostedService
+	{
+		[SingletonDocument("gc")]
+		class GcState : SingletonBase
+		{
+			public int NextCycle { get; set; } = 1;
+			public List<GcNamespaceState> Namespaces { get; set; } = new List<GcNamespaceState>();
+		}
+
+		class GcNamespaceState
+		{
+			public NamespaceId Id { get; set; }
+			public int Cycle { get; set; }
+			public DateTime LastTime { get; set; }
+		}
+
+		readonly MongoService _mongoService;
+		readonly RedisService _redisService;
+		readonly StorageService _storageService;
+		readonly SingletonDocument<GcState> _gcState;
+		readonly ITicker _ticker;
+		readonly ILogger _logger;
+		readonly IMongoCollection<GcNode> _nodeCollection;
+		readonly IMongoCollection<GcPage> _pageCollection;
+
+		public GcService(MongoService mongoService, RedisService redisService, StorageService storageService, IClock clock, ILogger<GcService> logger)
+		{
+			_mongoService = mongoService;
+			_redisService = redisService;
+			_storageService = storageService;
+			_gcState = new SingletonDocument<GcState>(mongoService);
+			_ticker = clock.AddTicker<GcService>(TimeSpan.FromMinutes(5.0), TickAsync, logger);
+			_logger = logger;
+			_nodeCollection = _mongoService.GetCollection<GcNode>("Storage.GcNodes");
+			_pageCollection = _mongoService.GetCollection<GcPage>("Storage.GcPages", builder => builder.Ascending(x => x.Cycle).Ascending(x => x.BaseIndex), unique: true);
+		}
+
+		/// <inheritdoc/>
+		public async Task StartAsync(CancellationToken cancellationToken)
+		{
+			await _ticker.StartAsync();
+		}
+
+		/// <inheritdoc/>
+		public async Task StopAsync(CancellationToken cancellationToken)
+		{
+			await _ticker.StopAsync();
+		}
+
+		/// <summary>
+		/// Find the next namespace to run GC on
+		/// </summary>
+		/// <param name="cancellationToken"></param>
+		/// <returns></returns>
+		async ValueTask TickAsync(CancellationToken cancellationToken)
+		{
+			HashSet<NamespaceId> ranNamespaceIds = new HashSet<NamespaceId>();
+
+			DateTime utcNow = DateTime.UtcNow;
+			for (; ; )
+			{
+				// Synchronize the list of configured namespaces with the GC state object
+				List<NamespaceConfig> namespaces = await _storageService.GetNamespacesAsync(cancellationToken);
+
+				GcState state = await _gcState.GetAsync();
+				if (!Enumerable.SequenceEqual(namespaces.Select(x => x.Id).OrderBy(x => x), state.Namespaces.Select(x => x.Id).OrderBy(x => x)))
+				{
+					state = await _gcState.UpdateAsync(s => SyncNamespaceList(s, namespaces));
+				}
+
+				// Find all the namespaces that need to have GC run on them
+				List<(DateTime, GcNamespaceState)> pending = new List<(DateTime, GcNamespaceState)>();
+				foreach (GcNamespaceState namespaceState in state.Namespaces)
+				{
+					NamespaceConfig? config = namespaces.FirstOrDefault(x => x.Id == namespaceState.Id);
+					if (config != null)
+					{
+						DateTime time = namespaceState.LastTime + TimeSpan.FromHours(config.GcFrequencyHrs);
+						if (time > utcNow)
+						{
+							pending.Add((time, namespaceState));
+						}
+					}
+				}
+				pending.SortBy(x => x.Item1);
+
+				// Update the first one we can acquire a lock for
+				foreach ((_, GcNamespaceState namespaceState) in pending)
+				{
+					NamespaceId namespaceId = namespaceState.Id;
+					if (ranNamespaceIds.Add(namespaceId))
+					{
+						RedisKey key = $"gc/{namespaceState.Id}";
+						using (RedisLock namespaceLock = new RedisLock(_redisService.GetDatabase(), key))
+						{
+							if (await namespaceLock.AcquireAsync(TimeSpan.FromMinutes(20.0)))
+							{
+								try
+								{
+									await RunAsync(namespaceState, cancellationToken);
+								}
+								catch (Exception ex)
+								{
+									_logger.LogError(ex, "Exception while running garbage collection: {Message}", ex.Message);
+								}
+								break;
+							}
+						}
+					}
+				}
+			}
+		}
+
+		public async Task RunAsync(NamespaceId namespaceId, CancellationToken cancellationToken)
+		{
+			GcState state = await _gcState.GetAsync();
+			await RunAsync(FindOrAddNamespace(state, namespaceId), cancellationToken);
+		}
+
+		async Task RunAsync(GcNamespaceState namespaceState, CancellationToken cancellationToken)
+		{
+			NamespaceId namespaceId = namespaceState.Id;
+
+			int cycle = namespaceState.Cycle;
+			if (cycle == 0)
+			{
+				await _gcState.UpdateAsync(x => cycle = StartCycle(x, namespaceId));
+			}
+
+			IStorageClientImpl store = await _storageService.GetClientAsync(namespaceId, cancellationToken);
+			GcSession session = new GcSession(store, namespaceId, _pageCollection, _nodeCollection, cycle, _logger);
+			await session.RunAsync(cancellationToken);
+			await _gcState.UpdateAsync(x => CompleteCycle(x, namespaceId));
+			await _pageCollection.DeleteManyAsync(x => x.Cycle == cycle, cancellationToken);
+		}
+
+		static void SyncNamespaceList(GcState state, List<NamespaceConfig> namespaces)
+		{
+			HashSet<NamespaceId> validNamespaceIds = new HashSet<NamespaceId>(namespaces.Select(x => x.Id));
+			state.Namespaces.RemoveAll(x => !validNamespaceIds.Contains(x.Id));
+
+			HashSet<NamespaceId> currentNamespaceIds = new HashSet<NamespaceId>(state.Namespaces.Select(x => x.Id));
+			foreach (NamespaceConfig config in namespaces)
+			{
+				if (!currentNamespaceIds.Contains(config.Id))
+				{
+					state.Namespaces.Add(new GcNamespaceState { Id = config.Id, LastTime = DateTime.UtcNow });
+				}
+			}
+
+			state.Namespaces.SortBy(x => x.Id);
+		}
+
+		static int StartCycle(GcState state, NamespaceId namespaceId)
+		{
+			GcNamespaceState namespaceState = FindOrAddNamespace(state, namespaceId);
+			namespaceState.Cycle = ++state.NextCycle;
+			return namespaceState.Cycle;
+		}
+
+		static void CompleteCycle(GcState state, NamespaceId namespaceId)
+		{
+			GcNamespaceState namespaceState = FindOrAddNamespace(state, namespaceId);
+			namespaceState.Cycle = 0;
+			namespaceState.LastTime = DateTime.UtcNow;
+		}
+
+		static GcNamespaceState FindOrAddNamespace(GcState state, NamespaceId namespaceId)
+		{
+			GcNamespaceState? namespaceState = state.Namespaces.FirstOrDefault(x => x.Id == namespaceId);
+			if (namespaceState == null)
+			{
+				namespaceState = new GcNamespaceState();
+				namespaceState.Id = namespaceId;
+				namespaceState.Cycle = ++state.NextCycle;
+				state.Namespaces.Add(namespaceState);
+			}
+			return namespaceState;
+		}
+	}
+}
diff --git a/Engine/Source/Programs/Horde/Horde.Build/Storage/IRefCollection.cs b/Engine/Source/Programs/Horde/Horde.Build/Storage/IRefCollection.cs
--- a/Engine/Source/Programs/Horde/Horde.Build/Storage/IRefCollection.cs
+++ b/Engine/Source/Programs/Horde/Horde.Build/Storage/IRefCollection.cs
@@ -1,6 +1,7 @@
 // Copyright Epic Games, Inc. All Rights Reserved.
 
 using System;
+using System.Collections.Generic;
 using System.Threading;
 using System.Threading.Tasks;
 using EpicGames.Horde.Storage;
@@ -38,5 +39,12 @@ namespace Horde.Build.Storage
 		/// <param name="name">The ref name</param>
 		/// <param name="cancellationToken">Cancellation token for the operation</param>
 		Task DeleteRefAsync(NamespaceId namespaceId, RefName name, CancellationToken cancellationToken = default);
+
+		/// <summary>
+		/// Enumerate refs in a particular namespace
+		/// </summary>
+		/// <param name="namespaceId">Namespace for the ref</param>
+		/// <param name="cancellationToken">Cancellation token for the operation</param>
+		IAsyncEnumerable<NodeLocator> EnumerateRefs(NamespaceId namespaceId, CancellationToken cancellationToken = default);
 	}
 }
diff --git a/Engine/Source/Programs/Horde/Horde.Build/Storage/IStorageBackend.cs b/Engine/Source/Programs/Horde/Horde.Build/Storage/IStorageBackend.cs
--- a/Engine/Source/Programs/Horde/Horde.Build/Storage/IStorageBackend.cs
+++ b/Engine/Source/Programs/Horde/Horde.Build/Storage/IStorageBackend.cs
@@ -6,6 +6,7 @@ using System.IO;
 using System.Threading;
 using System.Threading.Tasks;
 using EpicGames.Core;
+using EpicGames.Horde.Storage;
 
 namespace Horde.Build.Storage
 {
@@ -75,13 +76,13 @@ namespace Horde.Build.Storage
 	/// <summary>
 	/// Extension methods for <see cref="IStorageBackend"/>
 	/// </summary>
-	public static class StorageBackendExtensions
+	public static class StorageBackend
 	{
 		/// <summary>
 		/// Wrapper for <see cref="IStorageBackend"/>
 		/// </summary>
 		/// <typeparam name="T"></typeparam>
-		class StorageBackend<T> : IStorageBackend<T>
+		class TypedStorageBackend<T> : IStorageBackend<T>
 		{
 			readonly IStorageBackend _inner;
 
@@ -89,7 +90,7 @@ namespace Horde.Build.Storage
 			/// Constructor
 			/// </summary>
 			/// <param name="inner"></param>
-			public StorageBackend(IStorageBackend inner) => _inner = inner;
+			public TypedStorageBackend(IStorageBackend inner) => _inner = inner;
 
 			/// <inheritdoc/>
 			public Task<Stream?> TryReadAsync(string path, CancellationToken cancellationToken) => _inner.TryReadAsync(path, cancellationToken);
@@ -110,6 +111,53 @@ namespace Horde.Build.Storage
 			public IAsyncEnumerable<string> EnumerateAsync(CancellationToken cancellationToken = default) => _inner.EnumerateAsync(cancellationToken);
 		}
 
+		/// <summary>
+		/// Extension for blob files
+		/// </summary>
+		public const string BlobExtension = ".blob";
+
+		/// <summary>
+		/// Gets the blob id from a path
+		/// </summary>
+		/// <param name="path">Path to the blob</param>
+		/// <returns>Path to the blob</returns>
+		public static BlobId GetBlobIdFromPath(string path)
+		{
+			BlobId blobId;
+			if (!TryGetBlobIdFromPath(path, out blobId))
+			{
+				throw new ArgumentException("Path is not a valid blob identifier", nameof(path));
+			}
+			return blobId;
+		}
+
+		/// <summary>
+		/// Gets the path to a blob
+		/// </summary>
+		/// <param name="blobId">Blob identifier</param>
+		/// <returns>Path to the blob</returns>
+		public static string GetBlobPath(BlobId blobId) => $"{blobId}{BlobExtension}";
+
+		/// <summary>
+		/// Gets a blob id from a path within the storage backend
+		/// </summary>
+		/// <param name="path">Path to the file</param>
+		/// <param name="blobId">Receives the blob id on success</param>
+		/// <returns>True on success</returns>
+		public static bool TryGetBlobIdFromPath(string path, out BlobId blobId)
+		{
+			if (path.EndsWith(BlobExtension, StringComparison.Ordinal))
+			{
+				blobId = new BlobId(path.Substring(0, path.Length - BlobExtension.Length));
+				return true;
+			}
+			else
+			{
+				blobId = default;
+				return false;
+			}
+		}
+
 		/// <summary>
 		/// Creates a typed wrapper around the given storage backend
 		/// </summary>
@@ -118,7 +166,43 @@ namespace Horde.Build.Storage
 		/// <returns></returns>
 		public static IStorageBackend<T> ForType<T>(this IStorageBackend backend)
 		{
-			return new StorageBackend<T>(backend);
+			return new TypedStorageBackend<T>(backend);
+		}
+
+		/// <summary>
+		/// Writes a block of memory to storage
+		/// </summary>
+		/// <param name="storageBackend"></param>
+		/// <param name="path"></param>
+		/// <param name="cancellationToken"></param>
+		/// <returns></returns>
+		public static async Task<Stream> ReadAsync(this IStorageBackend storageBackend, string path, CancellationToken cancellationToken = default)
+		{
+			Stream? stream = await storageBackend.TryReadAsync(path, cancellationToken);
+			if (stream == null)
+			{
+				throw new FileNotFoundException($"Unable to read from path {path}");
+			}
+			return stream;
+		}
+
+		/// <summary>
+		/// Writes a block of memory to storage
+		/// </summary>
+		/// <param name="storageBackend"></param>
+		/// <param name="path"></param>
+		/// <param name="offset">Offset of the data to read</param>
+		/// <param name="length">Length of the data to read</param>
+		/// <param name="cancellationToken"></param>
+		/// <returns></returns>
+		public static async Task<Stream> ReadAsync(this IStorageBackend storageBackend, string path, int offset, int length, CancellationToken cancellationToken = default)
+		{
+			Stream? stream = await storageBackend.TryReadAsync(path, offset, length, cancellationToken);
+			if (stream == null)
+			{
+				throw new FileNotFoundException($"Unable to read from path {path}");
+			}
+			return stream;
 		}
 
 		/// <summary>
diff --git a/Engine/Source/Programs/Horde/Horde.Build/Storage/RefCollection.cs b/Engine/Source/Programs/Horde/Horde.Build/Storage/RefCollection.cs
--- a/Engine/Source/Programs/Horde/Horde.Build/Storage/RefCollection.cs
+++ b/Engine/Source/Programs/Horde/Horde.Build/Storage/RefCollection.cs
@@ -1,6 +1,8 @@
 // Copyright Epic Games, Inc. All Rights Reserved.
 
 using System;
+using System.Collections.Generic;
+using System.Runtime.CompilerServices;
 using System.Threading;
 using System.Threading.Tasks;
 using EpicGames.Horde.Storage;
@@ -166,5 +168,20 @@ namespace Horde.Build.Storage
 			
 			AddRefToCache(namespaceId, name, target);
 		}
+
+		/// <inheritdoc/>
+		public async IAsyncEnumerable<NodeLocator> EnumerateRefs(NamespaceId namespaceId, [EnumeratorCancellation] CancellationToken cancellationToken = default)
+		{
+			using (IAsyncCursor<RefDocument> cursor = await _refs.Find(x => x.NamespaceId == namespaceId).ToCursorAsync(cancellationToken))
+			{
+				while (await cursor.MoveNextAsync(cancellationToken))
+				{
+					foreach (RefDocument refDoc in cursor.Current)
+					{
+						yield return new NodeLocator(refDoc.Blob, refDoc.ExportIdx);
+					}
+				}
+			}
+		}
 	}
 }
diff --git a/Engine/Source/Programs/Horde/Horde.Build/Storage/StorageConfig.cs b/Engine/Source/Programs/Horde/Horde.Build/Storage/StorageConfig.cs
--- a/Engine/Source/Programs/Horde/Horde.Build/Storage/StorageConfig.cs
+++ b/Engine/Source/Programs/Horde/Horde.Build/Storage/StorageConfig.cs
@@ -107,6 +107,11 @@ namespace Horde.Build.Storage
 		[Required]
 		public BackendId Backend { get; set; }
 
+		/// <summary>
+		/// How frequently to run garbage collection, in hours.
+		/// </summary>
+		public double GcFrequencyHrs { get; set; } = 2.0;
+
 		/// <summary>
 		/// Access list for this namespace
 		/// </summary>
diff --git a/Engine/Source/Programs/Horde/Horde.Build/Storage/StorageService.cs b/Engine/Source/Programs/Horde/Horde.Build/Storage/StorageService.cs
--- a/Engine/Source/Programs/Horde/Horde.Build/Storage/StorageService.cs
+++ b/Engine/Source/Programs/Horde/Horde.Build/Storage/StorageService.cs
@@ -5,11 +5,13 @@ using System.Buffers;
 using System.Collections;
 using System.Collections.Generic;
 using System.IO;
+using System.Linq;
 using System.Reflection;
 using System.Security.Claims;
 using System.Text.Json;
 using System.Threading;
 using System.Threading.Tasks;
+using Amazon.EC2.Model;
 using EpicGames.Core;
 using EpicGames.Horde.Storage;
 using Horde.Build.Acls;
@@ -39,31 +41,53 @@ namespace Horde.Build.Storage
 		}
 	}
 
+	/// <summary>
+	/// Interface for storage clients which includes a backend implementation. Some functionality is exposed through the backend which is not part of the regular storage API (eg. enumerating).
+	/// </summary>
+	public interface IStorageClientImpl : IStorageClient
+	{
+		/// <summary>
+		/// Configuration for this namespace
+		/// </summary>
+		public NamespaceConfig Config { get; }
+
+		/// <summary>
+		/// The storage backend
+		/// </summary>
+		public IStorageBackend Backend { get; }
+
+		/// <summary>
+		/// The ref collection
+		/// </summary>
+		public IRefCollection Refs { get; }
+	}
+
 	/// <summary>
 	/// Functionality related to the storage service
 	/// </summary>
 	public sealed class StorageService : IDisposable
 	{
-		sealed class StorageClient : StorageClientBase, IDisposable
+		sealed class StorageClient : StorageClientBase, IStorageClientImpl, IDisposable
 		{
 			public NamespaceConfig Config { get; }
 
+			public IStorageBackend Backend { get; }
+
 			NamespaceId NamespaceId => Config.Id;
 
-			readonly IStorageBackend _backend;
-			readonly IRefCollection _refs;
+			public IRefCollection Refs { get; }
 
 			public StorageClient(NamespaceConfig config, IStorageBackend backend, IRefCollection refs, IMemoryCache cache, ILogger logger)
 				: base(cache, logger)
 			{
 				Config = config;
-				_backend = backend;
-				_refs = refs;
+				Backend = backend;
+				Refs = refs;
 			}
 
 			public void Dispose()
 			{
-				if (_backend is IDisposable disposable)
+				if (Backend is IDisposable disposable)
 				{
 					disposable.Dispose();
 				}
@@ -78,7 +102,7 @@ namespace Horde.Build.Storage
 			{
 				string path = GetBlobPath(locator.BlobId);
 
-				Stream? stream = await _backend.TryReadAsync(path, cancellationToken);
+				Stream? stream = await Backend.TryReadAsync(path, cancellationToken);
 				if (stream == null)
 				{
 					throw new StorageException($"Unable to read data from {path}");
@@ -92,7 +116,7 @@ namespace Horde.Build.Storage
 			{
 				string path = GetBlobPath(locator.BlobId);
 
-				Stream? stream = await _backend.TryReadAsync(path, offset, length, cancellationToken);
+				Stream? stream = await Backend.TryReadAsync(path, offset, length, cancellationToken);
 				if (stream == null)
 				{
 					throw new StorageException($"Unable to read data from {path}");
@@ -107,7 +131,7 @@ namespace Horde.Build.Storage
 				BlobId id = BlobId.CreateNew(prefix);
 
 				string path = GetBlobPath(id);
-				await _backend.WriteAsync(path, stream, cancellationToken);
+				await Backend.WriteAsync(path, stream, cancellationToken);
 
 				return new BlobLocator(HostId.Empty, id);
 			}
@@ -117,13 +141,13 @@ namespace Horde.Build.Storage
 			#region Refs
 
 			/// <inheritdoc/>
-			public override Task<NodeLocator> TryReadRefTargetAsync(RefName name, DateTime cacheTime = default, CancellationToken cancellationToken = default) => _refs.TryReadRefTargetAsync(NamespaceId, name, cacheTime, cancellationToken);
+			public override Task<NodeLocator> TryReadRefTargetAsync(RefName name, DateTime cacheTime = default, CancellationToken cancellationToken = default) => Refs.TryReadRefTargetAsync(NamespaceId, name, cacheTime, cancellationToken);
 
 			/// <inheritdoc/>
-			public override Task WriteRefTargetAsync(RefName name, NodeLocator target, CancellationToken cancellationToken = default) => _refs.WriteRefTargetAsync(NamespaceId, name, target, cancellationToken);
+			public override Task WriteRefTargetAsync(RefName name, NodeLocator target, CancellationToken cancellationToken = default) => Refs.WriteRefTargetAsync(NamespaceId, name, target, cancellationToken);
 
 			/// <inheritdoc/>
-			public override Task DeleteRefAsync(RefName name, CancellationToken cancellationToken = default) => _refs.DeleteRefAsync(NamespaceId, name, cancellationToken);
+			public override Task DeleteRefAsync(RefName name, CancellationToken cancellationToken = default) => Refs.DeleteRefAsync(NamespaceId, name, cancellationToken);
 
 			#endregion
 		}
@@ -131,10 +155,10 @@ namespace Horde.Build.Storage
 		class NamespaceInfo : IDisposable
 		{
 			public NamespaceConfig Config { get; }
-			public IStorageClient Client { get; }
+			public IStorageClientImpl Client { get; }
 			public Acl? Acl { get; }
 
-			public NamespaceInfo(NamespaceConfig config, IStorageClient client)
+			public NamespaceInfo(NamespaceConfig config, IStorageClientImpl client)
 			{
 				Config = config;
 				Client = client;
@@ -211,13 +235,24 @@ namespace Horde.Build.Storage
 			return namespaceInfo;
 		}
 
+		/// <summary>
+		/// Finds the configuration for all current namespaces
+		/// </summary>
+		/// <param name="cancellationToken">Cancellation token for the operation</param>
+		/// <returns>List of namespace configurations</returns>
+		public async Task<List<NamespaceConfig>> GetNamespacesAsync(CancellationToken cancellationToken)
+		{
+			State state = await _cachedState.GetAsync(cancellationToken);
+			return state.Namespaces.Select(x => x.Value.Config).ToList();
+		}
+
 		/// <summary>
 		/// Gets a storage client for the given namespace
 		/// </summary>
 		/// <param name="namespaceId">Namespace identifier</param>
 		/// <param name="cancellationToken">Cancellation token for the operation</param>
 		/// <returns></returns>
-		public async ValueTask<IStorageClient> GetClientAsync(NamespaceId namespaceId, CancellationToken cancellationToken)
+		public async ValueTask<IStorageClientImpl> GetClientAsync(NamespaceId namespaceId, CancellationToken cancellationToken)
 		{
 			NamespaceInfo namespaceInfo = await GetNamespaceInfoAsync(namespaceId, cancellationToken);
 			return namespaceInfo.Client;
diff --git a/Engine/Source/Programs/Horde/Horde.Build/Utilities/MongoExtensions.cs b/Engine/Source/Programs/Horde/Horde.Build/Utilities/MongoExtensions.cs
--- a/Engine/Source/Programs/Horde/Horde.Build/Utilities/MongoExtensions.cs
+++ b/Engine/Source/Programs/Horde/Horde.Build/Utilities/MongoExtensions.cs
@@ -4,7 +4,9 @@ using System;
 using System.Collections.Generic;
 using System.Linq;
 using System.Linq.Expressions;
+using System.Runtime.CompilerServices;
 using System.Text.Json;
+using System.Threading;
 using System.Threading.Tasks;
 using MongoDB.Bson;
 using MongoDB.Driver;
@@ -147,13 +149,14 @@ namespace Horde.Build.Utilities
 		/// <typeparam name="TDocument"></typeparam>
 		/// <param name="collection">Collection to insert into</param>
 		/// <param name="newDocuments">The document to insert</param>
-		public static async Task InsertManyIgnoreDuplicatesAsync<TDocument>(this IMongoCollection<TDocument> collection, List<TDocument> newDocuments)
+		/// <param name="cancellationToken">Cancellation token for the operation</param>
+		public static async Task InsertManyIgnoreDuplicatesAsync<TDocument>(this IMongoCollection<TDocument> collection, List<TDocument> newDocuments, CancellationToken cancellationToken = default)
 		{
 			try
 			{
 				if (newDocuments.Count > 0)
 				{
-					await collection.InsertManyAsync(newDocuments, new InsertManyOptions { IsOrdered = false });
+					await collection.InsertManyAsync(newDocuments, new InsertManyOptions { IsOrdered = false }, cancellationToken);
 				}
 			}
 			catch (MongoWriteException ex)
@@ -289,6 +292,28 @@ namespace Horde.Build.Utilities
 		}
 #pragma warning restore IDE0060
 
+		/// <summary>
+		/// Creates an async enumerator from the given query
+		/// </summary>
+		/// <typeparam name="T">Document type</typeparam>
+		/// <param name="source">Query source</param>
+		/// <param name="cancellationToken">Cancellation token for the operation</param>
+		/// <returns></returns>
+		public static async IAsyncEnumerable<T> ToAsyncEnumerable<T>(this IAsyncCursorSource<T> source, [EnumeratorCancellation] CancellationToken cancellationToken)
+		{
+			using (IAsyncCursor<T> cursor = await source.ToCursorAsync(cancellationToken))
+			{
+				while (await cursor.MoveNextAsync(cancellationToken))
+				{
+					foreach (T value in cursor.Current)
+					{
+						yield return value;
+						cancellationToken.ThrowIfCancellationRequested();
+					}
+				}
+			}
+		}
+
 		/// <summary>
 		/// Converts a JsonElement to a BsonValue
 		/// </summary>
diff --git a/Engine/Source/Programs/Shared/EpicGames.Core/AsyncUtils.cs b/Engine/Source/Programs/Shared/EpicGames.Core/AsyncUtils.cs
--- a/Engine/Source/Programs/Shared/EpicGames.Core/AsyncUtils.cs
+++ b/Engine/Source/Programs/Shared/EpicGames.Core/AsyncUtils.cs
@@ -2,6 +2,7 @@
 
 using System;
 using System.Collections.Generic;
+using System.Runtime.CompilerServices;
 using System.Threading;
 using System.Threading.Tasks;
 
@@ -140,5 +141,63 @@ namespace EpicGames.Core
 
 			return results;
 		}
+
+		/// <summary>
+		/// Starts prefetching the next item from an async enumerator while the current one is being processes
+		/// </summary>
+		/// <typeparam name="T">Value type</typeparam>
+		/// <param name="source">Sequence to enumerate</param>
+		/// <param name="cancellationToken">Cancellation token for the operation</param>
+		/// <returns></returns>
+		public static async IAsyncEnumerable<T> Prefetch<T>(this IAsyncEnumerable<T> source, [EnumeratorCancellation] CancellationToken cancellationToken)
+		{
+			await using IAsyncEnumerator<T> enumerator = source.GetAsyncEnumerator(cancellationToken);
+			if (await enumerator.MoveNextAsync())
+			{
+				T value = enumerator.Current;
+
+				for (; ; )
+				{
+					cancellationToken.ThrowIfCancellationRequested();
+
+					Task<bool> task = enumerator.MoveNextAsync().AsTask();
+					try
+					{
+						yield return value;
+					}
+					finally
+					{
+						await task; // Async state machine throws a NotSupportedException if disposed before awaiting this task
+					}
+
+					if (!await task)
+					{
+						break;
+					}
+
+					value = enumerator.Current;
+				}
+			}
+		}
+
+		/// <summary>
+		/// Starts prefetching a number of items from an async enumerator while the current one is being processes
+		/// </summary>
+		/// <typeparam name="T">Value type</typeparam>
+		/// <param name="source">Sequence to enumerate</param>
+		/// <param name="count">Number of items to prefetch</param>
+		/// <param name="cancellationToken">Cancellation token for the operation</param>
+		/// <returns></returns>
+		public static IAsyncEnumerable<T> Prefetch<T>(this IAsyncEnumerable<T> source, int count, CancellationToken cancellationToken = default)
+		{
+			if (count == 0)
+			{
+				return source;
+			}
+			else
+			{
+				return Prefetch(source, count - 1, cancellationToken);
+			}
+		}
 	}
 }
diff --git a/Engine/Source/Programs/Shared/EpicGames.Core/ListSegment.cs b/Engine/Source/Programs/Shared/EpicGames.Core/ListSegment.cs
--- a/Engine/Source/Programs/Shared/EpicGames.Core/ListSegment.cs
+++ b/Engine/Source/Programs/Shared/EpicGames.Core/ListSegment.cs
@@ -16,7 +16,7 @@ namespace EpicGames.Core
 		/// <summary>
 		/// An empty list segment of the given element type
 		/// </summary>
-		public static ListSegment<T> Empty<T>() => new ListSegment<T>(null!, 0, 0);
+		public static ListSegment<T> Empty<T>() => new ListSegment<T>(Array.Empty<T>(), 0, 0);
 	}
 
 	/// <summary>
diff --git a/Engine/Source/Programs/Shared/EpicGames.Core/StreamUtils.cs b/Engine/Source/Programs/Shared/EpicGames.Core/StreamUtils.cs
--- a/Engine/Source/Programs/Shared/EpicGames.Core/StreamUtils.cs
+++ b/Engine/Source/Programs/Shared/EpicGames.Core/StreamUtils.cs
@@ -12,6 +12,27 @@ namespace EpicGames.Core
 	/// </summary>
 	public static class StreamUtils
 	{
+		/// <summary>
+		/// Issues read calls until the given buffer has been filled with data or the stream is at an end
+		/// </summary>
+		/// <param name="stream">Stream to read from</param>
+		/// <param name="buffer">Buffer to receive the output data</param>
+		/// <param name="cancellationToken">Cancellation token for the operation</param>
+		public static async Task<int> ReadGreedyAsync(this Stream stream, Memory<byte> buffer, CancellationToken cancellationToken = default)
+		{
+			int length = 0;
+			while (length < buffer.Length)
+			{
+				int readLength = await stream.ReadAsync(buffer.Slice(length), cancellationToken);
+				if (readLength == 0)
+				{
+					break;
+				}
+				length += readLength;
+			}
+			return length;
+		}
+
 		/// <summary>
 		/// Reads a fixed amount of data from a stream, throwing an exception if the entire buffer cannot be read.
 		/// </summary>
diff --git a/Engine/Source/Programs/Shared/EpicGames.Horde/Storage/BlobId.cs b/Engine/Source/Programs/Shared/EpicGames.Horde/Storage/BlobId.cs
--- a/Engine/Source/Programs/Shared/EpicGames.Horde/Storage/BlobId.cs
+++ b/Engine/Source/Programs/Shared/EpicGames.Horde/Storage/BlobId.cs
@@ -5,6 +5,7 @@ using EpicGames.Serialization;
 using System;
 using System.Buffers.Binary;
 using System.ComponentModel;
+using System.Diagnostics;
 using System.Globalization;
 using System.Text;
 using System.Text.Json;
@@ -63,7 +64,7 @@ namespace EpicGames.Horde.Storage
 			_ = validate;
 		}
 
-		static readonly Utf8String s_process;
+		static Utf8String s_process;
 		static int s_counter;
 
 		/// <summary>
@@ -117,15 +118,32 @@ namespace EpicGames.Horde.Storage
 		/// <param name="output">Buffer to receive the output. The first 24 bytes will be written to.</param>
 		public static void GenerateUniqueId(Span<byte> output)
 		{
-			uint timestamp = (uint)((DateTime.UtcNow - DateTime.UnixEpoch).Ticks / TimeSpan.TicksPerSecond);
-			StringUtils.FormatUtf8HexString(timestamp, output);
+			if (s_process.IsEmpty)
+			{
+				output.Slice(0, 16).Fill((byte)'0');
+			}
+			else
+			{
+				uint timestamp = (uint)((DateTime.UtcNow - DateTime.UnixEpoch).Ticks / TimeSpan.TicksPerSecond);
+				StringUtils.FormatUtf8HexString(timestamp, output);
 
-			s_process.Span.CopyTo(output.Slice(8));
+				s_process.Span.CopyTo(output.Slice(8));
+			}
 
 			uint counter = (uint)Interlocked.Increment(ref s_counter);
 			StringUtils.FormatUtf8HexString(counter, output.Slice(16));
 		}
 
+		/// <summary>
+		/// Enables the generation of deterministic ids, for tests
+		/// </summary>
+		/// <param name="counter">Next counter for new ids</param>
+		public static void UseDeterministicIds(int counter = 0)
+		{
+			s_process = Utf8String.Empty;
+			s_counter = counter;
+		}
+
 		/// <summary>
 		/// Checks whether this blob id is valid
 		/// </summary>
diff --git a/Engine/Source/Programs/Shared/EpicGames.Horde/Storage/BlobLocator.cs b/Engine/Source/Programs/Shared/EpicGames.Horde/Storage/BlobLocator.cs
--- a/Engine/Source/Programs/Shared/EpicGames.Horde/Storage/BlobLocator.cs
+++ b/Engine/Source/Programs/Shared/EpicGames.Horde/Storage/BlobLocator.cs
@@ -100,6 +100,29 @@ namespace EpicGames.Horde.Storage
 			}
 		}
 
+		/// <summary>
+		/// Constructor
+		/// </summary>
+		/// <param name="hostId"></param>
+		/// <param name="blobId"></param>
+		public BlobLocator(ReadOnlySpan<byte> hostId, ReadOnlySpan<byte> blobId)
+		{
+			if (hostId.Length > 0)
+			{
+				byte[] buffer = new byte[hostId.Length + 1 + blobId.Length];
+
+				hostId.CopyTo(buffer);
+				buffer[hostId.Length] = (byte)':';
+				blobId.CopyTo(buffer.AsSpan(hostId.Length + 1));
+
+				Inner = new Utf8String(buffer);
+			}
+			else
+			{
+				Inner = new Utf8String(blobId.ToArray());
+			}
+		}
+
 		/// <summary>
 		/// Constructor
 		/// </summary>
diff --git a/Engine/Source/Programs/Shared/EpicGames.Horde/Storage/HostId.cs b/Engine/Source/Programs/Shared/EpicGames.Horde/Storage/HostId.cs
--- a/Engine/Source/Programs/Shared/EpicGames.Horde/Storage/HostId.cs
+++ b/Engine/Source/Programs/Shared/EpicGames.Horde/Storage/HostId.cs
@@ -113,7 +113,7 @@ namespace EpicGames.Horde.Storage
 				{
 					if (!IsValidChar(text[idx]))
 					{
-						throw new ArgumentException($"'{(char)text[idx]} is not a valid host identifier character", name);
+						throw new ArgumentException($"'{(char)text[idx]}' is not a valid host identifier character", name);
 					}
 				}
 			}
diff --git a/Engine/Source/Programs/Shared/EpicGames.Horde/Storage/IStorageClient.cs b/Engine/Source/Programs/Shared/EpicGames.Horde/Storage/IStorageClient.cs
--- a/Engine/Source/Programs/Shared/EpicGames.Horde/Storage/IStorageClient.cs
+++ b/Engine/Source/Programs/Shared/EpicGames.Horde/Storage/IStorageClient.cs
@@ -137,6 +137,14 @@ namespace EpicGames.Horde.Storage
 		/// <returns></returns>
 		Task<Bundle> ReadBundleAsync(BlobLocator locator, CancellationToken cancellationToken = default);
 
+		/// <summary>
+		/// Reads the header for a bundle from the store
+		/// </summary>
+		/// <param name="locator">The blob identifier</param>
+		/// <param name="cancellationToken">Cancellation token for the operation</param>
+		/// <returns></returns>
+		Task<BundleHeader> ReadBundleHeaderAsync(BlobLocator locator, CancellationToken cancellationToken = default);
+
 		/// <summary>
 		/// Writes a new bundle to the store
 		/// </summary>
@@ -246,6 +254,9 @@ namespace EpicGames.Horde.Storage
 			/// <inheritdoc/>
 			public Task<Bundle> ReadBundleAsync(BlobLocator locator, CancellationToken cancellationToken = default) => _inner.ReadBundleAsync(locator, cancellationToken);
 
+			/// <inheritdoc/>
+			public Task<BundleHeader> ReadBundleHeaderAsync(BlobLocator locator, CancellationToken cancellationToken = default) => _inner.ReadBundleHeaderAsync(locator, cancellationToken);
+
 			/// <inheritdoc/>
 			public Task<BlobLocator> WriteBundleAsync(Bundle bundle, Utf8String prefix = default, CancellationToken cancellationToken = default) => _inner.WriteBundleAsync(bundle, prefix, cancellationToken);
 
diff --git a/Engine/Source/Programs/Shared/EpicGames.Horde/Storage/StorageClientBase.cs b/Engine/Source/Programs/Shared/EpicGames.Horde/Storage/StorageClientBase.cs
--- a/Engine/Source/Programs/Shared/EpicGames.Horde/Storage/StorageClientBase.cs
+++ b/Engine/Source/Programs/Shared/EpicGames.Horde/Storage/StorageClientBase.cs
@@ -412,6 +412,18 @@ namespace EpicGames.Horde.Storage
 			}
 		}
 
+		/// <summary>
+		/// Reads a bundle header from the given blob locator, or retrieves it from the cache
+		/// </summary>
+		/// <param name="locator"></param>
+		/// <param name="cancellationToken"></param>
+		/// <returns></returns>
+		public async Task<BundleHeader> ReadBundleHeaderAsync(BlobLocator locator, CancellationToken cancellationToken = default)
+		{
+			BundleInfo info = await GetBundleInfoAsync(locator, cancellationToken);
+			return info.Header;
+		}
+
 		/// <summary>
 		/// Reads the header and structural metadata about the bundle
 		/// </summary>
