Horde: Prevent garbage collecting objects uploaded within the last few hours.

diff --git a/Engine/Source/Programs/Horde/Horde.Build.Tests/GcServiceTests.cs b/Engine/Source/Programs/Horde/Horde.Build.Tests/GcServiceTests.cs
--- a/Engine/Source/Programs/Horde/Horde.Build.Tests/GcServiceTests.cs
+++ b/Engine/Source/Programs/Horde/Horde.Build.Tests/GcServiceTests.cs
@@ -24,7 +24,7 @@ namespace Horde.Build.Tests
 		{
 			GlobalConfig globalConfig = new GlobalConfig();
 			globalConfig.Storage.Backends.Add(new BackendConfig { Id = new BackendId("default-backend"), Type = StorageBackendType.Memory });
-			globalConfig.Storage.Namespaces.Add(new NamespaceConfig { Id = new NamespaceId("default"), Backend = new BackendId("default-backend") });
+			globalConfig.Storage.Namespaces.Add(new NamespaceConfig { Id = new NamespaceId("default"), Backend = new BackendId("default-backend"), GcDelayHrs = 0.0 });
 			await ConfigCollection.AddConfigAsync("globals", globalConfig);
 			Assert.IsNotNull(await GlobalsService.TryUpdateAsync(await GlobalsService.GetAsync(), "globals"));
 		}
@@ -51,6 +51,8 @@ namespace Horde.Build.Tests
 
 			HashSet<BlobLocator> nodes = await FindNodes(store, roots);
 
+			await Clock.AdvanceAsync(TimeSpan.FromDays(1.0));
+
 			GcService gc = ServiceProvider.GetRequiredService<GcService>();
 			await gc.RunAsync(new NamespaceId("default"), CancellationToken.None);
 
diff --git a/Engine/Source/Programs/Horde/Horde.Build/Storage/GcService.cs b/Engine/Source/Programs/Horde/Horde.Build/Storage/GcService.cs
--- a/Engine/Source/Programs/Horde/Horde.Build/Storage/GcService.cs
+++ b/Engine/Source/Programs/Horde/Horde.Build/Storage/GcService.cs
@@ -68,6 +68,9 @@ namespace Horde.Build.Storage
 
 		[BsonElement("r")]
 		public List<IoHash>? Imports { get; set; }
+
+		[BsonElement("t")]
+		public DateTime Time { get; set; }
 	}
 
 	class GcSession
@@ -79,10 +82,13 @@ namespace Horde.Build.Storage
 		readonly IMongoCollection<GcNode> _nodeCollection;
 		readonly int _cycle;
 		readonly ILogger _logger;
-		GcPage _writePage = new GcPage();
 		readonly HashSet<IoHash> _reachable = new HashSet<IoHash>();
+		readonly DateTime _startTimeUtc;
+		readonly DateTime _deleteTimeUtc;
+
+		GcPage _writePage;
 
-		public GcSession(IStorageClientImpl store, NamespaceId namespaceId, IMongoCollection<GcPage> pageCollection, IMongoCollection<GcNode> nodeCollection, int cycle, ILogger logger)
+		public GcSession(IStorageClientImpl store, NamespaceId namespaceId, IMongoCollection<GcPage> pageCollection, IMongoCollection<GcNode> nodeCollection, int cycle, DateTime startTimeUtc, ILogger logger)
 		{
 			_store = store;
 			_namespaceId = namespaceId;
@@ -90,7 +96,10 @@ namespace Horde.Build.Storage
 			_pageCollection = pageCollection;
 			_nodeCollection = nodeCollection;
 			_cycle = cycle;
+			_startTimeUtc = startTimeUtc;
+			_deleteTimeUtc = startTimeUtc - TimeSpan.FromHours(store.Config.GcDelayHrs);
 			_logger = logger;
+
 			_writePage = new GcPage(_cycle, 0, -1);
 		}
 
@@ -133,9 +142,7 @@ namespace Horde.Build.Storage
 							IoHash hash = GetBlobDataId(locator.Blob.BlobId);
 							if (await WriteHashAsync(hash, cancellationToken))
 							{
-								FilterDefinition<GcNode> filter = Builders<GcNode>.Filter.Eq(x => x.Id, hash);
-								UpdateDefinition<GcNode> update = Builders<GcNode>.Update.SetOnInsert(x => x.Locator, locator.Blob);
-								requests.Add(new UpdateOneModel<GcNode>(filter, update) { IsUpsert = true });
+								requests.Add(CreateGcNodeUpsert(hash, locator.Blob));
 							}
 
 							moveNextTask = enumerator.MoveNextAsync(cancellationToken).AsTask();
@@ -183,8 +190,15 @@ namespace Horde.Build.Storage
 					IoHash hash = GetBlobDataId(blobId);
 					if (!_reachable.Contains(hash))
 					{
-						await _nodeCollection.DeleteOneAsync(x => x.Id == hash, cancellationToken);
-						await _store.Backend.DeleteAsync(path, cancellationToken);
+						DeleteResult result = await _nodeCollection.DeleteOneAsync(x => x.Id == hash && x.Time < _deleteTimeUtc, cancellationToken);
+						if (result.DeletedCount > 0 || _deleteTimeUtc == _startTimeUtc)
+						{
+							await _store.Backend.DeleteAsync(path, cancellationToken);
+						}
+						else
+						{
+							await _nodeCollection.BulkWriteAsync(new[] { CreateGcNodeUpsert(hash, new BlobLocator(HostId.Empty, blobId)) }, cancellationToken: cancellationToken);
+						}
 					}
 				}
 			}
@@ -271,9 +285,7 @@ namespace Horde.Build.Storage
 
 					if (!_reachable.Contains(importId))
 					{
-						FilterDefinition<GcNode> filter = Builders<GcNode>.Filter.Eq(x => x.Id, importId);
-						UpdateDefinition<GcNode> update = Builders<GcNode>.Update.SetOnInsert(x => x.Locator, import.Locator);
-						updates.Add(new UpdateOneModel<GcNode>(filter, update) { IsUpsert = true });
+						updates.Add(CreateGcNodeUpsert(importId, import.Locator));
 					}
 				}
 			}
@@ -283,6 +295,13 @@ namespace Horde.Build.Storage
 			await _nodeCollection.BulkWriteAsync(updates, new BulkWriteOptions { IsOrdered = false }, cancellationToken);
 		}
 
+		WriteModel<GcNode> CreateGcNodeUpsert(IoHash hash, BlobLocator locator)
+		{
+			FilterDefinition<GcNode> filter = Builders<GcNode>.Filter.Eq(x => x.Id, hash);
+			UpdateDefinition<GcNode> update = Builders<GcNode>.Update.SetOnInsert(x => x.Locator, locator).SetOnInsert(x => x.Time, _startTimeUtc);
+			return new UpdateOneModel<GcNode>(filter, update) { IsUpsert = true };
+		}
+
 		async ValueTask<bool> WriteHashAsync(IoHash hash, CancellationToken cancellationToken)
 		{
 			if (_reachable.Add(hash))
@@ -330,13 +349,15 @@ namespace Horde.Build.Storage
 		{
 			public NamespaceId Id { get; set; }
 			public int Cycle { get; set; }
-			public DateTime LastTime { get; set; }
+			public DateTime StartTime { get; set; }
+			public DateTime LastStartTime { get; set; }
 		}
 
 		readonly MongoService _mongoService;
 		readonly RedisService _redisService;
 		readonly StorageService _storageService;
 		readonly SingletonDocument<GcState> _gcState;
+		readonly IClock _clock;
 		readonly ITicker _ticker;
 		readonly ILogger _logger;
 		readonly IMongoCollection<GcNode> _nodeCollection;
@@ -348,9 +369,10 @@ namespace Horde.Build.Storage
 			_redisService = redisService;
 			_storageService = storageService;
 			_gcState = new SingletonDocument<GcState>(mongoService);
+			_clock = clock;
 			_ticker = clock.AddTicker<GcService>(TimeSpan.FromMinutes(5.0), TickAsync, logger);
 			_logger = logger;
-			_nodeCollection = _mongoService.GetCollection<GcNode>("Storage.GcNodes");
+			_nodeCollection = _mongoService.GetCollection<GcNode>("Storage.GcNodes", builder => builder.Ascending(x => x.Time));
 			_pageCollection = _mongoService.GetCollection<GcPage>("Storage.GcPages", builder => builder.Ascending(x => x.Cycle).Ascending(x => x.BaseIndex), unique: true);
 		}
 
@@ -394,7 +416,7 @@ namespace Horde.Build.Storage
 					NamespaceConfig? config = namespaces.FirstOrDefault(x => x.Id == namespaceState.Id);
 					if (config != null)
 					{
-						DateTime time = namespaceState.LastTime + TimeSpan.FromHours(config.GcFrequencyHrs);
+						DateTime time = namespaceState.LastStartTime + TimeSpan.FromHours(config.GcFrequencyHrs);
 						if (time > utcNow)
 						{
 							pending.Add((time, namespaceState));
@@ -443,11 +465,12 @@ namespace Horde.Build.Storage
 			int cycle = namespaceState.Cycle;
 			if (cycle == 0)
 			{
-				await _gcState.UpdateAsync(x => cycle = StartCycle(x, namespaceId));
+				GcState newState = await _gcState.UpdateAsync(x => cycle = StartCycle(x, namespaceId));
+				namespaceState = FindOrAddNamespace(newState, namespaceId);
 			}
 
 			IStorageClientImpl store = await _storageService.GetClientAsync(namespaceId, cancellationToken);
-			GcSession session = new GcSession(store, namespaceId, _pageCollection, _nodeCollection, cycle, _logger);
+			GcSession session = new GcSession(store, namespaceId, _pageCollection, _nodeCollection, cycle, namespaceState.StartTime, _logger);
 			await session.RunAsync(cancellationToken);
 			await _gcState.UpdateAsync(x => CompleteCycle(x, namespaceId));
 			await _pageCollection.DeleteManyAsync(x => x.Cycle == cycle, cancellationToken);
@@ -463,17 +486,18 @@ namespace Horde.Build.Storage
 			{
 				if (!currentNamespaceIds.Contains(config.Id))
 				{
-					state.Namespaces.Add(new GcNamespaceState { Id = config.Id, LastTime = DateTime.UtcNow });
+					state.Namespaces.Add(new GcNamespaceState { Id = config.Id, LastStartTime = DateTime.UtcNow });
 				}
 			}
 
 			state.Namespaces.SortBy(x => x.Id);
 		}
 
-		static int StartCycle(GcState state, NamespaceId namespaceId)
+		int StartCycle(GcState state, NamespaceId namespaceId)
 		{
 			GcNamespaceState namespaceState = FindOrAddNamespace(state, namespaceId);
 			namespaceState.Cycle = ++state.NextCycle;
+			namespaceState.StartTime = _clock.UtcNow;
 			return namespaceState.Cycle;
 		}
 
@@ -481,7 +505,7 @@ namespace Horde.Build.Storage
 		{
 			GcNamespaceState namespaceState = FindOrAddNamespace(state, namespaceId);
 			namespaceState.Cycle = 0;
-			namespaceState.LastTime = DateTime.UtcNow;
+			namespaceState.LastStartTime = namespaceState.StartTime;
 		}
 
 		static GcNamespaceState FindOrAddNamespace(GcState state, NamespaceId namespaceId)
@@ -491,7 +515,6 @@ namespace Horde.Build.Storage
 			{
 				namespaceState = new GcNamespaceState();
 				namespaceState.Id = namespaceId;
-				namespaceState.Cycle = ++state.NextCycle;
 				state.Namespaces.Add(namespaceState);
 			}
 			return namespaceState;
diff --git a/Engine/Source/Programs/Horde/Horde.Build/Storage/StorageConfig.cs b/Engine/Source/Programs/Horde/Horde.Build/Storage/StorageConfig.cs
--- a/Engine/Source/Programs/Horde/Horde.Build/Storage/StorageConfig.cs
+++ b/Engine/Source/Programs/Horde/Horde.Build/Storage/StorageConfig.cs
@@ -112,6 +112,11 @@ namespace Horde.Build.Storage
 		/// </summary>
 		public double GcFrequencyHrs { get; set; } = 2.0;
 
+		/// <summary>
+		/// How long to keep newly uploaded orphanned objects before allowing them to be deleted, in hours.
+		/// </summary>
+		public double GcDelayHrs { get; set; } = 6.0;
+
 		/// <summary>
 		/// Access list for this namespace
 		/// </summary>
