Horde: Change to an empty message when an issue has been closed and there is nothing else to display, to prevents no_text errors from Slack.

diff --git a/Engine/Source/Programs/Horde/Horde.Build/Notifications/Sinks/SlackNotificationSink.cs b/Engine/Source/Programs/Horde/Horde.Build/Notifications/Sinks/SlackNotificationSink.cs
--- a/Engine/Source/Programs/Horde/Horde.Build/Notifications/Sinks/SlackNotificationSink.cs
+++ b/Engine/Source/Programs/Horde/Horde.Build/Notifications/Sinks/SlackNotificationSink.cs
@@ -1014,6 +1014,11 @@ namespace Horde.Build.Notifications.Sinks
 					}
 				}
 
+				if (message.Blocks.Count == 0)
+				{
+					message.AddSection("");
+				}
+
 				await SendOrUpdateMessageToThreadAsync(triageChannel, $"{eventId}_buttons", null, threadId, message);
 			}
 
