Horde: Remove log lines used to debug Redis I/O issues

diff --git a/Engine/Source/Programs/Horde/Horde.Build/Logs/LogFileService.cs b/Engine/Source/Programs/Horde/Horde.Build/Logs/LogFileService.cs
--- a/Engine/Source/Programs/Horde/Horde.Build/Logs/LogFileService.cs
+++ b/Engine/Source/Programs/Horde/Horde.Build/Logs/LogFileService.cs
@@ -1018,7 +1018,6 @@ namespace Horde.Build.Logs
 			
 			// Get all the chunks older than 20 minutes
 			List<(LogId, long)> flushChunks = await _builder.TouchChunksAsync(TimeSpan.FromMinutes(10.0));
-			_logger.LogInformation("Performing incremental flush of log builder ({NumChunks} chunks)", flushChunks.Count);
 
 			scope.Span.SetTag("numChunks", flushChunks.Count);
 
@@ -1160,7 +1159,6 @@ namespace Horde.Build.Logs
 				}
 			}
 
-			_logger.LogInformation("Running {NumWrites} writes with a concurrency of {Concurrency}", offsetsToWrite.Count, MaxConcurrentChunkWrites);
 			scope.Span.SetTag("numOffsetsToWrite", offsetsToWrite.Count);
 			ParallelOptions opts = new() { MaxDegreeOfParallelism = MaxConcurrentChunkWrites, CancellationToken = cancellationToken };
 			await Parallel.ForEachAsync(offsetsToWrite, opts, async (x, innerCt) =>
@@ -1217,7 +1215,6 @@ namespace Horde.Build.Logs
 			}
 			
 			scope.Span.SetTag("NumWriteTasks", chunkWriteTasks.Count);
-			_logger.LogInformation("Writing {NumLogChunks} in parallel", chunkWriteTasks.Count);
 
 			// Wait for the tasks to complete, periodically updating the log file object
 			ILogFile? logFile = logFileInterface;
