Horde: Additional logging for creating/disposing of Perforce connections.

diff --git a/Engine/Source/Programs/Horde/Horde.Build/Perforce/PerforceService.cs b/Engine/Source/Programs/Horde/Horde.Build/Perforce/PerforceService.cs
--- a/Engine/Source/Programs/Horde/Horde.Build/Perforce/PerforceService.cs
+++ b/Engine/Source/Programs/Horde/Horde.Build/Perforce/PerforceService.cs
@@ -476,6 +476,8 @@ namespace Horde.Build.Perforce
 				return handle;
 			}
 
+			_logger.LogDebug("Unable to allocate existing pooled connection (total: {Total}, cluster: {Cluster}, user: {UserName}, stream: {StreamName})", _pooledConnections.Count, clusterName, userName, streamName);
+
 			using (PooledConnectionHandle perforce = await ConnectAsync(cluster, userName, cancellationToken))
 			{
 				InfoRecord info = await perforce.GetInfoAsync(cancellationToken);
@@ -487,7 +489,7 @@ namespace Horde.Build.Perforce
 				client.Stream = streamName;
 				await perforce.CreateClientAsync(client, cancellationToken);
 
-				_logger.LogDebug("Created client {ClientName} (user: {User}, root: {Root}, stream: {Stream}, server: {ServerId})", client.Name, client.Owner, client.Root, client.Stream, info.ServerId ?? "unknown");
+				_logger.LogDebug("Created client {ClientName} (user: {User}, root: {Root}, stream: {Stream}, server: {ServerId}, serverAndPort: {ServerAndPort})", client.Name, client.Owner, client.Root, client.Stream, info.ServerId ?? "unknown", perforce.Settings.ServerAndPort);
 				return await CreatePooledConnectionAsync(perforce.Settings.ServerAndPort, perforce.Credentials, client, cancellationToken);
 			}
 		}
diff --git a/Engine/Source/Programs/Shared/EpicGames.Perforce/NativePerforceConnection.cs b/Engine/Source/Programs/Shared/EpicGames.Perforce/NativePerforceConnection.cs
--- a/Engine/Source/Programs/Shared/EpicGames.Perforce/NativePerforceConnection.cs
+++ b/Engine/Source/Programs/Shared/EpicGames.Perforce/NativePerforceConnection.cs
@@ -318,6 +318,8 @@ namespace EpicGames.Perforce
 			_backgroundThread = new Thread(BackgroundThreadProc);
 			_backgroundThread.IsBackground = true;
 			_backgroundThread.Start();
+
+			logger.LogTrace("Created Perforce connection {ConnectionId} (server: {ServerAndPort}, user: {UserName}, client: {ClientName})", _uniqueId, settings.ServerAndPort, settings.UserName, settings.ClientName);
 		}
 
 		/// <summary>
@@ -385,6 +387,11 @@ namespace EpicGames.Perforce
 		/// <inheritdoc/>
 		void Dispose(bool disposing)
 		{
+			if (disposing)
+			{
+				Logger.LogTrace("Disposing Perforce connection {ConnectionId}", _uniqueId);
+			}
+
 			if (_backgroundThread != null)
 			{
 				_requests.Add(null);
