Horde: Fix missing disposal of async enumerators, causing slow leak of Perforce connections when checking for changes.

diff --git a/Engine/Source/Programs/Horde/Horde.Build/Jobs/Schedules/ScheduleService.cs b/Engine/Source/Programs/Horde/Horde.Build/Jobs/Schedules/ScheduleService.cs
--- a/Engine/Source/Programs/Horde/Horde.Build/Jobs/Schedules/ScheduleService.cs
+++ b/Engine/Source/Programs/Horde/Horde.Build/Jobs/Schedules/ScheduleService.cs
@@ -355,7 +355,7 @@ namespace Horde.Build.Jobs.Schedules
 			// Cache the Perforce history as we're iterating through changes to improve query performance
 			ICommitCollection commits = _commitService.GetCollection(stream);
 			IAsyncEnumerable<ICommit> commitEnumerable = commits.FindAsync(minChangeNumber, null, null, schedule.Config.Commits, cancellationToken);
-			IAsyncEnumerator<ICommit> commitEnumerator = commitEnumerable.GetAsyncEnumerator(cancellationToken);
+			await using IAsyncEnumerator<ICommit> commitEnumerator = commitEnumerable.GetAsyncEnumerator(cancellationToken);
 
 			// Start as many jobs as possible
 			List<(int Change, int CodeChange)> triggerChanges = new List<(int, int)>();
diff --git a/Engine/Source/Programs/Horde/Horde.Build/Logs/LogFileService.cs b/Engine/Source/Programs/Horde/Horde.Build/Logs/LogFileService.cs
--- a/Engine/Source/Programs/Horde/Horde.Build/Logs/LogFileService.cs
+++ b/Engine/Source/Programs/Horde/Horde.Build/Logs/LogFileService.cs
@@ -1584,7 +1584,7 @@ namespace Horde.Build.Logs
 			List<int> results = new List<int>();
 			if (count > 0)
 			{
-				IAsyncEnumerator<int> enumerator = SearchLogDataInternalAsync(logFile, text, firstLine, searchStats).GetAsyncEnumerator();
+				await using IAsyncEnumerator<int> enumerator = SearchLogDataInternalAsync(logFile, text, firstLine, searchStats).GetAsyncEnumerator();
 				while (await enumerator.MoveNextAsync() && results.Count < count)
 				{
 					results.Add(enumerator.Current);
