Horde: Fix incorrect argument to describe command when enumerating changes. We don't want shelved change information.

diff --git a/Engine/Source/Programs/Horde/Horde.Build/Perforce/PerforceService.cs b/Engine/Source/Programs/Horde/Horde.Build/Perforce/PerforceService.cs
--- a/Engine/Source/Programs/Horde/Horde.Build/Perforce/PerforceService.cs
+++ b/Engine/Source/Programs/Horde/Horde.Build/Perforce/PerforceService.cs
@@ -1214,7 +1214,7 @@ namespace Horde.Build.Perforce
 								break;
 							}
 
-							List<DescribeRecord> describeRecords = await perforce.DescribeAsync(DescribeOptions.Shelved, -1, changesRecords.Select(x => x.Number).ToArray(), cancellationToken);
+							List<DescribeRecord> describeRecords = await perforce.DescribeAsync(DescribeOptions.None, -1, changesRecords.Select(x => x.Number).ToArray(), cancellationToken);
 							foreach (DescribeRecord describeRecord in describeRecords)
 							{
 								List<string> files = viewMap.MapFiles(describeRecord.Files.Select(x => x.DepotFile), info.PathComparison).ToList();
