Horde: Populate the cache with new changes when queries intersect with the existing cache range.

diff --git a/Engine/Source/Programs/Horde/Horde.Build/Perforce/PerforceServiceCache.cs b/Engine/Source/Programs/Horde/Horde.Build/Perforce/PerforceServiceCache.cs
--- a/Engine/Source/Programs/Horde/Horde.Build/Perforce/PerforceServiceCache.cs
+++ b/Engine/Source/Programs/Horde/Horde.Build/Perforce/PerforceServiceCache.cs
@@ -138,6 +138,12 @@ namespace Horde.Build.Perforce
 				CommitTags = commitTags;
 			}
 
+			public static async Task<CachedCommitDoc> FromCommitAsync(ICommit commit, CancellationToken cancellationToken)
+			{
+				List<CommitTag> commitTags = new List<CommitTag>(await commit.GetTagsAsync(cancellationToken));
+				return new CachedCommitDoc(commit, commitTags);
+			}
+
 			public void PostLoad(PerforceServiceCache owner, IStream stream)
 			{
 				_owner = owner;
@@ -418,8 +424,7 @@ namespace Horde.Build.Perforce
 							}
 
 							CachedCommitDoc commitDoc = new CachedCommitDoc(commit, commitTags);
-							FilterDefinition<CachedCommitDoc> filter = Builders<CachedCommitDoc>.Filter.Expr(x => x.StreamId == commit.StreamId && x.Number == commit.Number);
-							await _commits.ReplaceOneAsync(filter, commitDoc, new ReplaceOptions { IsUpsert = true }, cancellationToken);
+							await AddCachedCommitAsync(commitDoc, cancellationToken);
 						}
 					}
 				}
@@ -441,6 +446,12 @@ namespace Horde.Build.Perforce
 			}
 		}
 
+		async Task AddCachedCommitAsync(CachedCommitDoc commitDoc, CancellationToken cancellationToken)
+		{
+			FilterDefinition<CachedCommitDoc> filter = Builders<CachedCommitDoc>.Filter.Expr(x => x.StreamId == commitDoc.StreamId && x.Number == commitDoc.Number);
+			await _commits.ReplaceOneAsync(filter, commitDoc, new ReplaceOptions { IsUpsert = true }, cancellationToken);
+		}
+
 		#endregion
 
 		#region ICommitSource implementation
@@ -532,6 +543,32 @@ namespace Horde.Build.Perforce
 						}
 
 						maxChange = minReplicatedChange - 1;
+
+						// Expand the range of cached changes if necessary
+						if (maxResults == null || maxResults.Value > 0)
+						{
+							await foreach (ICommit commit in base.FindAsync(minChange, maxChange, maxResults, null, cancellationToken))
+							{
+								CachedCommitDoc cachedCommit = await CachedCommitDoc.FromCommitAsync(commit, cancellationToken);
+								await _owner.AddCachedCommitAsync(cachedCommit, cancellationToken);
+								await _owner._mongoService.UpdateSingletonAsync<CacheState>(x => TryUpdateRange(x, _stream, commit.Number, maxChange));
+								_owner._logger.LogDebug("Adding new cached commit for {StreamId} at change {Change}", _stream.Id, commit.Number);
+
+								if (tags == null || tags.Any(x => cachedCommit.CommitTags.Contains(x)))
+								{
+									yield return cachedCommit;
+								}
+
+								maxResults = maxResults.HasValue ? maxResults.Value - 1 : null;
+							}
+
+							if (maxResults == null || maxResults.Value > 0)
+							{
+								int newMinChange = minChange ?? 0;
+								_owner._logger.LogDebug("Extending range for {StreamId} cache to {Change}..", _stream.Id, newMinChange);
+								await _owner._mongoService.UpdateSingletonAsync<CacheState>(x => TryUpdateRange(x, _stream, newMinChange, maxChange));
+							}
+						}
 					}
 				}
 
@@ -546,6 +583,24 @@ namespace Horde.Build.Perforce
 				}
 			}
 
+			static bool TryUpdateRange(CacheState state, IStream stream, int newMinChange, int? maxChange)
+			{
+				ClusterState? clusterState;
+				if (state.Clusters.TryGetValue(stream.Config.ClusterName, out clusterState))
+				{
+					int streamMinChange;
+					if (clusterState.MinChanges.TryGetValue(stream.Id, out streamMinChange))
+					{
+						if (newMinChange < streamMinChange && (maxChange == null || maxChange.Value >= streamMinChange - 1))
+						{
+							clusterState.MinChanges[stream.Id] = newMinChange;
+							return true;
+						}
+					}
+				}
+				return false;
+			}
+
 			public override async Task<ICommit> GetAsync(int changeNumber, CancellationToken cancellationToken = default)
 			{
 				CachedCommitDoc? commit = await _owner._commits.Find(x => x.StreamId == _stream.Id && x.Number == changeNumber).FirstOrDefaultAsync(cancellationToken);
diff --git a/Engine/Source/Programs/Horde/Horde.Build/Server/MongoService.cs b/Engine/Source/Programs/Horde/Horde.Build/Server/MongoService.cs
--- a/Engine/Source/Programs/Horde/Horde.Build/Server/MongoService.cs
+++ b/Engine/Source/Programs/Horde/Horde.Build/Server/MongoService.cs
@@ -818,13 +818,31 @@ namespace Horde.Build.Server
 		/// <typeparam name="T"></typeparam>
 		/// <param name="updater"></param>
 		/// <returns></returns>
-		public async Task UpdateSingletonAsync<T>(Action<T> updater) where T : SingletonBase, new()
+		public Task UpdateSingletonAsync<T>(Action<T> updater) where T : SingletonBase, new()
+		{
+			bool Update(T instance)
+			{
+				updater(instance);
+				return true;
+			}
+			return UpdateSingletonAsync<T>(Update);
+		}
+
+		/// <summary>
+		/// Updates a singleton
+		/// </summary>
+		/// <typeparam name="T"></typeparam>
+		/// <param name="updater"></param>
+		/// <returns></returns>
+		public async Task UpdateSingletonAsync<T>(Func<T, bool> updater) where T : SingletonBase, new()
 		{
 			for (; ; )
 			{
 				T document = await GetSingletonAsync(() => new T());
-				updater(document);
-
+				if (!updater(document))
+				{
+					break;
+				}
 				if (await TryUpdateSingletonAsync(document))
 				{
 					break;
