Horde: Report "Horde" as the user name when an issue is closed automatically.

diff --git a/Engine/Source/Programs/Horde/Horde.Build/Issues/IssueCollection.cs b/Engine/Source/Programs/Horde/Horde.Build/Issues/IssueCollection.cs
--- a/Engine/Source/Programs/Horde/Horde.Build/Issues/IssueCollection.cs
+++ b/Engine/Source/Programs/Horde/Horde.Build/Issues/IssueCollection.cs
@@ -526,6 +526,14 @@ namespace Horde.Build.Issues
 			{
 				return "null";
 			}
+			else if (UserId == IIssue.ResolvedByUnknownId)
+			{
+				return "Horde (Unknown)";
+			}
+			else if (UserId == IIssue.ResolvedByTimeoutId)
+			{
+				return "Horde (Timeout)";
+			}
 
 			IUser? user = await _userCollection.GetCachedUserAsync(UserId);
 			if (user == null)
