Horde: Add a non-empty section block to Slack messages for closed issues.

diff --git a/Engine/Source/Programs/Horde/Horde.Build/Notifications/Sinks/SlackNotificationSink.cs b/Engine/Source/Programs/Horde/Horde.Build/Notifications/Sinks/SlackNotificationSink.cs
--- a/Engine/Source/Programs/Horde/Horde.Build/Notifications/Sinks/SlackNotificationSink.cs
+++ b/Engine/Source/Programs/Horde/Horde.Build/Notifications/Sinks/SlackNotificationSink.cs
@@ -1016,7 +1016,7 @@ namespace Horde.Build.Notifications.Sinks
 
 				if (message.Blocks.Count == 0)
 				{
-					message.AddSection("");
+					message.AddSection("Issue has been closed.");
 				}
 
 				await SendOrUpdateMessageToThreadAsync(triageChannel, $"{eventId}_buttons", null, threadId, message);
