Horde: Do not log an error when submit fails due to a trigger.

diff --git a/Engine/Source/Programs/Horde/Horde.Build/Perforce/PerforceService.cs b/Engine/Source/Programs/Horde/Horde.Build/Perforce/PerforceService.cs
--- a/Engine/Source/Programs/Horde/Horde.Build/Perforce/PerforceService.cs
+++ b/Engine/Source/Programs/Horde/Horde.Build/Perforce/PerforceService.cs
@@ -1075,6 +1075,11 @@ namespace Horde.Build.Perforce
 					SubmitRecord submit = await perforce.SubmitShelvedAsync(change, cancellationToken);
 					return (submit.SubmittedChangeNumber, records[0].Description);
 				}
+				catch (PerforceException ex) when (ex.Error != null && ex.Error.Generic == PerforceGenericCode.NotYet)
+				{
+					_logger.LogInformation(ex, "Unable to submit shelved change {Change}: {Message}", change, ex.Message);
+					return (null, $"Submit command failed: {ex.Message}");
+				}
 				catch (Exception ex)
 				{
 					_logger.LogError(ex, "Unable to submit shelved change {Change}: {Message}", change, ex.Message);
