Horde: Lower log levels for fleet service to reduce log output

diff --git a/Engine/Source/Programs/Horde/Horde.Build/Agents/Fleet/FleetService.cs b/Engine/Source/Programs/Horde/Horde.Build/Agents/Fleet/FleetService.cs
--- a/Engine/Source/Programs/Horde/Horde.Build/Agents/Fleet/FleetService.cs
+++ b/Engine/Source/Programs/Horde/Horde.Build/Agents/Fleet/FleetService.cs
@@ -111,7 +111,6 @@ namespace Horde.Build.Agents.Fleet
 
 		internal async ValueTask TickLeaderAsync(CancellationToken stoppingToken)
 		{
-			_logger.LogInformation("Autoscaling pools...");
 			Stopwatch stopwatch = Stopwatch.StartNew();
 			using IScope _ = GlobalTracer.Instance.BuildSpan("AutoscaleService.TickAsync").StartActive();
 			
@@ -120,12 +119,11 @@ namespace Horde.Build.Agents.Fleet
 			await ResizePoolsAsync(output, stoppingToken);
 
 			stopwatch.Stop();
-			_logger.LogInformation("Autoscaling pools took {ElapsedTime} ms", stopwatch.ElapsedMilliseconds);
+			_logger.LogDebug("Autoscaling pools took {ElapsedTime} ms", stopwatch.ElapsedMilliseconds);
 		}
 
 		internal async ValueTask TickHighFrequencyAsync(CancellationToken stoppingToken)
 		{
-			_logger.LogInformation("Autoscaling pools (high frequency)...");
 			Stopwatch stopwatch = Stopwatch.StartNew();
 			using IScope _ = GlobalTracer.Instance.BuildSpan("AutoscaleService.TickHighFrequency").StartActive();
 			
@@ -133,7 +131,7 @@ namespace Horde.Build.Agents.Fleet
 			await Task.Delay(0, stoppingToken);
 
 			stopwatch.Stop();
-			_logger.LogInformation("Autoscaling pools (high frequency) took {ElapsedTime} ms", stopwatch.ElapsedMilliseconds);
+			_logger.LogDebug("Autoscaling pools (high frequency) took {ElapsedTime} ms", stopwatch.ElapsedMilliseconds);
 		}
 
 		internal async Task<List<PoolSizeResult>> CalculatePoolSizesAsync(List<PoolSizeResult> poolSizeDatas, CancellationToken cancellationToken)
@@ -236,7 +234,7 @@ namespace Horde.Build.Agents.Fleet
 			}
 			catch (Exception ex)
 			{
-				_logger.LogError(ex, "Failed to scale {PoolName}:\n{Exception}", pool.Name, ex);
+				_logger.LogInformation(ex, "Failed to scale {PoolName}:\n{Exception}", pool.Name, ex);
 				return;
 			}
 
