Horde: Allow the agent to terminate gracefully when the system is shutting down.

diff --git a/Engine/Source/Programs/Horde/Horde.Agent/Services/WorkerService.cs b/Engine/Source/Programs/Horde/Horde.Agent/Services/WorkerService.cs
--- a/Engine/Source/Programs/Horde/Horde.Agent/Services/WorkerService.cs
+++ b/Engine/Source/Programs/Horde/Horde.Agent/Services/WorkerService.cs
@@ -119,9 +119,17 @@ namespace Horde.Agent.Services
 						for (int idx = 10; idx > 0; idx--)
 						{
 							_logger.LogInformation("Waiting for shutdown ({Count})", idx);
-							await Task.Delay(TimeSpan.FromSeconds(60.0), stoppingToken);
+							try
+							{
+								await Task.Delay(TimeSpan.FromSeconds(60.0), stoppingToken);
+								_logger.LogInformation("Shutdown aborted.");
+							}
+							catch (OperationCanceledException)
+							{
+								_logger.LogInformation("Agent is shutting down.");
+								return;
+							}
 						}
-						_logger.LogInformation("Shutdown aborted.");
 					}
 				}
 				
