Horde: Add logging whenever Polly retries an HTTP request.

diff --git a/Engine/Source/Programs/Horde/Horde.Agent/Services/GrpcService.cs b/Engine/Source/Programs/Horde/Horde.Agent/Services/GrpcService.cs
--- a/Engine/Source/Programs/Horde/Horde.Agent/Services/GrpcService.cs
+++ b/Engine/Source/Programs/Horde/Horde.Agent/Services/GrpcService.cs
@@ -67,7 +67,8 @@ namespace Horde.Agent.Services
 			customCertHandler.ServerCertificateCustomValidationCallback += (sender, cert, chain, errors) => CertificateHelper.CertificateValidationCallBack(_logger, sender, cert, chain, errors, _serverProfile);
 
 			TimeSpan[] retryDelay = { TimeSpan.FromSeconds(1), TimeSpan.FromSeconds(5), TimeSpan.FromSeconds(10) };
-			IAsyncPolicy<HttpResponseMessage> policy = HttpPolicyExtensions.HandleTransientHttpError().WaitAndRetryAsync(retryDelay);
+			IAsyncPolicy<HttpResponseMessage> policy = HttpPolicyExtensions.HandleTransientHttpError()
+				.WaitAndRetryAsync(retryDelay, onRetry: (result, timeSpan) => _logger.LogInformation("Retrying http request {Time} (status: {Code})", timeSpan, result.Result?.StatusCode));
 #pragma warning disable CA2000 // Dispose objects before losing scope
 			PolicyHttpMessageHandler retryHandler = new PolicyHttpMessageHandler(policy);
 #pragma warning restore CA2000 // Dispose objects before losing scope
