Horde: Increase verbosity for WriteOutput messages, to attempt to get additional data on intermittent failures.

diff --git a/Engine/Source/Programs/Horde/Horde.Build/Startup.cs b/Engine/Source/Programs/Horde/Horde.Build/Startup.cs
--- a/Engine/Source/Programs/Horde/Horde.Build/Startup.cs
+++ b/Engine/Source/Programs/Horde/Horde.Build/Startup.cs
@@ -920,7 +920,7 @@ namespace Horde.Build
 				options.GetLevel = GetRequestLoggingLevel;
 				options.EnrichDiagnosticContext = (diagnosticContext, httpContext) =>
 				{
-					diagnosticContext.Set("RemoteIP", httpContext.Connection.RemoteIpAddress);
+					diagnosticContext.Set("RemoteIP", httpContext?.Connection?.RemoteIpAddress);
 				};
 			});
 
@@ -1007,7 +1007,7 @@ namespace Horde.Build
 				}
 				if (requestPath.Equals("/Horde.HordeRpc/WriteOutput", StringComparison.OrdinalIgnoreCase))
 				{
-					return LogEventLevel.Verbose;
+					return LogEventLevel.Information;
 				}
 				if (requestPath.StartsWith("/Horde.HordeRpc", StringComparison.OrdinalIgnoreCase))
 				{
