Horde: Fix warnings about unused arguments for -server=... and -workingdir=...

diff --git a/Engine/Source/Programs/Horde/Horde.Agent/Program.cs b/Engine/Source/Programs/Horde/Horde.Agent/Program.cs
--- a/Engine/Source/Programs/Horde/Horde.Agent/Program.cs
+++ b/Engine/Source/Programs/Horde/Horde.Agent/Program.cs
@@ -204,7 +204,7 @@ namespace Horde.Agent
 
 			// Execute all the commands
 			IServiceProvider serviceProvider = services.BuildServiceProvider();
-			return await CommandHost.RunAsync(new CommandLineArguments(args), serviceProvider, typeof(Horde.Agent.Modes.Service.RunCommand));
+			return await CommandHost.RunAsync(arguments, serviceProvider, typeof(Horde.Agent.Modes.Service.RunCommand));
 		}
 
 		static void ConfigureTracing(string environment, string version)
