Horde: Log the size of artifacts that are uploaded.

diff --git a/Engine/Source/Programs/Horde/Horde.Agent/Utility/ArtifactUploader.cs b/Engine/Source/Programs/Horde/Horde.Agent/Utility/ArtifactUploader.cs
--- a/Engine/Source/Programs/Horde/Horde.Agent/Utility/ArtifactUploader.cs
+++ b/Engine/Source/Programs/Horde/Horde.Agent/Utility/ArtifactUploader.cs
@@ -87,9 +87,9 @@ namespace Horde.Agent.Utility
 		/// <returns></returns>
 		private static async Task<string> DoUploadAsync(HordeRpc.HordeRpcClient client, string jobId, string batchId, string stepId, string artifactName, FileReference artifactFile, ILogger logger, CancellationToken cancellationToken)
         {
-			logger.LogInformation("Uploading artifact {ArtifactName} from {ArtifactFile}", artifactName, artifactFile);
             using (FileStream artifactStream = FileReference.Open(artifactFile, FileMode.Open, FileAccess.Read, FileShare.Read))
 			{
+				logger.LogInformation("Uploading artifact {ArtifactName} from {ArtifactFile} ({Size:n0} bytes)", artifactName, artifactFile, artifactStream.Length);
 				using (AsyncClientStreamingCall<UploadArtifactRequest, UploadArtifactResponse> cursor = client.UploadArtifact(null, null, cancellationToken))
                 {
 					// Upload the metadata in the initial request
