Horde: Cap maximum allowed artifact size at 512mb.

diff --git a/Engine/Source/Programs/Horde/Horde.Agent/Utility/ArtifactUploader.cs b/Engine/Source/Programs/Horde/Horde.Agent/Utility/ArtifactUploader.cs
--- a/Engine/Source/Programs/Horde/Horde.Agent/Utility/ArtifactUploader.cs
+++ b/Engine/Source/Programs/Horde/Horde.Agent/Utility/ArtifactUploader.cs
@@ -63,6 +63,15 @@ namespace Horde.Agent.Utility
         {
 			try
 			{
+				const long MaxSize = 512 * 1024 * 1024;
+
+				FileInfo artifactInfo = artifactFile.ToFileInfo();
+				if (artifactInfo.Length > MaxSize)
+				{
+					logger.LogInformation("Not uploading {File}; size exceeds maximum allowed ({Size:n0} > {MaxSize:n0})", artifactFile, artifactInfo.Length, MaxSize);
+					return null;
+				}
+
 				string artifactId = await rpcConnection.InvokeAsync(rpcClient => DoUploadAsync(rpcClient, jobId, batchId, stepId, artifactName, artifactFile, logger, cancellationToken), new RpcContext(), cancellationToken);
 				return artifactId;
 			}
