Horde: Reduce max artifact size below max gRPC message size.

diff --git a/Engine/Source/Programs/Horde/Horde.Agent/Utility/ArtifactUploader.cs b/Engine/Source/Programs/Horde/Horde.Agent/Utility/ArtifactUploader.cs
--- a/Engine/Source/Programs/Horde/Horde.Agent/Utility/ArtifactUploader.cs
+++ b/Engine/Source/Programs/Horde/Horde.Agent/Utility/ArtifactUploader.cs
@@ -63,7 +63,7 @@ namespace Horde.Agent.Utility
         {
 			try
 			{
-				const long MaxSize = 512 * 1024 * 1024;
+				const long MaxSize = 192 * 1024 * 1024;
 
 				FileInfo artifactInfo = artifactFile.ToFileInfo();
 				if (artifactInfo.Length > MaxSize)
