Horde: Support cancellation tokens in LogFileService and LogFileCollection

diff --git a/Engine/Source/Programs/Horde/Horde.Build.Tests/IssueServiceTests.cs b/Engine/Source/Programs/Horde/Horde.Build.Tests/IssueServiceTests.cs
--- a/Engine/Source/Programs/Horde/Horde.Build.Tests/IssueServiceTests.cs
+++ b/Engine/Source/Programs/Horde/Horde.Build.Tests/IssueServiceTests.cs
@@ -78,8 +78,8 @@ namespace Horde.Build.Tests
 
 			private async Task WriteAsync(LogLevel level, byte[] line)
 			{
-				ILogFile logFile = (await _logFileService.GetLogFileAsync(_logId))!;
-				LogMetadata metadata = await _logFileService.GetMetadataAsync(logFile);
+				ILogFile logFile = (await _logFileService.GetLogFileAsync(_logId, CancellationToken.None))!;
+				LogMetadata metadata = await _logFileService.GetMetadataAsync(logFile, CancellationToken.None);
 				await _logFileService.WriteLogDataAsync(logFile, metadata.Length, metadata.MaxLineIndex, line, false);
 
 				if (level >= LogLevel.Warning)
@@ -88,7 +88,7 @@ namespace Horde.Build.Tests
 					if (@event.LineIndex == 0)
 					{
 						EventSeverity severity = (level == LogLevel.Warning) ? EventSeverity.Warning : EventSeverity.Error;
-						await _logFileService.CreateEventsAsync(new List<NewLogEventData> { new NewLogEventData { LogId = _logId, LineIndex = metadata.MaxLineIndex, LineCount = @event.LineCount, Severity = severity } });
+						await _logFileService.CreateEventsAsync(new List<NewLogEventData> { new NewLogEventData { LogId = _logId, LineIndex = metadata.MaxLineIndex, LineCount = @event.LineCount, Severity = severity } }, CancellationToken.None);
 					}
 				}
 			}
@@ -274,11 +274,11 @@ namespace Horde.Build.Tests
 			bytes.AddRange(JsonSerializer.SerializeToUtf8Bytes(data));
 			bytes.Add((byte)'\n');
 
-			ILogFile logFile = (await LogFileService.GetLogFileAsync(logId))!;
-			LogMetadata metadata = await LogFileService.GetMetadataAsync(logFile);
+			ILogFile logFile = (await LogFileService.GetLogFileAsync(logId, CancellationToken.None))!;
+			LogMetadata metadata = await LogFileService.GetMetadataAsync(logFile, CancellationToken.None);
 			await LogFileService.WriteLogDataAsync(logFile, metadata.Length, metadata.MaxLineIndex, bytes.ToArray(), false);
 
-			await LogFileService.CreateEventsAsync(new List<NewLogEventData> { new NewLogEventData { LogId = logId, LineIndex = metadata.MaxLineIndex, LineCount = 1, Severity = severity } });
+			await LogFileService.CreateEventsAsync(new List<NewLogEventData> { new NewLogEventData { LogId = logId, LineIndex = metadata.MaxLineIndex, LineCount = 1, Severity = severity } }, CancellationToken.None);
 		}
 
 		private TestJsonLogger CreateLogger(IJob job, int batchIdx, int stepIdx)
@@ -604,7 +604,7 @@ namespace Horde.Build.Tests
 				await ParseEventsAsync(job, 0, 0, lines);
 				await UpdateCompleteStep(job, 0, 0, JobStepOutcome.Warnings);
 
-				ILogFile? log = await LogFileService.GetLogFileAsync(job.Batches[0].Steps[0].LogId!.Value);
+				ILogFile? log = await LogFileService.GetLogFileAsync(job.Batches[0].Steps[0].LogId!.Value, CancellationToken.None);
 				List<ILogEvent> events = await LogFileService.FindEventsAsync(log!);
 				Assert.AreEqual(1, events.Count);
 				Assert.AreEqual(2, events[0].LineCount);
@@ -1753,7 +1753,7 @@ namespace Horde.Build.Tests
 				await ParseAsync(job.Batches[0].Steps[0].LogId!.Value, lines);
 				await UpdateCompleteStep(job, 0, 0, JobStepOutcome.Failure);
 
-				ILogFile? log = await LogFileService.GetLogFileAsync(job.Batches[0].Steps[0].LogId!.Value);
+				ILogFile? log = await LogFileService.GetLogFileAsync(job.Batches[0].Steps[0].LogId!.Value, CancellationToken.None);
 				List<ILogEvent> events = await LogFileService.FindEventsAsync(log!);
 				Assert.AreEqual(1, events.Count);
 				Assert.AreEqual(40, events[0].LineCount);
diff --git a/Engine/Source/Programs/Horde/Horde.Build.Tests/LogFileServiceTests.cs b/Engine/Source/Programs/Horde/Horde.Build.Tests/LogFileServiceTests.cs
--- a/Engine/Source/Programs/Horde/Horde.Build.Tests/LogFileServiceTests.cs
+++ b/Engine/Source/Programs/Horde/Horde.Build.Tests/LogFileServiceTests.cs
@@ -6,6 +6,7 @@ using System.IO;
 using System.Linq;
 using System.Security.Claims;
 using System.Text;
+using System.Threading;
 using System.Threading.Tasks;
 using Horde.Build.Agents.Sessions;
 using Horde.Build.Jobs;
@@ -77,7 +78,7 @@ namespace Horde.Build.Tests
         public async Task WriteLogLifecycleOldTest()
         {
 			JobId jobId = JobId.GenerateNewId();
-            ILogFile logFile = await _logFileService.CreateLogFileAsync(jobId, null, LogType.Text);
+            ILogFile logFile = await _logFileService.CreateLogFileAsync(jobId, null, LogType.Text, null, cancellationToken: CancellationToken.None);
 
             logFile = (await ((ILogFileService)_logFileService).WriteLogDataAsync(logFile, 0, 0, Encoding.ASCII.GetBytes("hello\n"), true))!;
 			logFile = (await ((ILogFileService)_logFileService).WriteLogDataAsync(logFile, 6, 1, Encoding.ASCII.GetBytes("foo\nbar\n"), true))!;
@@ -86,14 +87,14 @@ namespace Horde.Build.Tests
             Assert.AreEqual("hello", await ReadLogFile(_logFileService, logFile, 0, 5));
             Assert.AreEqual("foo\nbar\nbaz\n", await ReadLogFile(_logFileService, logFile, 6, 12));
 
-            LogMetadata metadata = await _logFileService.GetMetadataAsync(logFile);
+            LogMetadata metadata = await _logFileService.GetMetadataAsync(logFile, CancellationToken.None);
             Assert.AreEqual(6 + 8 + 4, metadata.Length);
             Assert.AreEqual(4, metadata.MaxLineIndex);
 
-            Assert.AreEqual((0, 0), await _logFileService.GetLineOffsetAsync(logFile, 0));
-            Assert.AreEqual((1, 6), await _logFileService.GetLineOffsetAsync(logFile, 1));
-            Assert.AreEqual((2, 10), await _logFileService.GetLineOffsetAsync(logFile, 2));
-            Assert.AreEqual((3, 14), await _logFileService.GetLineOffsetAsync(logFile, 3));
+            Assert.AreEqual((0, 0), await _logFileService.GetLineOffsetAsync(logFile, 0, CancellationToken.None));
+            Assert.AreEqual((1, 6), await _logFileService.GetLineOffsetAsync(logFile, 1, CancellationToken.None));
+            Assert.AreEqual((2, 10), await _logFileService.GetLineOffsetAsync(logFile, 2, CancellationToken.None));
+            Assert.AreEqual((3, 14), await _logFileService.GetLineOffsetAsync(logFile, 3, CancellationToken.None));
         }
         
         [TestMethod]
@@ -105,7 +106,7 @@ namespace Horde.Build.Tests
         
         private static async Task AssertMetadata(ILogFileService logFileService, ILogFile logFile, long expectedLength, long expectedMaxLineIndex)
         {
-	        LogMetadata metadata = await logFileService.GetMetadataAsync(logFile);
+	        LogMetadata metadata = await logFileService.GetMetadataAsync(logFile, CancellationToken.None);
 	        Assert.AreEqual(expectedLength, metadata.Length);
 	        Assert.AreEqual(expectedMaxLineIndex, metadata.MaxLineIndex);
         }
@@ -113,7 +114,7 @@ namespace Horde.Build.Tests
         private static async Task AssertChunk(ILogFileService logFileService, LogId logFileId, long numChunks, int chunkId, long offset, long length,
 	        long lineIndex)
         {
-	        ILogFile? logFile = await logFileService.GetLogFileAsync(logFileId);
+	        ILogFile? logFile = await logFileService.GetLogFileAsync(logFileId, CancellationToken.None);
 	        Assert.AreEqual(numChunks, logFile!.Chunks.Count);
 
 	        ILogChunk chunk = logFile.Chunks[chunkId];
@@ -124,8 +125,8 @@ namespace Horde.Build.Tests
 
         private static async Task AssertLineOffset(ILogFileService logFileService, LogId logFileId, int lineIndex, int clampedLineIndex, long offset)
         {
-	        ILogFile? logFile = await logFileService.GetLogFileAsync(logFileId);
-	        Assert.AreEqual((clampedLineIndex, offset), await logFileService.GetLineOffsetAsync(logFile!, lineIndex));
+	        ILogFile? logFile = await logFileService.GetLogFileAsync(logFileId, CancellationToken.None);
+	        Assert.AreEqual((clampedLineIndex, offset), await logFileService.GetLineOffsetAsync(logFile!, lineIndex, CancellationToken.None));
         }
         
         protected static async Task<string> ReadLogFile(ILogFileService logFileService, ILogFile logFile, long offset, long length)
@@ -245,16 +246,16 @@ namespace Horde.Build.Tests
 			// Will implicitly test GetLogFileAsync(), AddCachedLogFile()
 			JobId jobId = JobId.GenerateNewId();
             ObjectId sessionId = ObjectId.GenerateNewId();
-            ILogFile a = await _logFileService.CreateLogFileAsync(jobId, new ObjectId<ISession>(sessionId), LogType.Text);
-            ILogFile b = (await _logFileService.GetCachedLogFileAsync(a.Id))!;
+            ILogFile a = await _logFileService.CreateLogFileAsync(jobId, new ObjectId<ISession>(sessionId), LogType.Text, null, CancellationToken.None);
+            ILogFile b = (await _logFileService.GetCachedLogFileAsync(a.Id, CancellationToken.None))!;
             Assert.AreEqual(a.JobId, b.JobId);
             Assert.AreEqual(a.SessionId, b.SessionId);
             Assert.AreEqual(a.Type, b.Type);
 
-            ILogFile? notFound = await _logFileService.GetCachedLogFileAsync(LogId.GenerateNewId());
+            ILogFile? notFound = await _logFileService.GetCachedLogFileAsync(LogId.GenerateNewId(), CancellationToken.None);
             Assert.IsNull(notFound);
 
-            await _logFileService.CreateLogFileAsync(JobId.GenerateNewId(), new ObjectId<ISession>(ObjectId.GenerateNewId()), LogType.Text);
+            await _logFileService.CreateLogFileAsync(JobId.GenerateNewId(), new ObjectId<ISession>(ObjectId.GenerateNewId()), LogType.Text, null, CancellationToken.None);
             Assert.AreEqual(2, (await _logFileService.GetLogFilesAsync()).Count);
         }
 
@@ -263,8 +264,8 @@ namespace Horde.Build.Tests
         {
 			JobId jobId = JobId.GenerateNewId();
             ObjectId sessionId = ObjectId.GenerateNewId();
-            ILogFile logFile = await _logFileService.CreateLogFileAsync(jobId, new ObjectId<ISession>(sessionId), LogType.Text);
-            ILogFile logFileNoSession = await _logFileService.CreateLogFileAsync(jobId, null, LogType.Text);
+            ILogFile logFile = await _logFileService.CreateLogFileAsync(jobId, new ObjectId<ISession>(sessionId), LogType.Text, null, CancellationToken.None);
+            ILogFile logFileNoSession = await _logFileService.CreateLogFileAsync(jobId, null, LogType.Text, null, CancellationToken.None);
 
 			ClaimsPrincipal hasClaim = new ClaimsPrincipal(new ClaimsIdentity(new List<Claim>
             {
@@ -284,7 +285,7 @@ namespace Horde.Build.Tests
 		public async Task ChunkSplitting()
 		{
 			JobId jobId = JobId.GenerateNewId();
-			ILogFile logFile = await _logFileService.CreateLogFileAsync(jobId, null, LogType.Text);
+			ILogFile logFile = await _logFileService.CreateLogFileAsync(jobId, null, LogType.Text, null, CancellationToken.None);
 
 			long offset = 0;
 
@@ -292,27 +293,27 @@ namespace Horde.Build.Tests
 			const int MaxSubChunkLineCount = 128;
 
 			byte[] line1 = Encoding.UTF8.GetBytes("hello world\n");
-			await _logFileService.WriteLogDataAsync(logFile, offset, 0, line1, false, MaxChunkSize, MaxSubChunkLineCount);
+			await _logFileService.WriteLogDataAsync(logFile, offset, 0, line1, false, MaxChunkSize, MaxSubChunkLineCount, CancellationToken.None);
 			offset += line1.Length;
 
 			byte[] line2 = Encoding.UTF8.GetBytes("ab\n");
-			await _logFileService.WriteLogDataAsync(logFile, offset, 1, line2, false, MaxChunkSize, MaxSubChunkLineCount);
+			await _logFileService.WriteLogDataAsync(logFile, offset, 1, line2, false, MaxChunkSize, MaxSubChunkLineCount, CancellationToken.None);
 			offset += line2.Length;
 
 			byte[] line3 = Encoding.UTF8.GetBytes("a\n");
-			await _logFileService.WriteLogDataAsync(logFile, offset, 2, line3, false, MaxChunkSize, MaxSubChunkLineCount);
+			await _logFileService.WriteLogDataAsync(logFile, offset, 2, line3, false, MaxChunkSize, MaxSubChunkLineCount, CancellationToken.None);
 			offset += line3.Length;
 
 			byte[] line4 = Encoding.UTF8.GetBytes("b\n");
-			await _logFileService.WriteLogDataAsync(logFile, offset, 3, line4, false, MaxChunkSize, MaxSubChunkLineCount);
+			await _logFileService.WriteLogDataAsync(logFile, offset, 3, line4, false, MaxChunkSize, MaxSubChunkLineCount, CancellationToken.None);
 			offset += line4.Length;
 
 			byte[] line5 = Encoding.UTF8.GetBytes("a\nb\n");
-			await _logFileService.WriteLogDataAsync(logFile, offset, 4, line5, false, MaxChunkSize, MaxSubChunkLineCount);
+			await _logFileService.WriteLogDataAsync(logFile, offset, 4, line5, false, MaxChunkSize, MaxSubChunkLineCount, CancellationToken.None);
 			// offset += line5.Length;
 
 			await _logFileService.FlushAsync();
-			logFile = (await _logFileService.GetLogFileAsync(logFile.Id))!;
+			logFile = (await _logFileService.GetLogFileAsync(logFile.Id, CancellationToken.None))!;
 			Assert.AreEqual(4, logFile.Chunks.Count);
 			Assert.AreEqual(6, logFile.MaxLineIndex);
 
diff --git a/Engine/Source/Programs/Horde/Horde.Build.Tests/LogIndexingTests.cs b/Engine/Source/Programs/Horde/Horde.Build.Tests/LogIndexingTests.cs
--- a/Engine/Source/Programs/Horde/Horde.Build.Tests/LogIndexingTests.cs
+++ b/Engine/Source/Programs/Horde/Horde.Build.Tests/LogIndexingTests.cs
@@ -5,6 +5,7 @@ using System.Collections.Generic;
 using System.IO;
 using System.Linq;
 using System.Text;
+using System.Threading;
 using System.Threading.Tasks;
 using Horde.Build.Compute.Tests.Properties;
 using Horde.Build.Logs;
@@ -161,7 +162,7 @@ namespace Horde.Build.Tests
 						string str = lines[lineIdx].Substring(strOfs, strLen);
 
 						LogSearchStats stats = new LogSearchStats();
-						List<int> results = await _logFileService.SearchLogDataAsync(logFile, str, 0, 5, stats);
+						List<int> results = await _logFileService.SearchLogDataAsync(logFile, str, 0, 5, stats, CancellationToken.None);
 						Assert.AreEqual(1, results.Count);
 						Assert.AreEqual(lineIdx, results[0]);
 
@@ -187,7 +188,7 @@ namespace Horde.Build.Tests
 
 			{
 				LogSearchStats stats = new LogSearchStats();
-				List<int> results = await _logFileService.SearchLogDataAsync(logFile, "abc", 0, 5, stats);
+				List<int> results = await _logFileService.SearchLogDataAsync(logFile, "abc", 0, 5, stats, CancellationToken.None);
 				Assert.AreEqual(1, results.Count);
 				Assert.AreEqual(0, results[0]);
 
@@ -197,7 +198,7 @@ namespace Horde.Build.Tests
 			}
 			{
 				LogSearchStats stats = new LogSearchStats();
-				List<int> results = await _logFileService.SearchLogDataAsync(logFile, "def", 0, 5, stats);
+				List<int> results = await _logFileService.SearchLogDataAsync(logFile, "def", 0, 5, stats, CancellationToken.None);
 				Assert.AreEqual(1, results.Count);
 				Assert.AreEqual(1, results[0]);
 
@@ -207,7 +208,7 @@ namespace Horde.Build.Tests
 			}
 			{
 				LogSearchStats stats = new LogSearchStats();
-				List<int> results = await _logFileService.SearchLogDataAsync(logFile, "ghi", 0, 5, stats);
+				List<int> results = await _logFileService.SearchLogDataAsync(logFile, "ghi", 0, 5, stats, CancellationToken.None);
 				Assert.AreEqual(1, results.Count);
 				Assert.AreEqual(2, results[0]);
 
@@ -239,7 +240,7 @@ namespace Horde.Build.Tests
 		async Task SearchLogDataTestAsync(ILogFile logFile, string text, int firstLine, int count, int[] expectedLines)
 		{
 			LogSearchStats stats = new LogSearchStats();
-			List<int> lines = await _logFileService.SearchLogDataAsync(logFile, text, firstLine, count, stats);
+			List<int> lines = await _logFileService.SearchLogDataAsync(logFile, text, firstLine, count, stats, CancellationToken.None);
 			Assert.IsTrue(lines.SequenceEqual(expectedLines));
 		}
 	}
diff --git a/Engine/Source/Programs/Horde/Horde.Build/Issues/IssueService.cs b/Engine/Source/Programs/Horde/Horde.Build/Issues/IssueService.cs
--- a/Engine/Source/Programs/Horde/Horde.Build/Issues/IssueService.cs
+++ b/Engine/Source/Programs/Horde/Horde.Build/Issues/IssueService.cs
@@ -615,7 +615,7 @@ namespace Horde.Build.Issues
 			}
 
 			// Find the log file for this step
-			ILogFile? logFile = await _logFileService.GetLogFileAsync(step.LogId.Value);
+			ILogFile? logFile = await _logFileService.GetLogFileAsync(step.LogId.Value, CancellationToken.None);
 			if (logFile == null)
 			{
 				throw new ArgumentException($"Unable to retrieve log {step.LogId}");
diff --git a/Engine/Source/Programs/Horde/Horde.Build/Issues/IssuesController.cs b/Engine/Source/Programs/Horde/Horde.Build/Issues/IssuesController.cs
--- a/Engine/Source/Programs/Horde/Horde.Build/Issues/IssuesController.cs
+++ b/Engine/Source/Programs/Horde/Horde.Build/Issues/IssuesController.cs
@@ -3,6 +3,7 @@
 using System;
 using System.Collections.Generic;
 using System.Linq;
+using System.Threading;
 using System.Threading.Tasks;
 using Horde.Build.Acls;
 using Horde.Build.Auditing;
@@ -433,11 +434,22 @@ namespace Horde.Build.Issues
 		/// <param name="index">Index of the first event</param>
 		/// <param name="count">Number of events to return</param>
 		/// <param name="filter">Filter for the properties to return</param>
+		/// <param name="cancellationToken">Cancellation token for the request</param>
 		/// <returns>List of matching agents</returns>
 		[HttpGet]
 		[Route("/api/v1/issues/{issueId}/events")]
 		[ProducesResponseType(typeof(List<GetLogEventResponse>), 200)]
-		public async Task<ActionResult<List<object>>> GetIssueEventsAsync(int issueId, [FromQuery] JobId? jobId = null, [FromQuery] string? batchId = null, [FromQuery] string? stepId = null, [FromQuery(Name = "label")] int? labelIdx = null, [FromQuery] string[]? logIds = null, [FromQuery] int index = 0, [FromQuery] int count = 10, [FromQuery] PropertyFilter? filter = null)
+		public async Task<ActionResult<List<object>>> GetIssueEventsAsync(
+			int issueId,
+			[FromQuery] JobId? jobId = null,
+			[FromQuery] string? batchId = null,
+			[FromQuery] string? stepId = null,
+			[FromQuery(Name = "label")] int? labelIdx = null,
+			[FromQuery] string[]? logIds = null,
+			[FromQuery] int index = 0, 
+			[FromQuery] int count = 10,
+			[FromQuery] PropertyFilter? filter = null,
+			CancellationToken cancellationToken = default)
 		{
 			HashSet<LogId> logIdValues = new HashSet<LogId>();
 			if(jobId != null)
@@ -493,7 +505,7 @@ namespace Horde.Build.Issues
 			}
 
 			List<IIssueSpan> spans = await _issueCollection.FindSpansAsync(issueId);
-			List<ILogEvent> events = await _logFileService.FindEventsForSpansAsync(spans.Select(x => x.Id), logIdValues.ToArray(), index, count);
+			List<ILogEvent> events = await _logFileService.FindEventsForSpansAsync(spans.Select(x => x.Id), logIdValues.ToArray(), index, count, cancellationToken);
 
 			JobPermissionsCache permissionsCache = new JobPermissionsCache();
 			Dictionary<LogId, ILogFile?> logFiles = new Dictionary<LogId, ILogFile?>();
@@ -504,12 +516,12 @@ namespace Horde.Build.Issues
 				ILogFile? logFile;
 				if (!logFiles.TryGetValue(logEvent.LogId, out logFile))
 				{
-					logFile = await _logFileService.GetLogFileAsync(logEvent.LogId);
+					logFile = await _logFileService.GetLogFileAsync(logEvent.LogId, cancellationToken);
 					logFiles[logEvent.LogId] = logFile;
 				}
 				if (logFile != null && await _jobService.AuthorizeAsync(logFile.JobId, AclAction.ViewLog, User, permissionsCache))
 				{
-					ILogEventData data = await _logFileService.GetEventDataAsync(logFile, logEvent.LineIndex, logEvent.LineCount);
+					ILogEventData data = await _logFileService.GetEventDataAsync(logFile, logEvent.LineIndex, logEvent.LineCount, cancellationToken);
 					GetLogEventResponse response = new GetLogEventResponse(logEvent, data, issueId);
 					responses.Add(PropertyFilter.Apply(response, filter));
 				}
diff --git a/Engine/Source/Programs/Horde/Horde.Build/Logs/ILogFileCollection.cs b/Engine/Source/Programs/Horde/Horde.Build/Logs/ILogFileCollection.cs
--- a/Engine/Source/Programs/Horde/Horde.Build/Logs/ILogFileCollection.cs
+++ b/Engine/Source/Programs/Horde/Horde.Build/Logs/ILogFileCollection.cs
@@ -1,6 +1,7 @@
 // Copyright Epic Games, Inc. All Rights Reserved.
 
 using System.Collections.Generic;
+using System.Threading;
 using System.Threading.Tasks;
 using Horde.Build.Agents.Sessions;
 using Horde.Build.Jobs;
@@ -58,8 +59,9 @@ namespace Horde.Build.Logs
 		/// <param name="sessionId">Agent session allowed to update the log</param>
 		/// <param name="type">Type of events to be stored in the log</param>
 		/// <param name="logId">ID of the log file (optional)</param>
+		/// <param name="cancellationToken">Cancellation token for the call</param>
 		/// <returns>The new log file document</returns>
-		Task<ILogFile> CreateLogFileAsync(JobId jobId, SessionId? sessionId, LogType type, LogId? logId = null);
+		Task<ILogFile> CreateLogFileAsync(JobId jobId, SessionId? sessionId, LogType type, LogId? logId = null, CancellationToken cancellationToken = default);
 
 		/// <summary>
 		/// Adds a new chunk
@@ -67,36 +69,43 @@ namespace Horde.Build.Logs
 		/// <param name="logFileInterface">The current log file</param>
 		/// <param name="offset">Offset of the new chunk</param>
 		/// <param name="lineIndex">Line index for the start of the chunk</param>
+		/// <param name="cancellationToken">Cancellation token for the call</param>
 		/// <returns>The updated log file document</returns>
-		Task<ILogFile?> TryAddChunkAsync(ILogFile logFileInterface, long offset, int lineIndex);
+		Task<ILogFile?> TryAddChunkAsync(ILogFile logFileInterface, long offset, int lineIndex, CancellationToken cancellationToken);
 
 		/// <summary>
 		/// Update the log file with final information about certain chunks
 		/// </summary>
 		/// <param name="logFileInterface">The current log file</param>
 		/// <param name="chunks">Chunks to update. New chunks will be inserted</param>
+		/// <param name="cancellationToken">Cancellation token for the call</param>
 		/// <returns>The updated log file document</returns>
-		Task<ILogFile?> TryCompleteChunksAsync(ILogFile logFileInterface, IEnumerable<CompleteLogChunkUpdate> chunks);
+		Task<ILogFile?> TryCompleteChunksAsync(ILogFile logFileInterface, IEnumerable<CompleteLogChunkUpdate> chunks, CancellationToken cancellationToken);
 
 		/// <summary>
 		/// Update the log file with final information about the index
 		/// </summary>
 		/// <param name="logFileInterface">The current log file</param>
 		/// <param name="newIndexLength">New length of the index</param>
+		/// <param name="cancellationToken">Cancellation token for the call</param>
 		/// <returns>The updated log file document</returns>
-		Task<ILogFile?> TryUpdateIndexAsync(ILogFile logFileInterface, long newIndexLength);
+		Task<ILogFile?> TryUpdateIndexAsync(ILogFile logFileInterface, long newIndexLength, CancellationToken cancellationToken);
 
 		/// <summary>
 		/// Gets a logfile by ID
 		/// </summary>
 		/// <param name="logFileId">Unique id of the log file</param>
+		/// <param name="cancellationToken">Cancellation token for the call</param>
 		/// <returns>The logfile document</returns>
-		Task<ILogFile?> GetLogFileAsync(LogId logFileId);
+		Task<ILogFile?> GetLogFileAsync(LogId logFileId, CancellationToken cancellationToken);
 
 		/// <summary>
 		/// Gets all the log files
 		/// </summary>
+		/// <param name="index">Index of the first result to return</param>
+		/// <param name="count">Number of results to return</param>
+		/// <param name="cancellationToken">Cancellation token for the call</param>
 		/// <returns>List of log files</returns>
-		Task<List<ILogFile>> GetLogFilesAsync(int? index = null, int? count = null);
+		Task<List<ILogFile>> GetLogFilesAsync(int? index = null, int? count = null, CancellationToken cancellationToken = default);
 	}
 }
diff --git a/Engine/Source/Programs/Horde/Horde.Build/Logs/LogFileCollection.cs b/Engine/Source/Programs/Horde/Horde.Build/Logs/LogFileCollection.cs
--- a/Engine/Source/Programs/Horde/Horde.Build/Logs/LogFileCollection.cs
+++ b/Engine/Source/Programs/Horde/Horde.Build/Logs/LogFileCollection.cs
@@ -3,6 +3,7 @@
 using System.Collections.Generic;
 using System.Linq;
 using System.Net;
+using System.Threading;
 using System.Threading.Tasks;
 using Horde.Build.Agents.Sessions;
 using Horde.Build.Jobs;
@@ -117,15 +118,15 @@ namespace Horde.Build.Logs
 		}
 
 		/// <inheritdoc/>
-		public async Task<ILogFile> CreateLogFileAsync(JobId jobId, SessionId? sessionId, LogType type, LogId? logId)
+		public async Task<ILogFile> CreateLogFileAsync(JobId jobId, SessionId? sessionId, LogType type, LogId? logId, CancellationToken cancellationToken)
 		{
 			LogFileDocument newLogFile = new (jobId, sessionId, type, logId);
-			await _logFiles.InsertOneAsync(newLogFile);
+			await _logFiles.InsertOneAsync(newLogFile, null, cancellationToken);
 			return newLogFile;
 		}
 
 		/// <inheritdoc/>
-		public async Task<ILogFile?> TryAddChunkAsync(ILogFile logFileInterface, long offset, int lineIndex)
+		public async Task<ILogFile?> TryAddChunkAsync(ILogFile logFileInterface, long offset, int lineIndex, CancellationToken cancellationToken)
 		{
 			LogFileDocument logFile = ((LogFileDocument)logFileInterface).Clone();
 
@@ -144,7 +145,7 @@ namespace Horde.Build.Logs
 				update = update.Unset(x => x.MaxLineIndex);
 			}
 
-			if (!await TryUpdateLogFileAsync(logFile, update))
+			if (!await TryUpdateLogFileAsync(logFile, update, cancellationToken))
 			{
 				return null;
 			}
@@ -153,7 +154,7 @@ namespace Horde.Build.Logs
 		}
 
 		/// <inheritdoc/>
-		public async Task<ILogFile?> TryCompleteChunksAsync(ILogFile logFileInterface, IEnumerable<CompleteLogChunkUpdate> chunkUpdates)
+		public async Task<ILogFile?> TryCompleteChunksAsync(ILogFile logFileInterface, IEnumerable<CompleteLogChunkUpdate> chunkUpdates, CancellationToken cancellationToken)
 		{
 			LogFileDocument logFile = ((LogFileDocument)logFileInterface).Clone();
 
@@ -174,7 +175,7 @@ namespace Horde.Build.Logs
 			}
 
 			// Try to apply the updates
-			if (updates.Count > 0 && !await TryUpdateLogFileAsync(logFile, updateBuilder.Combine(updates)))
+			if (updates.Count > 0 && !await TryUpdateLogFileAsync(logFile, updateBuilder.Combine(updates), cancellationToken))
 			{
 				return null;
 			}
@@ -183,12 +184,12 @@ namespace Horde.Build.Logs
 		}
 
 		/// <inheritdoc/>
-		public async Task<ILogFile?> TryUpdateIndexAsync(ILogFile logFileInterface, long newIndexLength)
+		public async Task<ILogFile?> TryUpdateIndexAsync(ILogFile logFileInterface, long newIndexLength, CancellationToken cancellationToken)
 		{
 			LogFileDocument logFile = ((LogFileDocument)logFileInterface).Clone();
 
 			UpdateDefinition<LogFileDocument> update = Builders<LogFileDocument>.Update.Set(x => x.IndexLength, newIndexLength);
-			if (!await TryUpdateLogFileAsync(logFile, update))
+			if (!await TryUpdateLogFileAsync(logFile, update, cancellationToken))
 			{
 				return null;
 			}
@@ -198,23 +199,23 @@ namespace Horde.Build.Logs
 		}
 
 		/// <inheritdoc/>
-		private async Task<bool> TryUpdateLogFileAsync(LogFileDocument current, UpdateDefinition<LogFileDocument> update)
+		private async Task<bool> TryUpdateLogFileAsync(LogFileDocument current, UpdateDefinition<LogFileDocument> update, CancellationToken cancellationToken)
 		{
 			int prevUpdateIndex = current.UpdateIndex;
 			current.UpdateIndex++;
-			UpdateResult result = await _logFiles.UpdateOneAsync<LogFileDocument>(x => x.Id == current.Id && x.UpdateIndex == prevUpdateIndex, update.Set(x => x.UpdateIndex, current.UpdateIndex));
+			UpdateResult result = await _logFiles.UpdateOneAsync<LogFileDocument>(x => x.Id == current.Id && x.UpdateIndex == prevUpdateIndex, update.Set(x => x.UpdateIndex, current.UpdateIndex), cancellationToken: cancellationToken);
 			return result.ModifiedCount == 1;
 		}
 
 		/// <inheritdoc/>
-		public async Task<ILogFile?> GetLogFileAsync(LogId logFileId)
+		public async Task<ILogFile?> GetLogFileAsync(LogId logFileId, CancellationToken cancellationToken)
 		{
-			LogFileDocument logFile = await _logFiles.Find<LogFileDocument>(x => x.Id == logFileId).FirstOrDefaultAsync();
+			LogFileDocument logFile = await _logFiles.Find<LogFileDocument>(x => x.Id == logFileId).FirstOrDefaultAsync(cancellationToken);
 			return logFile;
 		}
 
 		/// <inheritdoc/>
-		public async Task<List<ILogFile>> GetLogFilesAsync(int? index = null, int? count = null)
+		public async Task<List<ILogFile>> GetLogFilesAsync(int? index = null, int? count = null, CancellationToken cancellationToken = default)
 		{
 			IFindFluent<LogFileDocument, LogFileDocument> query = _logFiles.Find(FilterDefinition<LogFileDocument>.Empty);
 			if(index != null)
@@ -226,7 +227,7 @@ namespace Horde.Build.Logs
 				query = query.Limit(count.Value);
 			}
 
-			List<LogFileDocument> results = await query.ToListAsync();
+			List<LogFileDocument> results = await query.ToListAsync(cancellationToken);
 			return results.ConvertAll<ILogFile>(x => x);
 		}
 	}
diff --git a/Engine/Source/Programs/Horde/Horde.Build/Logs/LogFileService.cs b/Engine/Source/Programs/Horde/Horde.Build/Logs/LogFileService.cs
--- a/Engine/Source/Programs/Horde/Horde.Build/Logs/LogFileService.cs
+++ b/Engine/Source/Programs/Horde/Horde.Build/Logs/LogFileService.cs
@@ -60,31 +60,35 @@ namespace Horde.Build.Logs
 		/// <param name="jobId">Unique id of the job that owns this log file</param>
 		/// <param name="sessionId">Agent session allowed to update the log</param>
 		/// <param name="type">Type of events to be stored in the log</param>
+		/// <param name="cancellationToken">Cancellation token for the call</param>
 		/// <param name="logId">ID of the log file (optional)</param>
 		/// <returns>The new log file document</returns>
-		Task<ILogFile> CreateLogFileAsync(JobId jobId, SessionId? sessionId, LogType type, LogId? logId = null);
+		Task<ILogFile> CreateLogFileAsync(JobId jobId, SessionId? sessionId, LogType type, LogId? logId = null, CancellationToken cancellationToken = default);
 
 		/// <summary>
 		/// Gets a logfile by ID
 		/// </summary>
 		/// <param name="logFileId">Unique id of the log file</param>
+		/// <param name="cancellationToken">Cancellation token for the call</param>
 		/// <returns>The logfile document</returns>
-		Task<ILogFile?> GetLogFileAsync(LogId logFileId);
+		Task<ILogFile?> GetLogFileAsync(LogId logFileId, CancellationToken cancellationToken);
 
 		/// <summary>
 		/// Gets a logfile by ID, returning a cached copy if available. This should only be used to retrieve constant properties set at creation, such as the session or job it's associated with.
 		/// </summary>
 		/// <param name="logFileId">Unique id of the log file</param>
+		/// <param name="cancellationToken">Cancellation token for the call</param>
 		/// <returns>The logfile document</returns>
-		Task<ILogFile?> GetCachedLogFileAsync(LogId logFileId);
+		Task<ILogFile?> GetCachedLogFileAsync(LogId logFileId, CancellationToken cancellationToken);
 
 		/// <summary>
 		/// Returns a list of log files
 		/// </summary>
 		/// <param name="index">Index of the first result to return</param>
 		/// <param name="count">Number of results to return</param>
+		/// <param name="cancellationToken">Cancellation token for the call</param>
 		/// <returns>List of logfile documents</returns>
-		Task<List<ILogFile>> GetLogFilesAsync(int? index = null, int? count = null);
+		Task<List<ILogFile>> GetLogFilesAsync(int? index = null, int? count = null, CancellationToken cancellationToken = default);
 
 		/// <summary>
 		/// Writes out chunk data and assigns to a file
@@ -96,22 +100,25 @@ namespace Horde.Build.Logs
 		/// <param name="flush">Whether the current chunk is complete and should be flushed</param>
 		/// <param name="maxChunkLength">The maximum chunk length. Defaults to 128kb.</param>
 		/// <param name="maxSubChunkLineCount">Maximum number of lines in each sub-chunk.</param>
+		/// <param name="cancellationToken">Cancellation token for the call</param>
 		/// <returns></returns>
-		Task<ILogFile?> WriteLogDataAsync(ILogFile logFile, long offset, int lineIndex, ReadOnlyMemory<byte> data, bool flush, int maxChunkLength = 256 * 1024, int maxSubChunkLineCount = 128);
+		Task<ILogFile?> WriteLogDataAsync(ILogFile logFile, long offset, int lineIndex, ReadOnlyMemory<byte> data, bool flush, int maxChunkLength = 256 * 1024, int maxSubChunkLineCount = 128, CancellationToken cancellationToken = default);
 
 		/// <summary>
 		/// Gets metadata about the log file
 		/// </summary>
 		/// <param name="logFile">The log file to query</param>
+		/// <param name="cancellationToken">Cancellation token for the call</param>
 		/// <returns>Metadata about the log file</returns>
-		Task<LogMetadata> GetMetadataAsync(ILogFile logFile);
+		Task<LogMetadata> GetMetadataAsync(ILogFile logFile, CancellationToken cancellationToken);
 
 		/// <summary>
 		/// Creates new log events
 		/// </summary>
 		/// <param name="newEvents">List of events</param>
+		/// <param name="cancellationToken">Cancellation token for the call</param>
 		/// <returns>Async task</returns>
-		Task CreateEventsAsync(List<NewLogEventData> newEvents);
+		Task CreateEventsAsync(List<NewLogEventData> newEvents, CancellationToken cancellationToken);
 
 		/// <summary>
 		/// Find events for a particular log file
@@ -120,16 +127,18 @@ namespace Horde.Build.Logs
 		/// <param name="spanId">Issue span to return events for</param>
 		/// <param name="index">Index of the first event to retrieve</param>
 		/// <param name="count">Number of events to retrieve</param>
+		/// <param name="cancellationToken">Cancellation token for the call</param>
 		/// <returns>List of log events</returns>
-		Task<List<ILogEvent>> FindEventsAsync(ILogFile logFile, ObjectId? spanId = null, int? index = null, int? count = null);
+		Task<List<ILogEvent>> FindEventsAsync(ILogFile logFile, ObjectId? spanId = null, int? index = null, int? count = null, CancellationToken cancellationToken = default);
 
 		/// <summary>
 		/// Adds events to a log span
 		/// </summary>
 		/// <param name="events">The events to add</param>
 		/// <param name="spanId">The span id</param>
+		/// <param name="cancellationToken">Cancellation token for the call</param>
 		/// <returns>Async task</returns>
-		Task AddSpanToEventsAsync(IEnumerable<ILogEvent> events, ObjectId spanId);
+		Task AddSpanToEventsAsync(IEnumerable<ILogEvent> events, ObjectId spanId, CancellationToken cancellationToken = default);
 
 		/// <summary>
 		/// Find events for an issue
@@ -138,8 +147,9 @@ namespace Horde.Build.Logs
 		/// <param name="logIds">Log ids to include</param>
 		/// <param name="index">Index within the events for results to return</param>
 		/// <param name="count">Number of results to return</param>
+		/// <param name="cancellationToken">Cancellation token for the call</param>
 		/// <returns>Async task</returns>
-		Task<List<ILogEvent>> FindEventsForSpansAsync(IEnumerable<ObjectId> spanIds, LogId[]? logIds, int index, int count);
+		Task<List<ILogEvent>> FindEventsForSpansAsync(IEnumerable<ObjectId> spanIds, LogId[]? logIds, int index, int count, CancellationToken cancellationToken = default);
 
 		/// <summary>
 		/// Gets the data for an event
@@ -147,8 +157,9 @@ namespace Horde.Build.Logs
 		/// <param name="logFile">The log file instance</param>
 		/// <param name="lineIndex">Index of the line in the file</param>
 		/// <param name="lineCount">Number of lines in the event</param>
+		/// <param name="cancellationToken">Cancellation token for the call</param>
 		/// <returns>New event data instance</returns>
-		Task<ILogEventData> GetEventDataAsync(ILogFile logFile, int lineIndex, int lineCount);
+		Task<ILogEventData> GetEventDataAsync(ILogFile logFile, int lineIndex, int lineCount, CancellationToken cancellationToken = default);
 
 		/// <summary>
 		/// Gets lines from the given log 
@@ -156,8 +167,9 @@ namespace Horde.Build.Logs
 		/// <param name="logFile">The log file</param>
 		/// <param name="offset">Offset of the data to return</param>
 		/// <param name="length">Length of the data to return</param>
+		/// <param name="cancellationToken">Cancellation token for the call</param>
 		/// <returns>Data for the requested range</returns>
-		Task<Stream> OpenRawStreamAsync(ILogFile logFile, long offset, long length);
+		Task<Stream> OpenRawStreamAsync(ILogFile logFile, long offset, long length, CancellationToken cancellationToken = default);
 
 		/// <summary>
 		/// Parses a stream of json text and outputs plain text
@@ -166,16 +178,18 @@ namespace Horde.Build.Logs
 		/// <param name="offset">Offset within the data to copy from</param>
 		/// <param name="length">Length of the data to copy</param>
 		/// <param name="outputStream">Output stream to receive the text data</param>
+		/// <param name="cancellationToken">Cancellation token for the call</param>
 		/// <returns>Async text</returns>
-		Task CopyPlainTextStreamAsync(ILogFile logFile, long offset, long length, Stream outputStream);
+		Task CopyPlainTextStreamAsync(ILogFile logFile, long offset, long length, Stream outputStream, CancellationToken cancellationToken = default);
 
 		/// <summary>
 		/// Gets the offset of the given line number
 		/// </summary>
 		/// <param name="logFile">The log file to search</param>
 		/// <param name="lineIdx">The line index to retrieve the offset for</param>
+		/// <param name="cancellationToken">Cancellation token for the call</param>
 		/// <returns>The actual clamped line number and offset</returns>
-		Task<(int, long)> GetLineOffsetAsync(ILogFile logFile, int lineIdx);
+		Task<(int, long)> GetLineOffsetAsync(ILogFile logFile, int lineIdx, CancellationToken cancellationToken);
 
 		/// <summary>
 		/// Search for the specified text in a log file
@@ -185,8 +199,9 @@ namespace Horde.Build.Logs
 		/// <param name="firstLine">Line to start search from</param>
 		/// <param name="count">Number of results to return</param>
 		/// <param name="stats">Receives stats for the search</param>
+		/// <param name="cancellationToken">Cancellation token for the call</param>
 		/// <returns>List of line numbers containing the given term</returns>
-		Task<List<int>> SearchLogDataAsync(ILogFile logFile, string text, int firstLine, int count, LogSearchStats stats);
+		Task<List<int>> SearchLogDataAsync(ILogFile logFile, string text, int firstLine, int count, LogSearchStats stats, CancellationToken cancellationToken);
 	}
 
 	/// <summary>
@@ -202,13 +217,12 @@ namespace Horde.Build.Logs
 		/// <param name="offset">Offset within the log file to copy</param>
 		/// <param name="length">Length of the data to copy</param>
 		/// <param name="outputStream">Output stream to receive the text data</param>
+		/// <param name="cancellationToken">Cancellation token for the call</param>
 		/// <returns>Async text</returns>
-		public static async Task CopyRawStreamAsync(this ILogFileService logFileService, ILogFile logFile, long offset, long length, Stream outputStream)
+		public static async Task CopyRawStreamAsync(this ILogFileService logFileService, ILogFile logFile, long offset, long length, Stream outputStream, CancellationToken cancellationToken)
 		{
-			using (Stream stream = await logFileService.OpenRawStreamAsync(logFile, offset, length))
-			{
-				await stream.CopyToAsync(outputStream);
-			}
+			await using Stream stream = await logFileService.OpenRawStreamAsync(logFile, offset, length, cancellationToken);
+			await stream.CopyToAsync(outputStream, cancellationToken);
 		}
 	}
 	
@@ -436,15 +450,15 @@ namespace Horde.Build.Logs
 		}
 
 		/// <inheritdoc/>
-		public Task<ILogFile> CreateLogFileAsync(JobId jobId, SessionId? sessionId, LogType type, LogId? logId = null)
+		public Task<ILogFile> CreateLogFileAsync(JobId jobId, SessionId? sessionId, LogType type, LogId? logId, CancellationToken cancellationToken)
 		{
-			return _logFiles.CreateLogFileAsync(jobId, sessionId, type, logId);
+			return _logFiles.CreateLogFileAsync(jobId, sessionId, type, logId, cancellationToken);
 		}
 
 		/// <inheritdoc/>
-		public async Task<ILogFile?> GetLogFileAsync(LogId logFileId)
+		public async Task<ILogFile?> GetLogFileAsync(LogId logFileId, CancellationToken cancellationToken)
 		{
-			ILogFile? logFile = await _logFiles.GetLogFileAsync(logFileId);
+			ILogFile? logFile = await _logFiles.GetLogFileAsync(logFileId, cancellationToken);
 			if(logFile != null)
 			{
 				AddCachedLogFile(logFile);
@@ -462,25 +476,22 @@ namespace Horde.Build.Logs
 			_logFileCache.Set(logFile.Id, logFile, options);
 		}
 
-		/// <summary>
-		/// Gets a cached log file by id
-		/// </summary>
-		/// <param name="logFileId">The log file id</param>
-		/// <returns>New log file, or null if not found</returns>
-		public async Task<ILogFile?> GetCachedLogFileAsync(LogId logFileId)
+
+		/// <inheritdoc />
+		public async Task<ILogFile?> GetCachedLogFileAsync(LogId logFileId, CancellationToken cancellationToken)
 		{
 			object? logFile;
 			if (!_logFileCache.TryGetValue(logFileId, out logFile))
 			{
-				logFile = await GetLogFileAsync(logFileId);
+				logFile = await GetLogFileAsync(logFileId, cancellationToken);
 			}
 			return (ILogFile?)logFile;
 		}
 
 		/// <inheritdoc/>
-		public Task<List<ILogFile>> GetLogFilesAsync(int? index = null, int? count = null)
+		public Task<List<ILogFile>> GetLogFilesAsync(int? index = null, int? count = null, CancellationToken cancellationToken = default)
 		{
-			return _logFiles.GetLogFilesAsync(index, count);
+			return _logFiles.GetLogFilesAsync(index, count, cancellationToken);
 		}
 
 		class WriteState
@@ -498,7 +509,7 @@ namespace Horde.Build.Logs
 		}
 
 		/// <inheritdoc/>
-		public async Task<ILogFile?> WriteLogDataAsync(ILogFile logFile, long offset, int lineIndex, ReadOnlyMemory<byte> data, bool flush, int maxChunkLength, int maxSubChunkLineCount)
+		public async Task<ILogFile?> WriteLogDataAsync(ILogFile logFile, long offset, int lineIndex, ReadOnlyMemory<byte> data, bool flush, int maxChunkLength, int maxSubChunkLineCount, CancellationToken cancellationToken)
 		{
 			using IScope scope = GlobalTracer.Instance.BuildSpan("WriteLogDataAsync").StartActive();
 			scope.Span.SetTag("LogId", logFile.Id.ToString());
@@ -530,7 +541,7 @@ namespace Horde.Build.Logs
 				if (chunkIdx >= 0)
 				{
 					ILogChunk chunk = logFile.Chunks[chunkIdx];
-					if (await WriteLogChunkDataAsync(logFile, chunk, state, completeOffsets, maxChunkLength, maxSubChunkLineCount))
+					if (await WriteLogChunkDataAsync(logFile, chunk, state, completeOffsets, maxChunkLength, maxSubChunkLineCount, cancellationToken))
 					{
 						continue;
 					}
@@ -540,17 +551,17 @@ namespace Horde.Build.Logs
 				ILogFile? newLogFile;
 				if (logFile.Chunks.Count == 0)
 				{
-					newLogFile = await _logFiles.TryAddChunkAsync(logFile, 0, 0);
+					newLogFile = await _logFiles.TryAddChunkAsync(logFile, 0, 0, cancellationToken);
 				}
 				else
 				{
-					newLogFile = await _logFiles.TryAddChunkAsync(logFile, state._offset, state._lineIndex);
+					newLogFile = await _logFiles.TryAddChunkAsync(logFile, state._offset, state._lineIndex, cancellationToken);
 				}
 
 				// Try to add a new chunk at the new location
 				if (newLogFile == null)
 				{
-					newLogFile = await _logFiles.GetLogFileAsync(logFile.Id);
+					newLogFile = await _logFiles.GetLogFileAsync(logFile.Id, cancellationToken);
 					if (newLogFile == null)
 					{
 						_logger.LogError("Unable to update log file {LogId}", logFile.Id);
@@ -581,7 +592,7 @@ namespace Horde.Build.Logs
 			// Write all the chunks
 			if (completeOffsets.Count > 0 || flush)
 			{
-				ILogFile? newLogFile = await WriteCompleteChunksForLogAsync(logFile, completeOffsets, flush);
+				ILogFile? newLogFile = await WriteCompleteChunksForLogAsync(logFile, completeOffsets, flush, cancellationToken);
 				if (newLogFile == null)
 				{
 					return null;
@@ -600,8 +611,9 @@ namespace Horde.Build.Logs
 		/// <param name="completeOffsets">List of complete chunks</param>
 		/// <param name="maxChunkLength">Maximum length of each chunk</param>
 		/// <param name="maxSubChunkLineCount">Maximum number of lines in each subchunk</param>
+		/// <param name="cancellationToken">Cancellation token for the call</param> 
 		/// <returns>True if data was appended to </returns>
-		private async Task<bool> WriteLogChunkDataAsync(ILogFile logFile, ILogChunk chunk, WriteState state, List<long> completeOffsets, int maxChunkLength, int maxSubChunkLineCount)
+		private async Task<bool> WriteLogChunkDataAsync(ILogFile logFile, ILogChunk chunk, WriteState state, List<long> completeOffsets, int maxChunkLength, int maxSubChunkLineCount, CancellationToken cancellationToken)
 		{
 			// Don't allow data to be appended if the chunk is complete
 			if(chunk.Length > 0)
@@ -694,7 +706,7 @@ namespace Horde.Build.Logs
 		}
 
 		/// <inheritdoc/>
-		public async Task<LogMetadata> GetMetadataAsync(ILogFile logFile)
+		public async Task<LogMetadata> GetMetadataAsync(ILogFile logFile, CancellationToken cancellationToken)
 		{
 			LogMetadata metadata = new LogMetadata();
 			if (logFile.Chunks.Count > 0)
@@ -716,13 +728,13 @@ namespace Horde.Build.Logs
 		}
 
 		/// <inheritdoc/>
-		public Task CreateEventsAsync(List<NewLogEventData> newEvents)
+		public Task CreateEventsAsync(List<NewLogEventData> newEvents, CancellationToken cancellationToken)
 		{
 			return _logEvents.AddManyAsync(newEvents);
 		}
 
 		/// <inheritdoc/>
-		public Task<List<ILogEvent>> FindEventsAsync(ILogFile logFile, ObjectId? spanId = null, int? index = null, int? count = null)
+		public Task<List<ILogEvent>> FindEventsAsync(ILogFile logFile, ObjectId? spanId = null, int? index = null, int? count = null, CancellationToken cancellationToken = default)
 		{
 			return _logEvents.FindAsync(logFile.Id, spanId, index, count);
 		}
@@ -788,32 +800,32 @@ namespace Horde.Build.Logs
 		}
 
 		/// <inheritdoc/>
-		public Task AddSpanToEventsAsync(IEnumerable<ILogEvent> events, ObjectId spanId)
+		public Task AddSpanToEventsAsync(IEnumerable<ILogEvent> events, ObjectId spanId, CancellationToken cancellationToken)
 		{
 			return _logEvents.AddSpanToEventsAsync(events, spanId);
 		}
 
 		/// <inheritdoc/>
-		public Task<List<ILogEvent>> FindEventsForSpansAsync(IEnumerable<ObjectId> spanIds, LogId[]? logIds, int index, int count)
+		public Task<List<ILogEvent>> FindEventsForSpansAsync(IEnumerable<ObjectId> spanIds, LogId[]? logIds, int index, int count, CancellationToken cancellationToken)
 		{
 			return _logEvents.FindEventsForSpansAsync(spanIds, logIds, index, count);
 		}
 
 		/// <inheritdoc/>
-		public async Task<ILogEventData> GetEventDataAsync(ILogFile logFile, int lineIndex, int lineCount)
+		public async Task<ILogEventData> GetEventDataAsync(ILogFile logFile, int lineIndex, int lineCount, CancellationToken cancellationToken)
 		{
 			using IScope scope = GlobalTracer.Instance.BuildSpan("GetEventDataAsync").StartActive();
 			scope.Span.SetTag("LogId", logFile.Id.ToString());
 			scope.Span.SetTag("LineIndex", lineIndex.ToString(CultureInfo.InvariantCulture));
 			scope.Span.SetTag("LineCount", lineCount.ToString(CultureInfo.InvariantCulture));
 
-			(_, long minOffset) = await GetLineOffsetAsync(logFile, lineIndex);
-			(_, long maxOffset) = await GetLineOffsetAsync(logFile, lineIndex + lineCount);
+			(_, long minOffset) = await GetLineOffsetAsync(logFile, lineIndex, cancellationToken);
+			(_, long maxOffset) = await GetLineOffsetAsync(logFile, lineIndex + lineCount, cancellationToken);
 
 			byte[] data = new byte[maxOffset - minOffset];
-			using (Stream stream = await OpenRawStreamAsync(logFile, minOffset, maxOffset - minOffset))
+			using (Stream stream = await OpenRawStreamAsync(logFile, minOffset, maxOffset - minOffset, cancellationToken))
 			{
-				int length = await stream.ReadAsync(data.AsMemory());
+				int length = await stream.ReadAsync(data.AsMemory(), cancellationToken);
 				if(length != data.Length)
 				{
 					_logger.LogWarning("Read less than expected from log stream (Expected {Expected}, Got {Got})", data.Length, length);
@@ -856,7 +868,7 @@ namespace Horde.Build.Logs
 		}
 
 		/// <inheritdoc/>
-		public async Task<Stream> OpenRawStreamAsync(ILogFile logFile, long offset, long length)
+		public async Task<Stream> OpenRawStreamAsync(ILogFile logFile, long offset, long length, CancellationToken cancellationToken)
 		{
 			if (logFile.Chunks.Count == 0)
 			{
@@ -902,9 +914,9 @@ namespace Horde.Build.Logs
 		}
 
 		/// <inheritdoc/>
-		public async Task CopyPlainTextStreamAsync(ILogFile logFile, long offset, long length, Stream outputStream)
+		public async Task CopyPlainTextStreamAsync(ILogFile logFile, long offset, long length, Stream outputStream, CancellationToken cancellationToken)
 		{
-			using (Stream stream = await OpenRawStreamAsync(logFile, 0, Int64.MaxValue))
+			using (Stream stream = await OpenRawStreamAsync(logFile, 0, Int64.MaxValue, cancellationToken))
 			{
 				byte[] readBuffer = new byte[4096];
 				int readBufferLength = 0;
@@ -915,7 +927,7 @@ namespace Horde.Build.Logs
 				while (length > 0)
 				{
 					// Add more data to the buffer
-					int readBytes = await stream.ReadAsync(readBuffer.AsMemory(readBufferLength, readBuffer.Length - readBufferLength));
+					int readBytes = await stream.ReadAsync(readBuffer.AsMemory(readBufferLength, readBuffer.Length - readBufferLength), cancellationToken);
 					readBufferLength += readBytes;
 
 					// Copy as many lines as possible to the output
@@ -935,7 +947,7 @@ namespace Horde.Build.Logs
 						if (offset < writeBufferLength)
 						{
 							int writeLength = (int)Math.Min((long)writeBufferLength - offset, length);
-							await outputStream.WriteAsync(writeBuffer.AsMemory((int)offset, writeLength));
+							await outputStream.WriteAsync(writeBuffer.AsMemory((int)offset, writeLength), cancellationToken);
 							length -= writeLength;
 						}
 						offset = Math.Max(offset - writeBufferLength, 0);
@@ -964,7 +976,7 @@ namespace Horde.Build.Logs
 		}
 
 		/// <inheritdoc/>
-		public async Task<(int, long)> GetLineOffsetAsync(ILogFile logFile, int lineIdx)
+		public async Task<(int, long)> GetLineOffsetAsync(ILogFile logFile, int lineIdx, CancellationToken cancellationToken)
 		{
 			int chunkIdx = logFile.Chunks.GetChunkForLine(lineIdx);
 
@@ -1178,7 +1190,7 @@ namespace Horde.Build.Logs
 		/// <returns>Async task</returns>
 		private async Task<ILogFile?> WriteCompleteChunksForLogAsync(LogId logId, List<long> offsets, bool createIndex, CancellationToken cancellationToken = default)
 		{
-			ILogFile? logFile = await _logFiles.GetLogFileAsync(logId);
+			ILogFile? logFile = await _logFiles.GetLogFileAsync(logId, cancellationToken);
 			if(logFile != null)
 			{
 				logFile = await WriteCompleteChunksForLogAsync(logFile, offsets, createIndex, cancellationToken);
@@ -1249,7 +1261,7 @@ namespace Horde.Build.Logs
 					}
 
 					// Try to apply the updates
-					ILogFile? newLogFile = await _logFiles.TryCompleteChunksAsync(logFile, updates);
+					ILogFile? newLogFile = await _logFiles.TryCompleteChunksAsync(logFile, updates, cancellationToken);
 					if (newLogFile != null)
 					{
 						logFile = newLogFile;
@@ -1257,7 +1269,7 @@ namespace Horde.Build.Logs
 					}
 
 					// Update the log file
-					logFile = await GetLogFileAsync(logFile.Id);
+					logFile = await GetLogFileAsync(logFile.Id, cancellationToken);
 				}
 			}
 
@@ -1266,7 +1278,7 @@ namespace Horde.Build.Logs
 			{
 				try
 				{
-					logFile = await CreateIndexAsync(logFile);
+					logFile = await CreateIndexAsync(logFile, cancellationToken);
 				}
 				catch(Exception ex)
 				{
@@ -1281,8 +1293,9 @@ namespace Horde.Build.Logs
 		/// Creates an index for the given log file
 		/// </summary>
 		/// <param name="logFile">The log file object</param>
+		/// <param name="cancellationToken">Cancellation token for the call</param>
 		/// <returns>Updated log file</returns>
-		private async Task<ILogFile?> CreateIndexAsync(ILogFile logFile)
+		private async Task<ILogFile?> CreateIndexAsync(ILogFile logFile, CancellationToken cancellationToken)
 		{
 			if(logFile.Chunks.Count == 0)
 			{
@@ -1366,12 +1379,12 @@ namespace Horde.Build.Logs
 
 				while(newLogFile != null && newLength > (newLogFile.IndexLength ?? 0))
 				{
-					newLogFile = await _logFiles.TryUpdateIndexAsync(newLogFile, newLength);
+					newLogFile = await _logFiles.TryUpdateIndexAsync(newLogFile, newLength, cancellationToken);
 					if(newLogFile != null)
 					{
 						break;
 					}
-					newLogFile = await _logFiles.GetLogFileAsync(logFile.Id);
+					newLogFile = await _logFiles.GetLogFileAsync(logFile.Id, cancellationToken);
 				}
 			}
 			return newLogFile;
@@ -1572,7 +1585,7 @@ namespace Horde.Build.Logs
 		}
 
 		/// <inheritdoc/>
-		public async Task<List<int>> SearchLogDataAsync(ILogFile logFile, string text, int firstLine, int count, LogSearchStats searchStats)
+		public async Task<List<int>> SearchLogDataAsync(ILogFile logFile, string text, int firstLine, int count, LogSearchStats searchStats, CancellationToken cancellationToken)
 		{
 			Stopwatch timer = Stopwatch.StartNew();
 
@@ -1584,7 +1597,7 @@ namespace Horde.Build.Logs
 			List<int> results = new List<int>();
 			if (count > 0)
 			{
-				await using IAsyncEnumerator<int> enumerator = SearchLogDataInternalAsync(logFile, text, firstLine, searchStats).GetAsyncEnumerator();
+				await using IAsyncEnumerator<int> enumerator = SearchLogDataInternalAsync(logFile, text, firstLine, searchStats).GetAsyncEnumerator(cancellationToken);
 				while (await enumerator.MoveNextAsync() && results.Count < count)
 				{
 					results.Add(enumerator.Current);
diff --git a/Engine/Source/Programs/Horde/Horde.Build/Logs/LogsController.cs b/Engine/Source/Programs/Horde/Horde.Build/Logs/LogsController.cs
--- a/Engine/Source/Programs/Horde/Horde.Build/Logs/LogsController.cs
+++ b/Engine/Source/Programs/Horde/Horde.Build/Logs/LogsController.cs
@@ -5,6 +5,7 @@ using System.Collections.Generic;
 using System.IO;
 using System.Security.Claims;
 using System.Text;
+using System.Threading;
 using System.Threading.Tasks;
 using Horde.Build.Acls;
 using Horde.Build.Issues;
@@ -65,13 +66,14 @@ namespace Horde.Build.Logs
 		/// </summary>
 		/// <param name="logFileId">Id of the log file to get information about</param>
 		/// <param name="filter">Filter for the properties to return</param>
+		/// <param name="cancellationToken">Cancellation token for the request</param>
 		/// <returns>Information about the requested project</returns>
 		[HttpGet]
 		[Route("/api/v1/logs/{logFileId}")]
 		[ProducesResponseType(typeof(GetLogFileResponse), 200)]
-		public async Task<ActionResult<object>> GetLog(LogId logFileId, [FromQuery] PropertyFilter? filter = null)
+		public async Task<ActionResult<object>> GetLog(LogId logFileId, [FromQuery] PropertyFilter? filter = null, CancellationToken cancellationToken = default)
 		{
-			ILogFile? logFile = await _logFileService.GetLogFileAsync(logFileId);
+			ILogFile? logFile = await _logFileService.GetLogFileAsync(logFileId, cancellationToken);
 			if (logFile == null)
 			{
 				return NotFound();
@@ -81,7 +83,7 @@ namespace Horde.Build.Logs
 				return Forbid();
 			}
 
-			LogMetadata metadata = await _logFileService.GetMetadataAsync(logFile);
+			LogMetadata metadata = await _logFileService.GetMetadataAsync(logFile, cancellationToken);
 			return new GetLogFileResponse(logFile, metadata).ApplyFilter(filter);       
 		}
 
@@ -94,12 +96,20 @@ namespace Horde.Build.Logs
 		/// <param name="length">Number of bytes to return</param>
 		/// <param name="fileName">Name of the default filename to download</param>
 		/// <param name="download">Whether to download the file rather than display in the browser</param>
+		/// <param name="cancellationToken">Cancellation token for the request</param>
 		/// <returns>Raw log data for the requested range</returns>
 		[HttpGet]
 		[Route("/api/v1/logs/{logFileId}/data")]
-		public async Task<ActionResult> GetLogData(LogId logFileId, [FromQuery] LogOutputFormat format = LogOutputFormat.Raw, [FromQuery] long offset = 0, [FromQuery] long length = Int64.MaxValue, [FromQuery] string? fileName = null, [FromQuery] bool download = false)
+		public async Task<ActionResult> GetLogData(
+			LogId logFileId,
+			[FromQuery] LogOutputFormat format = LogOutputFormat.Raw,
+			[FromQuery] long offset = 0,
+			[FromQuery] long length = Int64.MaxValue,
+			[FromQuery] string? fileName = null,
+			[FromQuery] bool download = false,
+			CancellationToken cancellationToken = default)
 		{
-			ILogFile? logFile = await _logFileService.GetLogFileAsync(logFileId);
+			ILogFile? logFile = await _logFileService.GetLogFileAsync(logFileId, cancellationToken);
 			if (logFile == null)
 			{
 				return NotFound();
@@ -112,11 +122,11 @@ namespace Horde.Build.Logs
 			Func<Stream, ActionContext, Task> copyTask;
 			if (format == LogOutputFormat.Text && logFile.Type == LogType.Json)
 			{
-				copyTask = (outputStream, context) => _logFileService.CopyPlainTextStreamAsync(logFile, offset, length, outputStream);
+				copyTask = (outputStream, context) => _logFileService.CopyPlainTextStreamAsync(logFile, offset, length, outputStream, cancellationToken);
 			}
 			else
 			{
-				copyTask = (outputStream, context) => _logFileService.CopyRawStreamAsync(logFile, offset, length, outputStream);
+				copyTask = (outputStream, context) => _logFileService.CopyRawStreamAsync(logFile, offset, length, outputStream, cancellationToken);
 			}
 
 			return new CustomFileCallbackResult(fileName ?? $"log-{logFileId}.txt", "text/plain", !download, copyTask);
@@ -128,12 +138,13 @@ namespace Horde.Build.Logs
 		/// <param name="logFileId">Id of the log file to get information about</param>
 		/// <param name="index">Index of the first line to retrieve</param>
 		/// <param name="count">Number of lines to retrieve</param>
+		/// <param name="cancellationToken">Cancellation token for the request</param>
 		/// <returns>Information about the requested project</returns>
 		[HttpGet]
 		[Route("/api/v1/logs/{logFileId}/lines")]
-		public async Task<ActionResult> GetLogLines(LogId logFileId, [FromQuery] int index = 0, [FromQuery] int count = Int32.MaxValue)
+		public async Task<ActionResult> GetLogLines(LogId logFileId, [FromQuery] int index = 0, [FromQuery] int count = Int32.MaxValue, CancellationToken cancellationToken = default)
 		{
-			ILogFile? logFile = await _logFileService.GetLogFileAsync(logFileId);
+			ILogFile? logFile = await _logFileService.GetLogFileAsync(logFileId, cancellationToken);
 			if (logFile == null)
 			{
 				return NotFound();
@@ -143,15 +154,15 @@ namespace Horde.Build.Logs
 				return Forbid();
 			}
 
-			LogMetadata metadata = await _logFileService.GetMetadataAsync(logFile);
+			LogMetadata metadata = await _logFileService.GetMetadataAsync(logFile, cancellationToken);
 
-			(int minIndex, long minOffset) = await _logFileService.GetLineOffsetAsync(logFile, index);
-			(int maxIndex, long maxOffset) = await _logFileService.GetLineOffsetAsync(logFile, index + Math.Min(count, Int32.MaxValue - index));
+			(int minIndex, long minOffset) = await _logFileService.GetLineOffsetAsync(logFile, index, cancellationToken);
+			(int maxIndex, long maxOffset) = await _logFileService.GetLineOffsetAsync(logFile, index + Math.Min(count, Int32.MaxValue - index), cancellationToken);
 			index = minIndex;
 			count = maxIndex - minIndex;
 
 			byte[] result;
-			using (System.IO.Stream stream = await _logFileService.OpenRawStreamAsync(logFile, minOffset, maxOffset - minOffset))
+			using (System.IO.Stream stream = await _logFileService.OpenRawStreamAsync(logFile, minOffset, maxOffset - minOffset, cancellationToken))
 			{
 				result = new byte[stream.Length];
 				await stream.ReadFixedSizeDataAsync(result, 0, result.Length);
@@ -195,7 +206,7 @@ namespace Horde.Build.Logs
 							}
 							else if (result[offset] == (byte)'\n')
 							{
-								await stream.WriteAsync(result.AsMemory(startOffset, offset - startOffset));
+								await stream.WriteAsync(result.AsMemory(startOffset, offset - startOffset), cancellationToken);
 								offset++;
 								break;
 							}
@@ -246,7 +257,7 @@ namespace Horde.Build.Logs
 				Response.ContentType = "application/json";
 				Response.Headers.ContentLength = stream.Length;
 				stream.Position = 0;
-				await stream.CopyToAsync(Response.Body);
+				await stream.CopyToAsync(Response.Body, cancellationToken);
 			}
 			return new EmptyResult();
 		}
@@ -258,12 +269,18 @@ namespace Horde.Build.Logs
 		/// <param name="text">Text to search for</param>
 		/// <param name="firstLine">First line to search from</param>
 		/// <param name="count">Number of results to return</param>
+		/// <param name="cancellationToken">Cancellation token for the request</param>
 		/// <returns>Raw log data for the requested range</returns>
 		[HttpGet]
 		[Route("/api/v1/logs/{logFileId}/search")]
-		public async Task<ActionResult<SearchLogFileResponse>> SearchLogFileAsync(LogId logFileId, [FromQuery] string text, [FromQuery] int firstLine = 0, [FromQuery] int count = 5)
+		public async Task<ActionResult<SearchLogFileResponse>> SearchLogFileAsync(
+			LogId logFileId,
+			[FromQuery] string text,
+			[FromQuery] int firstLine = 0,
+			[FromQuery] int count = 5,
+			CancellationToken cancellationToken = default)
 		{
-			ILogFile? logFile = await _logFileService.GetLogFileAsync(logFileId);
+			ILogFile? logFile = await _logFileService.GetLogFileAsync(logFileId, cancellationToken);
 			if (logFile == null)
 			{
 				return NotFound();
@@ -275,7 +292,7 @@ namespace Horde.Build.Logs
 
 			SearchLogFileResponse response = new SearchLogFileResponse();
 			response.Stats = new LogSearchStats();
-			response.Lines = await _logFileService.SearchLogDataAsync(logFile, text, firstLine, count, response.Stats);
+			response.Lines = await _logFileService.SearchLogDataAsync(logFile, text, firstLine, count, response.Stats, cancellationToken);
 			return response;
 		}
 
@@ -285,13 +302,14 @@ namespace Horde.Build.Logs
 		/// <param name="logFileId">Id of the log file to get information about</param>
 		/// <param name="index">Index of the first line to retrieve</param>
 		/// <param name="count">Number of lines to retrieve</param>
+		/// <param name="cancellationToken">Cancellation token for the request</param>
 		/// <returns>Information about the requested project</returns>
 		[HttpGet]
 		[Route("/api/v1/logs/{logFileId}/events")]
 		[ProducesResponseType(typeof(List<GetLogEventResponse>), 200)]
-		public async Task<ActionResult<List<GetLogEventResponse>>> GetEventsAsync(LogId logFileId, [FromQuery] int? index = null, [FromQuery] int? count = null)
+		public async Task<ActionResult<List<GetLogEventResponse>>> GetEventsAsync(LogId logFileId, [FromQuery] int? index = null, [FromQuery] int? count = null, CancellationToken cancellationToken = default)
 		{
-			ILogFile? logFile = await _logFileService.GetLogFileAsync(logFileId);
+			ILogFile? logFile = await _logFileService.GetLogFileAsync(logFileId, cancellationToken);
 			if (logFile == null)
 			{
 				return NotFound();
@@ -301,14 +319,14 @@ namespace Horde.Build.Logs
 				return Forbid();
 			}
 
-			List<ILogEvent> logEvents = await _logFileService.FindEventsAsync(logFile, null, index, count);
+			List<ILogEvent> logEvents = await _logFileService.FindEventsAsync(logFile, null, index, count, cancellationToken);
 
 			Dictionary<ObjectId, int?> spanIdToIssueId = new Dictionary<ObjectId, int?>();
 
 			List<GetLogEventResponse> responses = new List<GetLogEventResponse>();
 			foreach (ILogEvent logEvent in logEvents)
 			{
-				ILogEventData logEventData = await _logFileService.GetEventDataAsync(logFile, logEvent.LineIndex, logEvent.LineCount);
+				ILogEventData logEventData = await _logFileService.GetEventDataAsync(logFile, logEvent.LineIndex, logEvent.LineCount, cancellationToken);
 
 				int? issueId = null;
 				if (logEvent.SpanId != null && !spanIdToIssueId.TryGetValue(logEvent.SpanId.Value, out issueId))
@@ -329,12 +347,13 @@ namespace Horde.Build.Logs
 		/// <param name="logFileId">The logfile id</param>
 		/// <param name="offset">Offset within the log file</param>
 		/// <param name="lineIndex">The line index</param>
+		/// <param name="cancellationToken">Cancellation token for the request</param>
 		/// <returns>Http result code</returns>
 		[HttpPost]
 		[Route("/api/v1/logs/{logFileId}")]
-		public async Task<ActionResult> WriteData(LogId logFileId, [FromQuery] long offset, [FromQuery] int lineIndex)
+		public async Task<ActionResult> WriteData(LogId logFileId, [FromQuery] long offset, [FromQuery] int lineIndex, CancellationToken cancellationToken)
 		{
-			ILogFile? logFile = await _logFileService.GetLogFileAsync(logFileId);
+			ILogFile? logFile = await _logFileService.GetLogFileAsync(logFileId, cancellationToken);
 			if (logFile == null)
 			{
 				return NotFound();
@@ -346,8 +365,8 @@ namespace Horde.Build.Logs
 
 			using (MemoryStream bodyStream = new MemoryStream())
 			{
-				await Request.Body.CopyToAsync(bodyStream);
-				await _logFileService.WriteLogDataAsync(logFile, offset, lineIndex, bodyStream.ToArray(), false);
+				await Request.Body.CopyToAsync(bodyStream, cancellationToken);
+				await _logFileService.WriteLogDataAsync(logFile, offset, lineIndex, bodyStream.ToArray(), false, cancellationToken: cancellationToken);
 			}
 			return Ok();
 		}
diff --git a/Engine/Source/Programs/Horde/Horde.Build/Notifications/NotificationService.cs b/Engine/Source/Programs/Horde/Horde.Build/Notifications/NotificationService.cs
--- a/Engine/Source/Programs/Horde/Horde.Build/Notifications/NotificationService.cs
+++ b/Engine/Source/Programs/Horde/Horde.Build/Notifications/NotificationService.cs
@@ -650,7 +650,7 @@ namespace Horde.Build.Notifications
 				return;
 			}
 
-			ILogFile? logFile = await _logFileService.GetLogFileAsync(step.LogId.Value);
+			ILogFile? logFile = await _logFileService.GetLogFileAsync(step.LogId.Value, CancellationToken.None);
 			if(logFile == null)
 			{
 				_logger.LogError("Step does not have a log file");
diff --git a/Engine/Source/Programs/Horde/Horde.Build/Notifications/Sinks/SlackNotificationSink.cs b/Engine/Source/Programs/Horde/Horde.Build/Notifications/Sinks/SlackNotificationSink.cs
--- a/Engine/Source/Programs/Horde/Horde.Build/Notifications/Sinks/SlackNotificationSink.cs
+++ b/Engine/Source/Programs/Horde/Horde.Build/Notifications/Sinks/SlackNotificationSink.cs
@@ -919,7 +919,7 @@ namespace Horde.Build.Notifications.Sinks
 				if (span.FirstFailure.LogId != null)
 				{
 					LogId logId = span.FirstFailure.LogId.Value;
-					ILogFile? logFile = await _logFileService.GetLogFileAsync(logId);
+					ILogFile? logFile = await _logFileService.GetLogFileAsync(logId, CancellationToken.None);
 					if (logFile != null)
 					{
 						events = await _logFileService.FindEventsAsync(logFile, span.Id, 0, 50);
@@ -1249,7 +1249,7 @@ namespace Horde.Build.Notifications.Sinks
 			if (lastSpan != null && lastSpan.LastFailure.LogId != null)
 			{
 				LogId logId = lastSpan.LastFailure.LogId.Value;
-				ILogFile? logFile = await _logFileService.GetLogFileAsync(logId);
+				ILogFile? logFile = await _logFileService.GetLogFileAsync(logId, CancellationToken.None);
 				if(logFile != null)
 				{
 					List<ILogEvent> events = await _logFileService.FindEventsAsync(logFile, lastSpan.Id, 0, 20);
diff --git a/Engine/Source/Programs/Horde/Horde.Build/Server/RpcService.cs b/Engine/Source/Programs/Horde/Horde.Build/Server/RpcService.cs
--- a/Engine/Source/Programs/Horde/Horde.Build/Server/RpcService.cs
+++ b/Engine/Source/Programs/Horde/Horde.Build/Server/RpcService.cs
@@ -917,7 +917,7 @@ namespace Horde.Build.Server
 				newEvent.LineCount = createEvent.LineCount;
 				newEvents.Add(newEvent);
 			}
-			await _logFileService.CreateEventsAsync(newEvents);
+			await _logFileService.CreateEventsAsync(newEvents, context.CancellationToken);
 			return new Empty();
 		}
 
@@ -929,7 +929,7 @@ namespace Horde.Build.Server
 		/// <returns>Information about the new agent</returns>
 		public override async Task<Empty> WriteOutput(WriteOutputRequest request, ServerCallContext context)
 		{
-			ILogFile? logFile = await _logFileService.GetCachedLogFileAsync(new LogId(request.LogId));
+			ILogFile? logFile = await _logFileService.GetCachedLogFileAsync(new LogId(request.LogId), context.CancellationToken);
 			if (logFile == null)
 			{
 				throw new StructuredRpcException(StatusCode.NotFound, "Resource not found");
@@ -939,7 +939,7 @@ namespace Horde.Build.Server
 				throw new StructuredRpcException(StatusCode.PermissionDenied, "Access denied");
 			}
 
-			await _logFileService.WriteLogDataAsync(logFile, request.Offset, request.LineIndex, request.Data.ToArray(), request.Flush);
+			await _logFileService.WriteLogDataAsync(logFile, request.Offset, request.LineIndex, request.Data.ToArray(), request.Flush, cancellationToken: context.CancellationToken);
 			return new Empty();
 		}
 
diff --git a/Engine/Source/Programs/Horde/Horde.Build/Tasks/ConformTaskSource.cs b/Engine/Source/Programs/Horde/Horde.Build/Tasks/ConformTaskSource.cs
--- a/Engine/Source/Programs/Horde/Horde.Build/Tasks/ConformTaskSource.cs
+++ b/Engine/Source/Programs/Horde/Horde.Build/Tasks/ConformTaskSource.cs
@@ -170,7 +170,7 @@ namespace Horde.Build.Tasks
 					LeaseId leaseId = LeaseId.GenerateNewId();
 					if (await AllocateConformLeaseAsync(agent.Id, task.Workspaces, leaseId))
 					{
-						ILogFile log = await _logService.CreateLogFileAsync(JobId.Empty, agent.SessionId, LogType.Json);
+						ILogFile log = await _logService.CreateLogFileAsync(JobId.Empty, agent.SessionId, LogType.Json, cancellationToken: cancellationToken);
 						task.LogId = log.Id.ToString();
 						task.RemoveUntrackedFiles = agent.RequestFullConform;
 
diff --git a/Engine/Source/Programs/Horde/Horde.Build/Tasks/RestartTaskSource.cs b/Engine/Source/Programs/Horde/Horde.Build/Tasks/RestartTaskSource.cs
--- a/Engine/Source/Programs/Horde/Horde.Build/Tasks/RestartTaskSource.cs
+++ b/Engine/Source/Programs/Horde/Horde.Build/Tasks/RestartTaskSource.cs
@@ -44,7 +44,7 @@ namespace Horde.Build.Tasks
 				return await DrainAsync(cancellationToken);
 			}
 
-			ILogFile log = await _logService.CreateLogFileAsync(JobId.Empty, agent.SessionId, LogType.Json);
+			ILogFile log = await _logService.CreateLogFileAsync(JobId.Empty, agent.SessionId, LogType.Json, cancellationToken: cancellationToken);
 
 			RestartTask task = new RestartTask();
 			task.LogId = log.Id.ToString();
diff --git a/Engine/Source/Programs/Horde/Horde.Build/Tasks/ShutdownTaskSource.cs b/Engine/Source/Programs/Horde/Horde.Build/Tasks/ShutdownTaskSource.cs
--- a/Engine/Source/Programs/Horde/Horde.Build/Tasks/ShutdownTaskSource.cs
+++ b/Engine/Source/Programs/Horde/Horde.Build/Tasks/ShutdownTaskSource.cs
@@ -43,7 +43,7 @@ namespace Horde.Build.Tasks
 				return await DrainAsync(cancellationToken);
 			}
 
-			ILogFile log = await _logService.CreateLogFileAsync(JobId.Empty, agent.SessionId, LogType.Json);
+			ILogFile log = await _logService.CreateLogFileAsync(JobId.Empty, agent.SessionId, LogType.Json, cancellationToken: cancellationToken);
 
 			ShutdownTask task = new ShutdownTask();
 			task.LogId = log.Id.ToString();
diff --git a/Engine/Source/Programs/Horde/Horde.Build/Tasks/UpgradeTaskSource.cs b/Engine/Source/Programs/Horde/Horde.Build/Tasks/UpgradeTaskSource.cs
--- a/Engine/Source/Programs/Horde/Horde.Build/Tasks/UpgradeTaskSource.cs
+++ b/Engine/Source/Programs/Horde/Horde.Build/Tasks/UpgradeTaskSource.cs
@@ -55,7 +55,7 @@ namespace Horde.Build.Tasks
 				return await DrainAsync(cancellationToken);
 			}
 
-			ILogFile logFile = await _logService.CreateLogFileAsync(JobId.Empty, agent.SessionId, LogType.Json);
+			ILogFile logFile = await _logService.CreateLogFileAsync(JobId.Empty, agent.SessionId, LogType.Json, cancellationToken: cancellationToken);
 
 			UpgradeTask task = new UpgradeTask();
 			task.SoftwareId = requiredVersion.ToString();
diff --git a/Engine/Source/Programs/Horde/Horde.Build/Ugs/UgsController.cs b/Engine/Source/Programs/Horde/Horde.Build/Ugs/UgsController.cs
--- a/Engine/Source/Programs/Horde/Horde.Build/Ugs/UgsController.cs
+++ b/Engine/Source/Programs/Horde/Horde.Build/Ugs/UgsController.cs
@@ -3,6 +3,7 @@
 using System;
 using System.Collections.Generic;
 using System.Linq;
+using System.Threading;
 using System.Threading.Tasks;
 using Horde.Build.Users;
 using Horde.Build.Logs;
@@ -247,30 +248,31 @@ namespace Horde.Build.Ugs
 		/// Retrieve information about builds for a specific issue
 		/// </summary>
 		/// <param name="issueId">Id of the issue to get information about</param>
+		/// <param name="cancellationToken">Cancellation token for the request</param>
 		/// <returns>List of matching agents</returns>
 		[HttpGet]
 		[Route("/ugs/api/issues/{issueId}/diagnostics")]
 		[ProducesResponseType(typeof(List<GetUgsIssueDiagnosticResponse>), 200)]
-		public async Task<ActionResult<List<GetUgsIssueDiagnosticResponse>>> GetIssueDiagnosticsAsync(int issueId)
+		public async Task<ActionResult<List<GetUgsIssueDiagnosticResponse>>> GetIssueDiagnosticsAsync(int issueId, CancellationToken cancellationToken)
 		{
 			List<GetUgsIssueDiagnosticResponse> diagnostics = new List<GetUgsIssueDiagnosticResponse>();
 
 			Dictionary<LogId, ILogFile?> logFiles = new Dictionary<LogId, ILogFile?>();
 
 			List<IIssueSpan> spans = await _issueService.Collection.FindSpansAsync(issueId);
-			List<ILogEvent> events = await _logFileService.FindEventsForSpansAsync(spans.Select(x => x.Id), null, 0, count: 10);
+			List<ILogEvent> events = await _logFileService.FindEventsForSpansAsync(spans.Select(x => x.Id), null, 0, count: 10, cancellationToken);
 
 			foreach (ILogEvent logEvent in events)
 			{
 				ILogFile? logFile;
 				if(!logFiles.TryGetValue(logEvent.LogId, out logFile))
 				{
-					logFile = await _logFileService.GetLogFileAsync(logEvent.LogId);
+					logFile = await _logFileService.GetLogFileAsync(logEvent.LogId, cancellationToken);
 					logFiles.Add(logEvent.LogId, logFile);
 				}
 				if (logFile != null)
 				{
-					ILogEventData eventData = await _logFileService.GetEventDataAsync(logFile, logEvent.LineIndex, logEvent.LineCount);
+					ILogEventData eventData = await _logFileService.GetEventDataAsync(logFile, logEvent.LineIndex, logEvent.LineCount, cancellationToken);
 					long buildId = logEvent.LogId.GetHashCode();
 					Uri url = new Uri(_settings.DashboardUrl, $"log/{logEvent.LogId}?lineindex={logEvent.LineIndex}");
 
