Horde: Add additional exception details to log events.

diff --git a/Engine/Source/Programs/Horde/Horde.Build/Horde.Build.csproj b/Engine/Source/Programs/Horde/Horde.Build/Horde.Build.csproj
--- a/Engine/Source/Programs/Horde/Horde.Build/Horde.Build.csproj
+++ b/Engine/Source/Programs/Horde/Horde.Build/Horde.Build.csproj
@@ -49,6 +49,7 @@
     <PackageReference Include="Polly.Extensions.Http" Version="3.0.0" />
     <PackageReference Include="Serilog.AspNetCore" Version="6.0.1" />
     <PackageReference Include="Serilog.Enrichers.OpenTracing" Version="0.0.2" />
+    <PackageReference Include="Serilog.Exceptions.Grpc" Version="8.4.0" />
     <PackageReference Include="Serilog.Extensions.Hosting" Version="5.0.1" />
     <PackageReference Include="Serilog.Sinks.Console" Version="4.1.0" />
     <PackageReference Include="Serilog.Sinks.File" Version="5.0.0" />
diff --git a/Engine/Source/Programs/Horde/Horde.Build/Program.cs b/Engine/Source/Programs/Horde/Horde.Build/Program.cs
--- a/Engine/Source/Programs/Horde/Horde.Build/Program.cs
+++ b/Engine/Source/Programs/Horde/Horde.Build/Program.cs
@@ -19,6 +19,9 @@ using Serilog;
 using Serilog.Configuration;
 using Serilog.Core;
 using Serilog.Events;
+using Serilog.Exceptions;
+using Serilog.Exceptions.Core;
+using Serilog.Exceptions.Grpc.Destructurers;
 using Serilog.Formatting.Json;
 using Serilog.Sinks.SystemConsole.Themes;
 
@@ -159,6 +162,9 @@ namespace Horde.Build
 			Serilog.Log.Logger = new LoggerConfiguration()
 				.WithHordeConfig(hordeSettings)
 				.Enrich.FromLogContext()
+				.Enrich.WithExceptionDetails(new DestructuringOptionsBuilder()
+					.WithDefaultDestructurers()
+					.WithDestructurers(new[] { new RpcExceptionDestructurer() }))
 				.WriteTo.Console(hordeSettings)
 				.WriteTo.File(Path.Combine(logDir.FullName, "Log.txt"), outputTemplate: "[{Timestamp:HH:mm:ss} {Level:w3}] {Indent}{Message:l}{NewLine}{Exception} [{SourceContext}]", rollingInterval: RollingInterval.Day, rollOnFileSizeLimit: true, fileSizeLimitBytes: 20 * 1024 * 1024, retainedFileCountLimit: 10)
 				.WriteTo.File(new JsonFormatter(renderMessage: true), Path.Combine(logDir.FullName, "Log.json"), rollingInterval: RollingInterval.Day, rollOnFileSizeLimit: true, fileSizeLimitBytes: 20 * 1024 * 1024, retainedFileCountLimit: 10)
