Horde: [TEMP] Introduce periodic errors in WriteOutput RPC calls, to test reporting/recovery.

diff --git a/Engine/Source/Programs/Horde/Horde.Build/Server/RpcService.cs b/Engine/Source/Programs/Horde/Horde.Build/Server/RpcService.cs
--- a/Engine/Source/Programs/Horde/Horde.Build/Server/RpcService.cs
+++ b/Engine/Source/Programs/Horde/Horde.Build/Server/RpcService.cs
@@ -921,6 +921,8 @@ namespace Horde.Build.Server
 			return new Empty();
 		}
 
+		static int s_counter = 0;
+
 		/// <summary>
 		/// Writes output to a log file
 		/// </summary>
@@ -939,6 +941,13 @@ namespace Horde.Build.Server
 				throw new StructuredRpcException(StatusCode.PermissionDenied, "Access denied");
 			}
 
+			int counter = Interlocked.Increment(ref s_counter);
+			if ((counter & 63) == 0)
+			{
+				_logger.LogError("Triggering error inside WriteOutput()");
+				throw new Exception("WriteOutput() test");
+			}
+
 			await _logFileService.WriteLogDataAsync(logFile, request.Offset, request.LineIndex, request.Data.ToArray(), request.Flush, cancellationToken: context.CancellationToken);
 			return new Empty();
 		}
