Horde: Remove gRPC client retrying based on Polly

diff --git a/Engine/Source/Programs/Horde/Horde.Agent/Services/GrpcService.cs b/Engine/Source/Programs/Horde/Horde.Agent/Services/GrpcService.cs
--- a/Engine/Source/Programs/Horde/Horde.Agent/Services/GrpcService.cs
+++ b/Engine/Source/Programs/Horde/Horde.Agent/Services/GrpcService.cs
@@ -63,18 +63,14 @@ namespace Horde.Agent.Services
 		/// <returns>New grpc channel</returns>
 		public GrpcChannel CreateGrpcChannel(Uri address, AuthenticationHeaderValue? authHeaderValue)
 		{
+			#pragma warning disable CA2000 // Dispose objects before losing scope
+			// HTTP client handler is disposed by GrpcChannel below
 			HttpClientHandler customCertHandler = new HttpClientHandler();
+			#pragma warning restore CA2000 // Dispose objects before losing scope
 			customCertHandler.ServerCertificateCustomValidationCallback += (sender, cert, chain, errors) => CertificateHelper.CertificateValidationCallBack(_logger, sender, cert, chain, errors, _serverProfile);
+			customCertHandler.CheckCertificateRevocationList = true;
 
-			TimeSpan[] retryDelay = { TimeSpan.FromSeconds(1), TimeSpan.FromSeconds(5), TimeSpan.FromSeconds(10) };
-			IAsyncPolicy<HttpResponseMessage> policy = HttpPolicyExtensions.HandleTransientHttpError()
-				.WaitAndRetryAsync(retryDelay, onRetry: (result, timeSpan) => _logger.LogInformation("Retrying http request {Time} (status: {Code})", timeSpan, result.Result?.StatusCode));
-#pragma warning disable CA2000 // Dispose objects before losing scope
-			PolicyHttpMessageHandler retryHandler = new PolicyHttpMessageHandler(policy);
-#pragma warning restore CA2000 // Dispose objects before losing scope
-			retryHandler.InnerHandler = customCertHandler;
-
-			HttpClient httpClient = new HttpClient(retryHandler, true);
+			HttpClient httpClient = new (customCertHandler, true);
 			httpClient.DefaultRequestHeaders.Add("Accept", "application/json");
 			if (authHeaderValue != null)
 			{
