Horde: Add support for automatically expiring refs after a period of time.

diff --git a/Engine/Source/Programs/Horde/Horde.Agent/Execution/JobExecutor.cs b/Engine/Source/Programs/Horde/Horde.Agent/Execution/JobExecutor.cs
--- a/Engine/Source/Programs/Horde/Horde.Agent/Execution/JobExecutor.cs
+++ b/Engine/Source/Programs/Horde/Horde.Agent/Execution/JobExecutor.cs
@@ -861,7 +861,7 @@ namespace Horde.Agent.Execution
 				}
 
 				// Write the final node
-				await treeWriter.WriteRefAsync(refName, outputNode, cancellationToken);
+				await treeWriter.WriteRefAsync(refName, outputNode, new RefOptions { Lifetime = TimeSpan.FromDays(3.0) }, cancellationToken);
 			}
 
 			return true;
diff --git a/Engine/Source/Programs/Horde/Horde.Build.Tests/BlobStoreTests.cs b/Engine/Source/Programs/Horde/Horde.Build.Tests/BlobStoreTests.cs
--- a/Engine/Source/Programs/Horde/Horde.Build.Tests/BlobStoreTests.cs
+++ b/Engine/Source/Programs/Horde/Horde.Build.Tests/BlobStoreTests.cs
@@ -9,33 +9,27 @@ using Microsoft.VisualStudio.TestTools.UnitTesting;
 using MongoDB.Driver;
 using EpicGames.Horde.Storage;
 using System.Buffers;
-using EpicGames.Horde.Storage.Backends;
-using Microsoft.Extensions.Caching.Memory;
-using Microsoft.Extensions.Logging.Abstractions;
+using System.Threading;
+using Horde.Build.Server;
+using Horde.Build.Storage;
+using Horde.Build.Utilities;
 
 namespace Horde.Build.Tests
 {
+	using BackendId = StringId<IStorageBackend>;
+
 	[TestClass]
 	public class BlobStoreTests : TestSetup
 	{
-		IMemoryCache? _cache;
-		MemoryStorageClient? _store;
-
-		IStorageClient CreateBlobStore()
-		{
-			_cache ??= new MemoryCache(new MemoryCacheOptions());
-			_store ??= new MemoryStorageClient(_cache, NullLogger.Instance);
-			return _store;
-		}
-
-		protected override void Dispose(bool disposing)
+		async Task<IStorageClient> CreateStorageClientAsync()
 		{
-			base.Dispose(disposing);
+			GlobalConfig globalConfig = new GlobalConfig();
+			globalConfig.Storage.Backends.Add(new BackendConfig { Id = new BackendId("default-backend"), Type = StorageBackendType.Memory });
+			globalConfig.Storage.Namespaces.Add(new NamespaceConfig { Id = new NamespaceId("default"), Backend = new BackendId("default-backend"), GcDelayHrs = 0.0 });
+			await ConfigCollection.AddConfigAsync("globals", globalConfig);
+			Assert.IsNotNull(await GlobalsService.TryUpdateAsync(await GlobalsService.GetAsync(), "globals"));
 
-			if (disposing)
-			{
-				_cache?.Dispose();
-			}
+			return await StorageService.GetClientAsync(new NamespaceId("default"), CancellationToken.None);
 		}
 
 		static byte[] CreateTestData(int length, int seed)
@@ -94,7 +88,7 @@ namespace Horde.Build.Tests
 		[TestMethod]
 		public async Task LeafTest()
 		{
-			IStorageClient store = CreateBlobStore();
+			IStorageClient store = await CreateStorageClientAsync();
 
 			byte[] input = CreateTestData(256, 0);
 			BlobLocator locator = await store.WriteBundleAsync(CreateTestBundle(input, Array.Empty<BlobLocator>()));
@@ -106,7 +100,7 @@ namespace Horde.Build.Tests
 		[TestMethod]
 		public async Task ReferenceTest()
 		{
-			IStorageClient store = CreateBlobStore();
+			IStorageClient store = await CreateStorageClientAsync();
 
 			byte[] input1 = CreateTestData(256, 1);
 			BlobLocator locator1 = await store.WriteBundleAsync(CreateTestBundle(input1, Array.Empty<BlobLocator>()));
@@ -141,5 +135,38 @@ namespace Horde.Build.Tests
 			Assert.IsTrue(refTarget2.Data.Span.SequenceEqual(input3));
 			Assert.IsTrue(refTarget2.References.SequenceEqual(new BlobLocator[] { locator1, locator2 }));
 		}
+
+		[TestMethod]
+		public async Task RefExpiryTest()
+		{
+			IStorageClient store = await CreateStorageClientAsync();
+			NodeLocator target = new NodeLocator(new BlobLocator(HostId.Empty, new BlobId("test")), 123);
+
+			await store.WriteRefTargetAsync("test-ref-1", target);
+			await store.WriteRefTargetAsync("test-ref-2", target, new RefOptions { Lifetime = TimeSpan.FromMinutes(30.0), Extend = true });
+			await store.WriteRefTargetAsync("test-ref-3", target, new RefOptions { Lifetime = TimeSpan.FromMinutes(30.0), Extend = false });
+
+			Assert.AreEqual(target, await store.TryReadRefTargetAsync("test-ref-1"));
+			Assert.AreEqual(target, await store.TryReadRefTargetAsync("test-ref-2"));
+			Assert.AreEqual(target, await store.TryReadRefTargetAsync("test-ref-3"));
+
+			await Clock.AdvanceAsync(TimeSpan.FromMinutes(25.0));
+
+			Assert.AreEqual(target, await store.TryReadRefTargetAsync("test-ref-1"));
+			Assert.AreEqual(target, await store.TryReadRefTargetAsync("test-ref-2"));
+			Assert.AreEqual(target, await store.TryReadRefTargetAsync("test-ref-3"));
+
+			await Clock.AdvanceAsync(TimeSpan.FromMinutes(25.0));
+
+			Assert.AreEqual(target, await store.TryReadRefTargetAsync("test-ref-1"));
+			Assert.AreEqual(target, await store.TryReadRefTargetAsync("test-ref-2"));
+			Assert.AreEqual(default, await store.TryReadRefTargetAsync("test-ref-3"));
+
+			await Clock.AdvanceAsync(TimeSpan.FromMinutes(35.0));
+
+			Assert.AreEqual(target, await store.TryReadRefTargetAsync("test-ref-1"));
+			Assert.AreEqual(default, await store.TryReadRefTargetAsync("test-ref-2"));
+			Assert.AreEqual(default, await store.TryReadRefTargetAsync("test-ref-3"));
+		}
 	}
 }
diff --git a/Engine/Source/Programs/Horde/Horde.Build/Commands/Bundles/CreateCommand.cs b/Engine/Source/Programs/Horde/Horde.Build/Commands/Bundles/CreateCommand.cs
--- a/Engine/Source/Programs/Horde/Horde.Build/Commands/Bundles/CreateCommand.cs
+++ b/Engine/Source/Programs/Horde/Horde.Build/Commands/Bundles/CreateCommand.cs
@@ -40,7 +40,7 @@ namespace Horde.Build.Commands.Bundles
 			DirectoryNode node = new DirectoryNode(DirectoryFlags.None);
 			await node.CopyFromDirectoryAsync(InputDir.ToDirectoryInfo(), new ChunkingOptions(), writer, CancellationToken.None);
 
-			await writer.WriteRefAsync(RefName, node, CancellationToken.None);
+			await writer.WriteRefAsync(RefName, node);
 			return 0;
 		}
 	}
diff --git a/Engine/Source/Programs/Horde/Horde.Build/Perforce/ReplicationService.cs b/Engine/Source/Programs/Horde/Horde.Build/Perforce/ReplicationService.cs
--- a/Engine/Source/Programs/Horde/Horde.Build/Perforce/ReplicationService.cs
+++ b/Engine/Source/Programs/Horde/Horde.Build/Perforce/ReplicationService.cs
@@ -672,7 +672,7 @@ namespace Horde.Build.Perforce
 					if (syncedSize > 0)
 					{
 						Stopwatch flushTimer = Stopwatch.StartNew();
-						await writer.WriteRefAsync(incRefName, syncNode, cancellationToken);
+						await writer.WriteRefAsync(incRefName, syncNode, cancellationToken: cancellationToken);
 						flushTimer.Stop();
 						deleteIncRef = true;
 					}
@@ -842,7 +842,7 @@ namespace Horde.Build.Perforce
 
 			// Return the new root object
 			_logger.LogInformation("Writing ref {RefId} for {StreamId} change {Change}", refName, stream.Id, change);
-			await writer.WriteRefAsync(refName, syncNode, cancellationToken);
+			await writer.WriteRefAsync(refName, syncNode, cancellationToken: cancellationToken);
 
 			// Delete the incremental state
 			if (deleteIncRef)
diff --git a/Engine/Source/Programs/Horde/Horde.Build/Storage/IRefCollection.cs b/Engine/Source/Programs/Horde/Horde.Build/Storage/IRefCollection.cs
--- a/Engine/Source/Programs/Horde/Horde.Build/Storage/IRefCollection.cs
+++ b/Engine/Source/Programs/Horde/Horde.Build/Storage/IRefCollection.cs
@@ -29,8 +29,9 @@ namespace Horde.Build.Storage
 		/// <param name="namespaceId">Namespace for the ref</param>
 		/// <param name="name">The ref name</param>
 		/// <param name="target">Target for the ref</param>
+		/// <param name="options">Options for the ref</param>
 		/// <param name="cancellationToken">Cancellation token for the operation</param>
-		Task WriteRefTargetAsync(NamespaceId namespaceId, RefName name, NodeLocator target, CancellationToken cancellationToken = default);
+		Task WriteRefTargetAsync(NamespaceId namespaceId, RefName name, NodeLocator target, RefOptions? options = null, CancellationToken cancellationToken = default);
 
 		/// <summary>
 		/// Delete a ref from storage
diff --git a/Engine/Source/Programs/Horde/Horde.Build/Storage/RefCollection.cs b/Engine/Source/Programs/Horde/Horde.Build/Storage/RefCollection.cs
--- a/Engine/Source/Programs/Horde/Horde.Build/Storage/RefCollection.cs
+++ b/Engine/Source/Programs/Horde/Horde.Build/Storage/RefCollection.cs
@@ -7,6 +7,7 @@ using System.Threading;
 using System.Threading.Tasks;
 using EpicGames.Horde.Storage;
 using Horde.Build.Server;
+using HordeCommon;
 using Microsoft.Extensions.Caching.Memory;
 using Microsoft.Extensions.Logging;
 using MongoDB.Bson;
@@ -55,6 +56,12 @@ namespace Horde.Build.Storage
 			[BsonElement("e")]
 			public int ExportIdx { get; set; }
 
+			[BsonElement("xa"), BsonIgnoreIfDefault]
+			public DateTime? ExpiresAtUtc { get; set; }
+
+			[BsonElement("xt"), BsonIgnoreIfDefault]
+			public TimeSpan? Lifetime { get; set; }
+
 			public RefDocument()
 			{
 				Name = RefName.Empty;
@@ -68,23 +75,26 @@ namespace Horde.Build.Storage
 				Blob = target.Blob;
 				ExportIdx = target.ExportIdx;
 			}
+
+			public bool HasExpired(DateTime utcNow) => ExpiresAtUtc.HasValue && utcNow >= ExpiresAtUtc.Value;
+
+			public bool RequiresTouch(DateTime utcNow) => ExpiresAtUtc.HasValue && Lifetime.HasValue && utcNow >= ExpiresAtUtc.Value - new TimeSpan(Lifetime.Value.Ticks / 4);
 		}
 
 		class CachedRefValue
 		{
-			public RefName Name { get; }
-			public NodeLocator Target { get; }
+			public RefDocument? Value { get; }
 			public DateTime Time { get; }
 
-			public CachedRefValue(RefName name, NodeLocator target)
+			public CachedRefValue(RefDocument? value)
 			{
-				Name = name;
-				Target = target;
+				Value = value;
 				Time = DateTime.UtcNow;
 			}
 		}
 
 		readonly IMongoCollection<RefDocument> _refs;
+		readonly IClock _clock;
 		readonly IMemoryCache _cache;
 		readonly ILogger _logger;
 
@@ -92,18 +102,20 @@ namespace Horde.Build.Storage
 		/// Constructor
 		/// </summary>
 		/// <param name="mongoService">The mongo service implementation</param>
+		/// <param name="clock">Clock for expiring refs</param>
 		/// <param name="cache">Cache for ref values</param>
 		/// <param name="logger">Logger for output</param>
-		public RefCollection(MongoService mongoService, IMemoryCache cache, ILogger<RefCollection> logger)
+		public RefCollection(MongoService mongoService, IClock clock, IMemoryCache cache, ILogger<RefCollection> logger)
 		{
 			_refs = mongoService.GetCollection<RefDocument>("Refs", keys => keys.Ascending(x => x.NamespaceId).Ascending(x => x.Name), unique: true);
+			_clock = clock;
 			_cache = cache;
 			_logger = logger;
 		}
 
-		CachedRefValue AddRefToCache(NamespaceId namespaceId, RefName name, NodeLocator target)
+		CachedRefValue AddRefToCache(NamespaceId namespaceId, RefName name, RefDocument? value)
 		{
-			CachedRefValue cacheValue = new CachedRefValue(name, target);
+			CachedRefValue cacheValue = new CachedRefValue(value);
 			using (ICacheEntry newEntry = _cache.CreateEntry(new RefKey(namespaceId, name)))
 			{
 				newEntry.Value = cacheValue;
@@ -120,6 +132,15 @@ namespace Horde.Build.Storage
 			AddRefToCache(namespaceId, name, default);
 		}
 
+		/// <summary>
+		/// Deletes a ref document that has reached its expiry time
+		/// </summary>
+		async Task DeleteExpiredRefAsync(RefDocument refDocument, CancellationToken cancellationToken = default)
+		{
+			await _refs.DeleteOneAsync(x => x.Id == refDocument.Id && x.ExpiresAtUtc == refDocument.ExpiresAtUtc, cancellationToken: cancellationToken);
+			AddRefToCache(refDocument.NamespaceId, refDocument.Name, default);
+		}
+
 		/// <inheritdoc/>
 		public async Task<NodeLocator> TryReadRefTargetAsync(NamespaceId namespaceId, RefName name, DateTime cacheTime = default, CancellationToken cancellationToken = default)
 		{
@@ -127,27 +148,45 @@ namespace Horde.Build.Storage
 			if (!_cache.TryGetValue(name, out entry) || entry.Time < cacheTime)
 			{
 				RefDocument? refDocument = await _refs.Find(x => x.NamespaceId == namespaceId && x.Name == name).FirstOrDefaultAsync(cancellationToken);
+				entry = AddRefToCache(namespaceId, name, refDocument);
+			}
+
+			if (entry.Value == null)
+			{
+				return default;
+			}
 
-				NodeLocator target;
-				if (refDocument == null)
+			if (entry.Value.ExpiresAtUtc != null)
+			{
+				DateTime utcNow = _clock.UtcNow;
+				if (entry.Value.HasExpired(utcNow))
 				{
-					target = default;
+					await DeleteExpiredRefAsync(entry.Value, cancellationToken);
+					return default;
 				}
-				else
+				if (entry.Value.RequiresTouch(utcNow))
 				{
-					target = new NodeLocator(refDocument.Blob, refDocument.ExportIdx);
+					await _refs.UpdateOneAsync(x => x.Id == entry.Value.Id, Builders<RefDocument>.Update.Set(x => x.ExpiresAtUtc, utcNow + entry.Value.Lifetime!.Value), cancellationToken: cancellationToken);
 				}
-
-				entry = AddRefToCache(namespaceId, name, target);
 			}
-			return entry.Target;
+
+			return new NodeLocator(entry.Value.Blob, entry.Value.ExportIdx);
 		}
 
 		/// <inheritdoc/>
-		public async Task WriteRefTargetAsync(NamespaceId namespaceId, RefName name, NodeLocator target, CancellationToken cancellationToken = default)
+		public async Task WriteRefTargetAsync(NamespaceId namespaceId, RefName name, NodeLocator target, RefOptions? options = null, CancellationToken cancellationToken = default)
 		{
 			RefDocument refDocument = new RefDocument(namespaceId, name, target);
-			
+
+			if (options != null && options.Lifetime.HasValue)
+			{
+				refDocument.ExpiresAtUtc = _clock.UtcNow + options.Lifetime.Value;
+				if (options.Extend ?? true)
+				{
+					refDocument.Lifetime = options.Lifetime;
+				}
+			}
+
 			ReplaceOneResult result = await _refs.ReplaceOneAsync(x => x.NamespaceId == namespaceId && x.Name == name, refDocument, new ReplaceOptions { IsUpsert = true }, cancellationToken);
 			if (result.ModifiedCount > 0)
 			{
@@ -166,19 +205,27 @@ namespace Horde.Build.Storage
 				_logger.LogError("Unable to update ref {NamespaceId}:{RefName} (modified: {Modified}, matched: {Matched}, upsert: {Upsert})", namespaceId, name, result.ModifiedCount, result.MatchedCount, result.UpsertedId);
 			}
 			
-			AddRefToCache(namespaceId, name, target);
+			AddRefToCache(namespaceId, name, refDocument);
 		}
 
 		/// <inheritdoc/>
 		public async IAsyncEnumerable<NodeLocator> EnumerateRefs(NamespaceId namespaceId, [EnumeratorCancellation] CancellationToken cancellationToken = default)
 		{
+			DateTime utcNow = _clock.UtcNow;
 			using (IAsyncCursor<RefDocument> cursor = await _refs.Find(x => x.NamespaceId == namespaceId).ToCursorAsync(cancellationToken))
 			{
 				while (await cursor.MoveNextAsync(cancellationToken))
 				{
 					foreach (RefDocument refDoc in cursor.Current)
 					{
-						yield return new NodeLocator(refDoc.Blob, refDoc.ExportIdx);
+						if (refDoc.HasExpired(utcNow))
+						{
+							await DeleteExpiredRefAsync(refDoc, cancellationToken);
+						}
+						else
+						{
+							yield return new NodeLocator(refDoc.Blob, refDoc.ExportIdx);
+						}
 					}
 				}
 			}
diff --git a/Engine/Source/Programs/Horde/Horde.Build/Storage/StorageController.cs b/Engine/Source/Programs/Horde/Horde.Build/Storage/StorageController.cs
--- a/Engine/Source/Programs/Horde/Horde.Build/Storage/StorageController.cs
+++ b/Engine/Source/Programs/Horde/Horde.Build/Storage/StorageController.cs
@@ -43,6 +43,11 @@ namespace Horde.Build.Storage
 		/// Export index for the ref
 		/// </summary>
 		public int ExportIdx { get; set; }
+
+		/// <summary>
+		/// Options for the ref
+		/// </summary>
+		public RefOptions? Options { get; set; }
 	}
 
 	/// <summary>
@@ -169,7 +174,7 @@ namespace Horde.Build.Storage
 
 			IStorageClient client = await _storageService.GetClientAsync(namespaceId, cancellationToken);
 			NodeLocator target = new NodeLocator(request.Locator, request.ExportIdx);
-			await client.WriteRefTargetAsync(refName, target, cancellationToken);
+			await client.WriteRefTargetAsync(refName, target, request.Options, cancellationToken);
 
 			return Ok();
 		}
diff --git a/Engine/Source/Programs/Horde/Horde.Build/Storage/StorageService.cs b/Engine/Source/Programs/Horde/Horde.Build/Storage/StorageService.cs
--- a/Engine/Source/Programs/Horde/Horde.Build/Storage/StorageService.cs
+++ b/Engine/Source/Programs/Horde/Horde.Build/Storage/StorageService.cs
@@ -144,7 +144,7 @@ namespace Horde.Build.Storage
 			public override Task<NodeLocator> TryReadRefTargetAsync(RefName name, DateTime cacheTime = default, CancellationToken cancellationToken = default) => Refs.TryReadRefTargetAsync(NamespaceId, name, cacheTime, cancellationToken);
 
 			/// <inheritdoc/>
-			public override Task WriteRefTargetAsync(RefName name, NodeLocator target, CancellationToken cancellationToken = default) => Refs.WriteRefTargetAsync(NamespaceId, name, target, cancellationToken);
+			public override Task WriteRefTargetAsync(RefName name, NodeLocator target, RefOptions? options = null, CancellationToken cancellationToken = default) => Refs.WriteRefTargetAsync(NamespaceId, name, target, options, cancellationToken);
 
 			/// <inheritdoc/>
 			public override Task DeleteRefAsync(RefName name, CancellationToken cancellationToken = default) => Refs.DeleteRefAsync(NamespaceId, name, cancellationToken);
diff --git a/Engine/Source/Programs/Shared/EpicGames.Horde.Tests/BundleStoreTests.cs b/Engine/Source/Programs/Shared/EpicGames.Horde.Tests/BundleStoreTests.cs
--- a/Engine/Source/Programs/Shared/EpicGames.Horde.Tests/BundleStoreTests.cs
+++ b/Engine/Source/Programs/Shared/EpicGames.Horde.Tests/BundleStoreTests.cs
@@ -85,7 +85,7 @@ namespace EpicGames.Horde.Tests
 
 				SimpleNode root = new SimpleNode(new ReadOnlySequence<byte>(new byte[] { 5 }), new[] { new TreeNodeRef<SimpleNode>(node4), new TreeNodeRef<SimpleNode>(node3) });
 
-				await writer.WriteRefAsync(new RefName("test"), root, CancellationToken.None);
+				await writer.WriteRefAsync(new RefName("test"), root);
 
 				await CheckTree(root);
 			}
diff --git a/Engine/Source/Programs/Shared/EpicGames.Horde/Storage/Backends/FileStorageClient.cs b/Engine/Source/Programs/Shared/EpicGames.Horde/Storage/Backends/FileStorageClient.cs
--- a/Engine/Source/Programs/Shared/EpicGames.Horde/Storage/Backends/FileStorageClient.cs
+++ b/Engine/Source/Programs/Shared/EpicGames.Horde/Storage/Backends/FileStorageClient.cs
@@ -103,7 +103,7 @@ namespace EpicGames.Horde.Storage.Backends
 		}
 
 		/// <inheritdoc/>
-		public override async Task WriteRefTargetAsync(RefName name, NodeLocator target, CancellationToken cancellationToken = default)
+		public override async Task WriteRefTargetAsync(RefName name, NodeLocator target, RefOptions? options = null, CancellationToken cancellationToken = default)
 		{
 			FileReference file = GetRefFile(name);
 			DirectoryReference.CreateDirectory(file.Directory);
diff --git a/Engine/Source/Programs/Shared/EpicGames.Horde/Storage/Backends/HttpStorageClient.cs b/Engine/Source/Programs/Shared/EpicGames.Horde/Storage/Backends/HttpStorageClient.cs
--- a/Engine/Source/Programs/Shared/EpicGames.Horde/Storage/Backends/HttpStorageClient.cs
+++ b/Engine/Source/Programs/Shared/EpicGames.Horde/Storage/Backends/HttpStorageClient.cs
@@ -119,9 +119,9 @@ namespace EpicGames.Horde.Storage.Backends
 		}
 
 		/// <inheritdoc/>
-		public override async Task WriteRefTargetAsync(RefName name, NodeLocator target, CancellationToken cancellationToken = default)
+		public override async Task WriteRefTargetAsync(RefName name, NodeLocator target, RefOptions? options = null, CancellationToken cancellationToken = default)
 		{
-			using (HttpResponseMessage response = await _httpClient.PutAsync($"api/v1/storage/{_namespaceId}/refs/{name}", new { locator = target.Blob, exportIdx = target.ExportIdx }, cancellationToken))
+			using (HttpResponseMessage response = await _httpClient.PutAsync($"api/v1/storage/{_namespaceId}/refs/{name}", new { locator = target.Blob, exportIdx = target.ExportIdx, options }, cancellationToken))
 			{
 				response.EnsureSuccessStatusCode();
 			}
diff --git a/Engine/Source/Programs/Shared/EpicGames.Horde/Storage/Backends/MemoryStorageClient.cs b/Engine/Source/Programs/Shared/EpicGames.Horde/Storage/Backends/MemoryStorageClient.cs
--- a/Engine/Source/Programs/Shared/EpicGames.Horde/Storage/Backends/MemoryStorageClient.cs
+++ b/Engine/Source/Programs/Shared/EpicGames.Horde/Storage/Backends/MemoryStorageClient.cs
@@ -82,7 +82,7 @@ namespace EpicGames.Horde.Storage.Backends
 		}
 
 		/// <inheritdoc/>
-		public override Task WriteRefTargetAsync(RefName name, NodeLocator target, CancellationToken cancellationToken = default)
+		public override Task WriteRefTargetAsync(RefName name, NodeLocator target, RefOptions? options = null, CancellationToken cancellationToken = default)
 		{
 			_refs[name] = target;
 			return Task.CompletedTask;
diff --git a/Engine/Source/Programs/Shared/EpicGames.Horde/Storage/IStorageClient.cs b/Engine/Source/Programs/Shared/EpicGames.Horde/Storage/IStorageClient.cs
--- a/Engine/Source/Programs/Shared/EpicGames.Horde/Storage/IStorageClient.cs
+++ b/Engine/Source/Programs/Shared/EpicGames.Horde/Storage/IStorageClient.cs
@@ -59,6 +59,22 @@ namespace EpicGames.Horde.Storage
 		}
 	}
 
+	/// <summary>
+	/// Options for a new ref
+	/// </summary>
+	public class RefOptions
+	{
+		/// <summary>
+		/// Time until a ref is expired
+		/// </summary>
+		public TimeSpan? Lifetime { get; set; }
+
+		/// <summary>
+		/// Whether to extend the remaining lifetime of a ref whenever it is fetched. Defaults to true.
+		/// </summary>
+		public bool? Extend { get; set; }
+	}
+
 	/// <summary>
 	/// Locates a node in storage
 	/// </summary>
@@ -195,18 +211,20 @@ namespace EpicGames.Horde.Storage
 		/// <param name="bundle">The bundle to write</param>
 		/// <param name="exportIdx">Index of the export in the bundle to be the root of the tree</param>
 		/// <param name="prefix">Prefix for blob names.</param>
+		/// <param name="options">Options for the new ref</param>
 		/// <param name="cancellationToken">Cancellation token for the operation</param>
 		/// <returns>Unique identifier for the blob</returns>
-		Task<NodeLocator> WriteRefAsync(RefName name, Bundle bundle, int exportIdx, Utf8String prefix = default, CancellationToken cancellationToken = default);
+		Task<NodeLocator> WriteRefAsync(RefName name, Bundle bundle, int exportIdx, Utf8String prefix = default, RefOptions? options = null, CancellationToken cancellationToken = default);
 
 		/// <summary>
 		/// Writes a new ref to the store
 		/// </summary>
 		/// <param name="name">Ref to write</param>
 		/// <param name="target">The target for the ref</param>
+		/// <param name="options">Options for the new ref</param>
 		/// <param name="cancellationToken">Cancellation token for the operation</param>
 		/// <returns>Unique identifier for the blob</returns>
-		Task WriteRefTargetAsync(RefName name, NodeLocator target, CancellationToken cancellationToken = default);
+		Task WriteRefTargetAsync(RefName name, NodeLocator target, RefOptions? options = null, CancellationToken cancellationToken = default);
 
 		/// <summary>
 		/// Reads data for a ref from the store
@@ -281,10 +299,10 @@ namespace EpicGames.Horde.Storage
 			public Task<NodeLocator> TryReadRefTargetAsync(RefName name, DateTime cacheTime = default, CancellationToken cancellationToken = default) => _inner.TryReadRefTargetAsync(name, cacheTime, cancellationToken);
 
 			/// <inheritdoc/>
-			public Task<NodeLocator> WriteRefAsync(RefName name, Bundle bundle, int exportIdx = 0, Utf8String prefix = default, CancellationToken cancellationToken = default) => _inner.WriteRefAsync(name, bundle, exportIdx, prefix, cancellationToken);
+			public Task<NodeLocator> WriteRefAsync(RefName name, Bundle bundle, int exportIdx = 0, Utf8String prefix = default, RefOptions? options = null, CancellationToken cancellationToken = default) => _inner.WriteRefAsync(name, bundle, exportIdx, prefix, options, cancellationToken);
 
 			/// <inheritdoc/>
-			public Task WriteRefTargetAsync(RefName name, NodeLocator target, CancellationToken cancellationToken = default) => _inner.WriteRefTargetAsync(name, target, cancellationToken);
+			public Task WriteRefTargetAsync(RefName name, NodeLocator target, RefOptions? options = null, CancellationToken cancellationToken = default) => _inner.WriteRefTargetAsync(name, target, options, cancellationToken);
 
 			#endregion
 		}
@@ -350,12 +368,13 @@ namespace EpicGames.Horde.Storage
 		/// <param name="node">Node to be written</param>
 		/// <param name="options">Options for the node writer</param>
 		/// <param name="prefix">Prefix for uploaded blobs</param>
+		/// <param name="refOptions">Options for the ref</param>
 		/// <param name="cancellationToken">Cancellation token for the operation</param>
 		/// <returns>Location of node targetted by the ref</returns>
-		public static async Task<NodeLocator> WriteNodeAsync(this IStorageClient store, RefName name, TreeNode node, TreeOptions? options = null, Utf8String prefix = default, CancellationToken cancellationToken = default)
+		public static async Task<NodeLocator> WriteNodeAsync(this IStorageClient store, RefName name, TreeNode node, TreeOptions? options = null, Utf8String prefix = default, RefOptions? refOptions = null, CancellationToken cancellationToken = default)
 		{
 			TreeWriter writer = new TreeWriter(store, options, prefix.IsEmpty ? name.Text : prefix);
-			return await writer.WriteRefAsync(name, node, cancellationToken);
+			return await writer.WriteRefAsync(name, node, refOptions, cancellationToken);
 		}
 
 		#endregion
diff --git a/Engine/Source/Programs/Shared/EpicGames.Horde/Storage/RefName.cs b/Engine/Source/Programs/Shared/EpicGames.Horde/Storage/RefName.cs
--- a/Engine/Source/Programs/Shared/EpicGames.Horde/Storage/RefName.cs
+++ b/Engine/Source/Programs/Shared/EpicGames.Horde/Storage/RefName.cs
@@ -78,6 +78,11 @@ namespace EpicGames.Horde.Storage
 		/// <inheritdoc/>
 		public static bool operator >=(RefName lhs, RefName rhs) => lhs.CompareTo(rhs) >= 0;
 
+		/// <summary>
+		/// Construct a ref from a string
+		/// </summary>
+		/// <param name="name">Name of the ref</param>
+		public static implicit operator RefName(string name) => new RefName(name);
 	}
 
 	/// <summary>
diff --git a/Engine/Source/Programs/Shared/EpicGames.Horde/Storage/StorageClientBase.cs b/Engine/Source/Programs/Shared/EpicGames.Horde/Storage/StorageClientBase.cs
--- a/Engine/Source/Programs/Shared/EpicGames.Horde/Storage/StorageClientBase.cs
+++ b/Engine/Source/Programs/Shared/EpicGames.Horde/Storage/StorageClientBase.cs
@@ -625,16 +625,16 @@ namespace EpicGames.Horde.Storage
 		public abstract Task<NodeLocator> TryReadRefTargetAsync(RefName name, DateTime cacheTime = default, CancellationToken cancellationToken = default);
 
 		/// <inheritdoc/>
-		public virtual async Task<NodeLocator> WriteRefAsync(RefName name, Bundle bundle, int exportIdx = 0, Utf8String prefix = default, CancellationToken cancellationToken = default)
+		public virtual async Task<NodeLocator> WriteRefAsync(RefName name, Bundle bundle, int exportIdx = 0, Utf8String prefix = default, RefOptions? options = null, CancellationToken cancellationToken = default)
 		{
 			BlobLocator locator = await WriteBundleAsync(bundle, prefix, cancellationToken);
 			NodeLocator target = new NodeLocator(locator, exportIdx);
-			await WriteRefTargetAsync(name, target, cancellationToken);
+			await WriteRefTargetAsync(name, target, options, cancellationToken);
 			return target;
 		}
 
 		/// <inheritdoc/>
-		public abstract Task WriteRefTargetAsync(RefName name, NodeLocator target, CancellationToken cancellationToken = default);
+		public abstract Task WriteRefTargetAsync(RefName name, NodeLocator target, RefOptions? options = null, CancellationToken cancellationToken = default);
 
 		#endregion
 	}
diff --git a/Engine/Source/Programs/Shared/EpicGames.Horde/Storage/TreeWriter.cs b/Engine/Source/Programs/Shared/EpicGames.Horde/Storage/TreeWriter.cs
--- a/Engine/Source/Programs/Shared/EpicGames.Horde/Storage/TreeWriter.cs
+++ b/Engine/Source/Programs/Shared/EpicGames.Horde/Storage/TreeWriter.cs
@@ -246,16 +246,17 @@ namespace EpicGames.Horde.Storage
 		/// </summary>
 		/// <param name="name">Name of the ref to write</param>
 		/// <param name="node"></param>
+		/// <param name="options"></param>
 		/// <param name="cancellationToken"></param>
 		/// <returns></returns>
-		public async Task<NodeLocator> WriteRefAsync(RefName name, TreeNode node, CancellationToken cancellationToken = default)
+		public async Task<NodeLocator> WriteRefAsync(RefName name, TreeNode node, RefOptions? options = null, CancellationToken cancellationToken = default)
 		{
 			TreeNodeRef nodeRef = new TreeNodeRef(node);
 			await WriteAsync(nodeRef, cancellationToken);
 			await FlushAsync(cancellationToken);
 
 			Debug.Assert(nodeRef.Locator.IsValid());
-			await _store.WriteRefTargetAsync(name, nodeRef.Locator, cancellationToken);
+			await _store.WriteRefTargetAsync(name, nodeRef.Locator, options, cancellationToken);
 			return nodeRef.Locator;
 		}
 
@@ -264,12 +265,13 @@ namespace EpicGames.Horde.Storage
 		/// </summary>
 		/// <param name="name">Name of the ref to write</param>
 		/// <param name="root">Hash of the root node</param>
+		/// <param name="options"></param>
 		/// <param name="cancellationToken">Cancellation token for the operation</param>
-		public async Task<NodeLocator> WriteAsync(RefName name, IoHash root, CancellationToken cancellationToken)
+		public async Task<NodeLocator> WriteAsync(RefName name, IoHash root, RefOptions? options = null, CancellationToken cancellationToken = default)
 		{
 			await FlushAsync(cancellationToken);
 			NodeLocator locator = GetLocator(root);
-			await _store.WriteRefTargetAsync(name, locator, cancellationToken);
+			await _store.WriteRefTargetAsync(name, locator, options, cancellationToken);
 			return locator;
 		}
 
