Horde: Add support for presigned URL redirects when uploading bundle data.

diff --git a/Engine/Source/Programs/Horde/Horde.Build/Storage/Backends/AwsStorageBackend.cs b/Engine/Source/Programs/Horde/Horde.Build/Storage/Backends/AwsStorageBackend.cs
--- a/Engine/Source/Programs/Horde/Horde.Build/Storage/Backends/AwsStorageBackend.cs
+++ b/Engine/Source/Programs/Horde/Horde.Build/Storage/Backends/AwsStorageBackend.cs
@@ -2,6 +2,7 @@
 
 using System;
 using System.Collections.Generic;
+using System.Data.SqlTypes;
 using System.Globalization;
 using System.IO;
 using System.Runtime.CompilerServices;
@@ -12,6 +13,7 @@ using Amazon.Extensions.NETCore.Setup;
 using Amazon.Runtime;
 using Amazon.S3;
 using Amazon.S3.Model;
+using Google.Protobuf.WellKnownTypes;
 using Horde.Build.Utilities;
 using Microsoft.Extensions.Configuration;
 using Microsoft.Extensions.Logging;
@@ -101,7 +103,7 @@ namespace Horde.Build.Storage.Backends
 	/// <summary>
 	/// FileStorage implementation using an s3 bucket
 	/// </summary>
-	public sealed class AwsStorageBackend : IStorageBackend, IDisposable
+	public sealed class AwsStorageBackend : IStorageBackendWithRedirects, IDisposable
 	{
 		/// <summary>
 		/// S3 Client
@@ -297,6 +299,40 @@ namespace Horde.Build.Storage.Backends
 			}
 		}
 
+		/// <inheritdoc/>
+		public ValueTask<Uri?> GetReadRedirectAsync(string path, CancellationToken cancellationToken = default) => new ValueTask<Uri?>(GetPresignedUrl(path, HttpVerb.GET));
+
+		/// <inheritdoc/>
+		public ValueTask<Uri?> GetWriteRedirectAsync(string path, CancellationToken cancellationToken = default) => new ValueTask<Uri?>(GetPresignedUrl(path, HttpVerb.PUT));
+
+		/// <summary>
+		/// Helper method to generate a presigned URL for a request
+		/// </summary>
+		Uri? GetPresignedUrl(string path, HttpVerb verb)
+		{
+			using IScope scope = GlobalTracer.Instance.BuildSpan("AwsStorageBackend.TryGetHttpRedirectUrl").StartActive();
+			scope.Span.SetTag("Path", path);
+
+			string fullPath = GetFullPath(path);
+
+			try
+			{
+				GetPreSignedUrlRequest newGetRequest = new GetPreSignedUrlRequest();
+				newGetRequest.BucketName = _options.AwsBucketName;
+				newGetRequest.Key = fullPath;
+				newGetRequest.Verb = verb;
+				newGetRequest.Expires = DateTime.UtcNow.AddHours(3.0);
+
+				string url = _client.GetPreSignedURL(newGetRequest);
+				return new Uri(url);
+			}
+			catch (Exception ex)
+			{
+				_logger.LogWarning(ex, "Unable to get presigned url for {Path} from S3", fullPath);
+				return null;
+			}
+		}
+
 		/// <inheritdoc/>
 		public async Task WriteAsync(string path, Stream inputStream, CancellationToken cancellationToken)
 		{
diff --git a/Engine/Source/Programs/Horde/Horde.Build/Storage/IStorageBackend.cs b/Engine/Source/Programs/Horde/Horde.Build/Storage/IStorageBackend.cs
--- a/Engine/Source/Programs/Horde/Horde.Build/Storage/IStorageBackend.cs
+++ b/Engine/Source/Programs/Horde/Horde.Build/Storage/IStorageBackend.cs
@@ -65,6 +65,28 @@ namespace Horde.Build.Storage
 		IAsyncEnumerable<string> EnumerateAsync(CancellationToken cancellationToken = default);
 	}
 
+	/// <summary>
+	/// Storage backend that supports redirecting HTTP traffic (eg. using S3 presigned urls)
+	/// </summary>
+	public interface IStorageBackendWithRedirects : IStorageBackend
+	{
+		/// <summary>
+		/// Gets a redirect for a read request
+		/// </summary>
+		/// <param name="path">Path to read from</param>
+		/// <param name="cancellationToken">Cancellation token for the operation</param>
+		/// <returns>Path to upload the data to</returns>
+		ValueTask<Uri?> GetReadRedirectAsync(string path, CancellationToken cancellationToken = default);
+
+		/// <summary>
+		/// Gets a redirect for a write request
+		/// </summary>
+		/// <param name="path">Path to write to</param>
+		/// <param name="cancellationToken">Cancellation token for the operation</param>
+		/// <returns>Path to upload the data to</returns>
+		ValueTask<Uri?> GetWriteRedirectAsync(string path, CancellationToken cancellationToken = default);
+	}
+
 	/// <summary>
 	/// Generic version of IStorageBackend, to allow for dependency injection of different singletons
 	/// </summary>
diff --git a/Engine/Source/Programs/Horde/Horde.Build/Storage/StorageController.cs b/Engine/Source/Programs/Horde/Horde.Build/Storage/StorageController.cs
--- a/Engine/Source/Programs/Horde/Horde.Build/Storage/StorageController.cs
+++ b/Engine/Source/Programs/Horde/Horde.Build/Storage/StorageController.cs
@@ -27,6 +27,16 @@ namespace Horde.Build.Storage
 		/// Locator for the uploaded bundle
 		/// </summary>
 		public BlobLocator Locator { get; set; }
+
+		/// <summary>
+		/// URL to upload the blob to.
+		/// </summary>
+		public Uri? UploadUrl { get; set; }
+
+		/// <summary>
+		/// Flag for whether the client could use a redirect instead (ie. not post content to the server, and get an upload url back).
+		/// </summary>
+		public bool? SupportsRedirects { get; set; }
 	}
 
 	/// <summary>
@@ -98,23 +108,38 @@ namespace Horde.Build.Storage
 		/// Uploads data to the storage service. 
 		/// </summary>
 		/// <param name="namespaceId">Namespace to fetch from</param>
-		/// <param name="file">Data to be uploaded</param>
+		/// <param name="file">Data to be uploaded. May be null, in which case the server may return a separate url.</param>
 		/// <param name="prefix">Prefix for the uploaded file</param>
 		/// <param name="cancellationToken">Cancellation token for the operation</param>
 		[HttpPost]
 		[Route("/api/v1/storage/{namespaceId}/blobs")]
-		public async Task<ActionResult<WriteBlobResponse>> WriteBlobAsync(NamespaceId namespaceId, IFormFile file, [FromForm] string? prefix = default, CancellationToken cancellationToken = default)
+		public async Task<ActionResult<WriteBlobResponse>> WriteBlobAsync(NamespaceId namespaceId, IFormFile? file, [FromForm] string? prefix = default, CancellationToken cancellationToken = default)
 		{
 			if (!await _storageService.AuthorizeAsync(namespaceId, User, AclAction.WriteBlobs, null, cancellationToken))
 			{
 				return Forbid(AclAction.WriteBlobs);
 			}
 
-			IStorageClient client = await _storageService.GetClientAsync(namespaceId, cancellationToken);
-			using (Stream stream = file.OpenReadStream())
+			IStorageClientImpl client = await _storageService.GetClientAsync(namespaceId, cancellationToken);
+			if (file == null)
 			{
-				BlobLocator locator = await client.WriteBlobAsync(stream, prefix: (prefix == null)? Utf8String.Empty : new Utf8String(prefix), cancellationToken: cancellationToken);
-				return new WriteBlobResponse { Locator = locator };
+				(BlobLocator Locator, Uri UploadUrl)? result = await client.GetWriteRedirectAsync(prefix ?? String.Empty, cancellationToken);
+				if (result == null)
+				{
+					return new WriteBlobResponse { SupportsRedirects = false };
+				}
+				else
+				{
+					return new WriteBlobResponse { Locator = result.Value.Locator, UploadUrl = result.Value.UploadUrl };
+				}
+			}
+			else
+			{
+				using (Stream stream = file.OpenReadStream())
+				{
+					BlobLocator locator = await client.WriteBlobAsync(stream, prefix: (prefix == null) ? Utf8String.Empty : new Utf8String(prefix), cancellationToken: cancellationToken);
+					return new WriteBlobResponse { Locator = locator, SupportsRedirects = client.SupportsRedirects? (bool?)true : null };
+				}
 			}
 		}
 
@@ -135,7 +160,13 @@ namespace Horde.Build.Storage
 				return Forbid(AclAction.WriteBlobs);
 			}
 
-			IStorageClient client = await _storageService.GetClientAsync(namespaceId, cancellationToken);
+			IStorageClientImpl client = await _storageService.GetClientAsync(namespaceId, cancellationToken);
+
+			Uri? redirectUrl = await client.GetReadRedirectAsync(locator, cancellationToken);
+			if (redirectUrl != null)
+			{
+				return Redirect(redirectUrl.ToString());
+			}
 
 #pragma warning disable CA2000 // Dispose objects before losing scope
 			// TODO: would be better to use the range header here, but seems to require a lot of plumbing to convert unseekable AWS streams into a format that works with range processing.
diff --git a/Engine/Source/Programs/Horde/Horde.Build/Storage/StorageService.cs b/Engine/Source/Programs/Horde/Horde.Build/Storage/StorageService.cs
--- a/Engine/Source/Programs/Horde/Horde.Build/Storage/StorageService.cs
+++ b/Engine/Source/Programs/Horde/Horde.Build/Storage/StorageService.cs
@@ -1,17 +1,14 @@
 // Copyright Epic Games, Inc. All Rights Reserved.
 
 using System;
-using System.Buffers;
 using System.Collections;
 using System.Collections.Generic;
 using System.IO;
 using System.Linq;
 using System.Reflection;
 using System.Security.Claims;
-using System.Text.Json;
 using System.Threading;
 using System.Threading.Tasks;
-using Amazon.EC2.Model;
 using EpicGames.Core;
 using EpicGames.Horde.Storage;
 using Horde.Build.Acls;
@@ -49,17 +46,38 @@ namespace Horde.Build.Storage
 		/// <summary>
 		/// Configuration for this namespace
 		/// </summary>
-		public NamespaceConfig Config { get; }
+		NamespaceConfig Config { get; }
 
 		/// <summary>
 		/// The storage backend
 		/// </summary>
-		public IStorageBackend Backend { get; }
+		IStorageBackend Backend { get; }
 
 		/// <summary>
 		/// The ref collection
 		/// </summary>
-		public IRefCollection Refs { get; }
+		IRefCollection Refs { get; }
+
+		/// <summary>
+		/// Whether the backend supports redirects
+		/// </summary>
+		bool SupportsRedirects { get; }
+
+		/// <summary>
+		/// Gets a redirect for a read request
+		/// </summary>
+		/// <param name="locator">Locator for the blob</param>
+		/// <param name="cancellationToken">Cancellation token for the operation</param>
+		/// <returns>Path to upload the data to</returns>
+		ValueTask<Uri?> GetReadRedirectAsync(BlobLocator locator, CancellationToken cancellationToken = default);
+
+		/// <summary>
+		/// Gets a redirect for a write request
+		/// </summary>
+		/// <param name="prefix">Prefix for the new blob locator</param>
+		/// <param name="cancellationToken">Cancellation token for the operation</param>
+		/// <returns>Locator and path to upload the data to</returns>
+		ValueTask<(BlobLocator, Uri)?> GetWriteRedirectAsync(Utf8String prefix = default, CancellationToken cancellationToken = default);
 	}
 
 	/// <summary>
@@ -77,12 +95,15 @@ namespace Horde.Build.Storage
 
 			public IRefCollection Refs { get; }
 
+			public bool SupportsRedirects { get; }
+
 			public StorageClient(NamespaceConfig config, IStorageBackend backend, IRefCollection refs, IMemoryCache cache, ILogger logger)
 				: base(cache, logger)
 			{
 				Config = config;
 				Backend = backend;
 				Refs = refs;
+				SupportsRedirects = backend is IStorageBackendWithRedirects;
 			}
 
 			public void Dispose()
@@ -111,6 +132,17 @@ namespace Horde.Build.Storage
 				return stream;
 			}
 
+			/// <inheritdoc/>
+			public ValueTask<Uri?> GetReadRedirectAsync(BlobLocator locator, CancellationToken cancellationToken = default)
+			{
+				IStorageBackendWithRedirects? redirectBackend = Backend as IStorageBackendWithRedirects;
+				if (redirectBackend == null)
+				{
+					return new ValueTask<Uri?>((Uri?)null);
+				}
+				return redirectBackend.GetReadRedirectAsync(GetBlobPath(locator.BlobId), cancellationToken);
+			}
+
 			/// <inheritdoc/>
 			public override async Task<Stream> ReadBlobRangeAsync(BlobLocator locator, int offset, int length, CancellationToken cancellationToken = default)
 			{
@@ -136,6 +168,27 @@ namespace Horde.Build.Storage
 				return new BlobLocator(HostId.Empty, id);
 			}
 
+			/// <inheritdoc/>
+			public async ValueTask<(BlobLocator, Uri)?> GetWriteRedirectAsync(Utf8String prefix = default, CancellationToken cancellationToken = default)
+			{
+				IStorageBackendWithRedirects? redirectBackend = Backend as IStorageBackendWithRedirects;
+				if (redirectBackend == null)
+				{
+					return null;
+				}
+
+				BlobId id = BlobId.CreateNew(prefix);
+				string path = GetBlobPath(id);
+
+				Uri? url = await redirectBackend.GetWriteRedirectAsync(path, cancellationToken);
+				if (url == null)
+				{
+					return null;
+				}
+
+				return (new BlobLocator(HostId.Empty, id), url);
+			}
+
 			#endregion
 
 			#region Refs
diff --git a/Engine/Source/Programs/Shared/EpicGames.Horde/Storage/Backends/HttpStorageClient.cs b/Engine/Source/Programs/Shared/EpicGames.Horde/Storage/Backends/HttpStorageClient.cs
--- a/Engine/Source/Programs/Shared/EpicGames.Horde/Storage/Backends/HttpStorageClient.cs
+++ b/Engine/Source/Programs/Shared/EpicGames.Horde/Storage/Backends/HttpStorageClient.cs
@@ -2,12 +2,15 @@
 
 using System;
 using System.IO;
+using System.Net;
 using System.Net.Http;
 using System.Net.Http.Headers;
 using System.Net.Http.Json;
+using System.Runtime.InteropServices;
 using System.Threading;
 using System.Threading.Tasks;
 using EpicGames.Core;
+using EpicGames.Horde.Storage;
 using Microsoft.Extensions.Caching.Memory;
 using Microsoft.Extensions.Logging;
 
@@ -21,6 +24,8 @@ namespace EpicGames.Horde.Storage.Backends
 		class WriteBlobResponse
 		{
 			public BlobLocator Locator { get; set; }
+			public Uri? UploadUrl { get; set; }
+			public bool? SupportsRedirects { get; set; }
 		}
 
 		class ReadRefResponse
@@ -31,6 +36,8 @@ namespace EpicGames.Horde.Storage.Backends
 
 		readonly NamespaceId _namespaceId;
 		readonly HttpClient _httpClient;
+		readonly ILogger _logger;
+		bool _supportsUploadRedirects = true;
 
 		/// <summary>
 		/// Constructor
@@ -40,6 +47,7 @@ namespace EpicGames.Horde.Storage.Backends
 		{
 			_namespaceId = namespaceId;
 			_httpClient = httpClient;
+			_logger = logger;
 		}
 
 		#region Blobs
@@ -68,14 +76,39 @@ namespace EpicGames.Horde.Storage.Backends
 
 		/// <inheritdoc/>
 		public override async Task<BlobLocator> WriteBlobAsync(Stream stream, Utf8String prefix = default, CancellationToken cancellationToken = default)
+		{
+			using StreamContent streamContent = new StreamContent(stream);
+
+			if (_supportsUploadRedirects)
+			{
+				WriteBlobResponse redirectResponse = await SendWriteRequestAsync(null, prefix, cancellationToken);
+				if(redirectResponse.UploadUrl != null)
+				{
+					_logger.LogDebug("Using upload redirect for {Locator}", redirectResponse.Locator);
+					using (HttpResponseMessage uploadResponse = await _httpClient.PostAsync(redirectResponse.UploadUrl, streamContent, cancellationToken))
+					{
+						uploadResponse.EnsureSuccessStatusCode();
+					}
+					return redirectResponse.Locator;
+				}
+			}
+
+			WriteBlobResponse response = await SendWriteRequestAsync(streamContent, prefix, cancellationToken);
+			_supportsUploadRedirects = response.SupportsRedirects ?? false;
+			return response.Locator;
+		}
+
+		async Task<WriteBlobResponse> SendWriteRequestAsync(StreamContent? streamContent, Utf8String prefix = default, CancellationToken cancellationToken = default)
 		{
 			using (HttpRequestMessage request = new HttpRequestMessage(HttpMethod.Post, $"api/v1/storage/{_namespaceId}/blobs"))
 			{
-				using StreamContent streamContent = new StreamContent(stream);
 				using StringContent stringContent = new StringContent(prefix.ToString());
 
 				MultipartFormDataContent form = new MultipartFormDataContent();
-				form.Add(streamContent, "file", "filename");
+				if (streamContent != null)
+				{
+					form.Add(streamContent, "file", "filename");
+				}
 				form.Add(stringContent, "prefix");
 
 				request.Content = form;
@@ -83,7 +116,7 @@ namespace EpicGames.Horde.Storage.Backends
 				{
 					response.EnsureSuccessStatusCode();
 					WriteBlobResponse? data = await response.Content.ReadFromJsonAsync<WriteBlobResponse>(cancellationToken: cancellationToken);
-					return data!.Locator;
+					return data!;
 				}
 			}
 		}
