Horde: Skip all pending steps when a batch fails to execute with a non-permanent reason.

diff --git a/Engine/Source/Programs/Horde/Horde.Build/Jobs/JobCollection.cs b/Engine/Source/Programs/Horde/Horde.Build/Jobs/JobCollection.cs
--- a/Engine/Source/Programs/Horde/Horde.Build/Jobs/JobCollection.cs
+++ b/Engine/Source/Programs/Horde/Horde.Build/Jobs/JobCollection.cs
@@ -6,6 +6,7 @@ using System.Diagnostics.CodeAnalysis;
 using System.IO;
 using System.Linq;
 using System.Threading.Tasks;
+using System.Xml.Linq;
 using EpicGames.Core;
 using Horde.Build.Acls;
 using Horde.Build.Agents;
@@ -1302,6 +1303,21 @@ namespace Horde.Build.Jobs
 				batch.Steps.RemoveAll(x => x.State == JobStepState.Waiting || x.State == JobStepState.Ready);
 			}
 
+			// Mark any steps in failed batches as skipped
+			foreach (JobStepBatchDocument batch in job.Batches)
+			{
+				if (batch.State == JobStepBatchState.Complete && batch.Error != JobStepBatchError.Incomplete)
+				{
+					foreach (JobStepDocument step in batch.Steps)
+					{
+						if (step.IsPending())
+						{
+							step.State = JobStepState.Skipped;
+						}
+					}
+				}
+			}
+
 			// Remove any skipped nodes whose skipped state is no longer valid
 			HashSet<INode> failedNodes = new HashSet<INode>();
 			foreach (JobStepBatchDocument batch in job.Batches)
