Horde: Rename Refs collection to Storage.Refs.

diff --git a/Engine/Source/Programs/Horde/Horde.Build/Storage/RefCollection.cs b/Engine/Source/Programs/Horde/Horde.Build/Storage/RefCollection.cs
--- a/Engine/Source/Programs/Horde/Horde.Build/Storage/RefCollection.cs
+++ b/Engine/Source/Programs/Horde/Horde.Build/Storage/RefCollection.cs
@@ -107,7 +107,7 @@ namespace Horde.Build.Storage
 		/// <param name="logger">Logger for output</param>
 		public RefCollection(MongoService mongoService, IClock clock, IMemoryCache cache, ILogger<RefCollection> logger)
 		{
-			_refs = mongoService.GetCollection<RefDocument>("Refs", keys => keys.Ascending(x => x.NamespaceId).Ascending(x => x.Name), unique: true);
+			_refs = mongoService.GetCollection<RefDocument>("Storage.Refs", keys => keys.Ascending(x => x.NamespaceId).Ascending(x => x.Name), unique: true);
 			_clock = clock;
 			_cache = cache;
 			_logger = logger;
