Horde: Fix wrong ordering of instances to start in AWS fleet manager

diff --git a/Engine/Source/Programs/Horde/Horde.Build.Tests/Fleet/AwsRecyclingFleetManagerTest.cs b/Engine/Source/Programs/Horde/Horde.Build.Tests/Fleet/AwsRecyclingFleetManagerTest.cs
--- a/Engine/Source/Programs/Horde/Horde.Build.Tests/Fleet/AwsRecyclingFleetManagerTest.cs
+++ b/Engine/Source/Programs/Horde/Horde.Build.Tests/Fleet/AwsRecyclingFleetManagerTest.cs
@@ -129,17 +129,17 @@ namespace Horde.Build.Tests.Fleet
 		{
 			FakeAmazonEc2 ec2 = new ();
 			Instance i1 = ec2.AddInstance(FakeAmazonEc2.StateStopped, InstanceType.M5Large, launchTime: DateTime.Parse("Aug 18, 2018"));
-			Instance i2 = ec2.AddInstance(FakeAmazonEc2.StateStopped, InstanceType.M5Large, launchTime: DateTime.Parse("Aug 5, 2018"));
-			Instance i3 = ec2.AddInstance(FakeAmazonEc2.StateStopped, InstanceType.M5Large, launchTime: DateTime.Parse("Aug 22, 2018"));
-			Instance i4 = ec2.AddInstance(FakeAmazonEc2.StateStopped, InstanceType.M5Large, launchTime: DateTime.Parse("Aug 1, 2018"));
-			Instance i5 = ec2.AddInstance(FakeAmazonEc2.StateStopped, InstanceType.M5Large, launchTime: DateTime.Parse("Aug 28, 2018"));
+			Instance i2 = ec2.AddInstance(FakeAmazonEc2.StateStopped, InstanceType.M5Large, launchTime: DateTime.Parse("Aug 22, 2018"));
+			Instance i3 = ec2.AddInstance(FakeAmazonEc2.StateStopped, InstanceType.M5Large, launchTime: DateTime.Parse("Aug 5, 2018"));
+			Instance i4 = ec2.AddInstance(FakeAmazonEc2.StateStopped, InstanceType.M5Large, launchTime: DateTime.Parse("Aug 28, 2018"));
+			Instance i5 = ec2.AddInstance(FakeAmazonEc2.StateStopped, InstanceType.M5Large, launchTime: DateTime.Parse("Aug 1, 2018"));
 			ec2.SetCapacity(FakeAmazonEc2.AzUsEast1A, InstanceType.M5Large, 2);
 
 			await ExpandPoolAsync(ec2.Get(), 2, new(new List<string>()));
 			Assert.AreEqual(3, ec2.GetStoppedInstanceCount());
 			Assert.AreEqual(2, ec2.GetPendingInstanceCount());
-			Assert.AreEqual(FakeAmazonEc2.StatePending.Name, ec2.GetInstance(i2.InstanceId)!.State.Name);
 			Assert.AreEqual(FakeAmazonEc2.StatePending.Name, ec2.GetInstance(i4.InstanceId)!.State.Name);
+			Assert.AreEqual(FakeAmazonEc2.StatePending.Name, ec2.GetInstance(i2.InstanceId)!.State.Name);
 		}
 
 		[TestMethod]
diff --git a/Engine/Source/Programs/Horde/Horde.Build/Agents/Fleet/Providers/AwsRecyclingFleetManager.cs b/Engine/Source/Programs/Horde/Horde.Build/Agents/Fleet/Providers/AwsRecyclingFleetManager.cs
--- a/Engine/Source/Programs/Horde/Horde.Build/Agents/Fleet/Providers/AwsRecyclingFleetManager.cs
+++ b/Engine/Source/Programs/Horde/Horde.Build/Agents/Fleet/Providers/AwsRecyclingFleetManager.cs
@@ -131,7 +131,7 @@ public sealed class AwsRecyclingFleetManager : IFleetManager
 		Dictionary<string, List<Instance>> result = new();
 		foreach ((string az, List<Instance> candidates) in candidatesPerAz)
 		{
-			result[az] = candidates.OrderBy(x => x.LaunchTime).Take(requestCountPerAz[az]).ToList();
+			result[az] = candidates.OrderByDescending(x => x.LaunchTime).Take(requestCountPerAz[az]).ToList();
 		}
 
 		return result;
