Horde: Fix inconsistent ref names in temp storage, and add support for uploading artifacts to storage backend.

diff --git a/Engine/Source/Programs/Horde/Horde.Agent/Execution/JobExecutor.cs b/Engine/Source/Programs/Horde/Horde.Agent/Execution/JobExecutor.cs
--- a/Engine/Source/Programs/Horde/Horde.Agent/Execution/JobExecutor.cs
+++ b/Engine/Source/Programs/Horde/Horde.Agent/Execution/JobExecutor.cs
@@ -17,6 +17,7 @@ using System.Threading.Tasks;
 using EpicGames.Core;
 using EpicGames.Horde.Storage;
 using EpicGames.Horde.Storage.Backends;
+using EpicGames.Horde.Storage.Nodes;
 using Grpc.Core;
 using Horde.Agent.Parser;
 using Horde.Agent.Services;
@@ -404,7 +405,7 @@ namespace Horde.Agent.Execution
 				arguments.Append($" CopyUAT -WithLauncher -TargetDir=\"{buildDir}\"");
 			}
 
-			int result = await ExecuteAutomationToolAsync(step, workspaceDir, arguments.ToString(), logger, cancellationToken);
+			int result = await ExecuteAutomationToolAsync(step, workspaceDir, arguments.ToString(), null, logger, cancellationToken);
 			if (result != 0)
 			{
 				return false;
@@ -645,7 +646,7 @@ namespace Horde.Agent.Execution
 				{
 					arguments.AppendArgument("-SharedStorageDir=", sharedStorageDir.FullName);
 				}
-				return await ExecuteAutomationToolAsync(step, workspaceDir, arguments.ToString(), logger, cancellationToken) == 0;
+				return await ExecuteAutomationToolAsync(step, workspaceDir, arguments.ToString(), null, logger, cancellationToken) == 0;
 			}
 		}
 
@@ -683,7 +684,7 @@ namespace Horde.Agent.Execution
 				string nodeName = input.Substring(0, slashIdx);
 				string tagName = input.Substring(slashIdx + 1);
 
-				TempStorageTagManifest fileList = await TempStorage.RetrieveTagAsync(storage, nodeName, tagName, manifestDir, logger, cancellationToken);
+				TempStorageTagManifest fileList = await TempStorage.RetrieveTagAsync(storage, RefPrefix, nodeName, tagName, manifestDir, logger, cancellationToken);
 				tagNameToFileSet[tagName] = fileList.ToFileSet(workspaceDir);
 				inputStorageBlocks.UnionWith(fileList.Blocks);
 			}
@@ -695,7 +696,7 @@ namespace Horde.Agent.Execution
 				scope.Span.SetTag("blocks", inputStorageBlocks.Count);
 				foreach (TempStorageBlockRef inputStorageBlock in inputStorageBlocks)
 				{
-					TempStorageBlockManifest manifest = await TempStorage.RetrieveBlockAsync(storage, inputStorageBlock.NodeName, inputStorageBlock.OutputName, workspaceDir, manifestDir, logger, cancellationToken);
+					TempStorageBlockManifest manifest = await TempStorage.RetrieveBlockAsync(storage, RefPrefix, inputStorageBlock.NodeName, inputStorageBlock.OutputName, workspaceDir, manifestDir, logger, cancellationToken);
 					inputManifests[inputStorageBlock] = manifest;
 				}
 				scope.Span.SetTag("size", inputManifests.Sum(x => x.Value.GetTotalSize()));
@@ -720,7 +721,7 @@ namespace Horde.Agent.Execution
 			}
 
 			// Run UAT
-			if (await ExecuteAutomationToolAsync(step, workspaceDir, arguments, logger, cancellationToken) != 0)
+			if (await ExecuteAutomationToolAsync(step, workspaceDir, arguments, storage, logger, cancellationToken) != 0)
 			{
 				return false;
 			}
@@ -823,7 +824,7 @@ namespace Horde.Agent.Execution
 			// Write all the storage blocks, and update the mapping from file to storage block
 			using (GlobalTracer.Instance.BuildSpan("TempStorage").WithTag("resource", "Write").StartActive())
 			{
-				RefName refName = new RefName(RefName.Sanitize($"{_job.StreamId}/{_job.Change}-{_jobId}/{step.Name}"));
+				RefName refName = TempStorage.GetRefNameForNode(RefPrefix, step.Name);
 
 				TreeOptions treeOptions = new TreeOptions();
 				TreeWriter treeWriter = new TreeWriter(storage, treeOptions, refName.Text);
@@ -869,7 +870,7 @@ namespace Horde.Agent.Execution
 			return true;
 		}
 
-		protected async Task<int> ExecuteAutomationToolAsync(BeginStepResponse step, DirectoryReference workspaceDir, string arguments, ILogger logger, CancellationToken cancellationToken)
+		protected async Task<int> ExecuteAutomationToolAsync(BeginStepResponse step, DirectoryReference workspaceDir, string arguments, IStorageClient? store, ILogger logger, CancellationToken cancellationToken)
 		{
 			int result;
 			using (IScope scope = GlobalTracer.Instance.BuildSpan("BuildGraph").StartActive())
@@ -881,15 +882,15 @@ namespace Horde.Agent.Execution
 
 				if (RuntimeInformation.IsOSPlatform(OSPlatform.Windows))
 				{
-					result = await ExecuteCommandAsync(step, workspaceDir, Environment.GetEnvironmentVariable("COMSPEC") ?? "cmd.exe", $"/C \"\"{workspaceDir}\\Engine\\Build\\BatchFiles\\RunUAT.bat\" {arguments}\"", logger, cancellationToken);
+					result = await ExecuteCommandAsync(step, workspaceDir, Environment.GetEnvironmentVariable("COMSPEC") ?? "cmd.exe", $"/C \"\"{workspaceDir}\\Engine\\Build\\BatchFiles\\RunUAT.bat\" {arguments}\"", store, logger, cancellationToken);
 				}
 				else if (RuntimeInformation.IsOSPlatform(OSPlatform.Linux))
 				{
-					result = await ExecuteCommandAsync(step, workspaceDir, "/bin/bash", $"\"{workspaceDir}/Engine/Build/BatchFiles/RunUAT.sh\" {arguments}", logger, cancellationToken);
+					result = await ExecuteCommandAsync(step, workspaceDir, "/bin/bash", $"\"{workspaceDir}/Engine/Build/BatchFiles/RunUAT.sh\" {arguments}", store, logger, cancellationToken);
 				}
 				else if (RuntimeInformation.IsOSPlatform(OSPlatform.OSX))
 				{
-					result = await ExecuteCommandAsync(step, workspaceDir, "/bin/sh", $"\"{workspaceDir}/Engine/Build/BatchFiles/RunUAT.sh\" {arguments}", logger, cancellationToken);
+					result = await ExecuteCommandAsync(step, workspaceDir, "/bin/sh", $"\"{workspaceDir}/Engine/Build/BatchFiles/RunUAT.sh\" {arguments}", store, logger, cancellationToken);
 				}
 				else
 				{
@@ -1006,7 +1007,9 @@ namespace Horde.Agent.Execution
 			}
 		}
 
-		async Task<int> ExecuteCommandAsync(BeginStepResponse step, DirectoryReference workspaceDir, string fileName, string arguments, ILogger logger, CancellationToken cancellationToken)
+		string RefPrefix => $"{_job.StreamId}/{_job.Change}-{_jobId}/";
+
+		async Task<int> ExecuteCommandAsync(BeginStepResponse step, DirectoryReference workspaceDir, string fileName, string arguments, IStorageClient? store, ILogger logger, CancellationToken cancellationToken)
 		{
 			// Combine all the supplied environment variables together
 			Dictionary<string, string> newEnvVars = new Dictionary<string, string>(_envVars, StringComparer.Ordinal);
@@ -1225,29 +1228,46 @@ namespace Horde.Agent.Execution
 				await UploadTestDataAsync(step.StepId, combinedTestData);
 			}
 					
-			if (DirectoryReference.Exists(logDir))
+			if (store != null)
 			{
-				Dictionary<FileReference, string> artifactFileToId = new Dictionary<FileReference, string>();
-				foreach (FileReference artifactFile in DirectoryReference.EnumerateFiles(logDir, "*", SearchOption.AllDirectories))
-				{
-					string artifactName = artifactFile.MakeRelativeTo(logDir);
+				RefName refName = new RefName(RefName.Sanitize($"{RefPrefix}/artifacts"));
 
-					string? artifactId = await ArtifactUploader.UploadAsync(RpcConnection, _jobId, _batchId, step.StepId, artifactName, artifactFile, logger, cancellationToken);
-					if (artifactId != null)
-					{
-						artifactFileToId[artifactFile] = artifactId;
-					}
-				}
+				TreeOptions treeOptions = new TreeOptions();
+				TreeWriter treeWriter = new TreeWriter(store, treeOptions, refName.Text);
+
+				DirectoryNode directoryNode = new DirectoryNode(DirectoryFlags.None);
+
+				ChunkingOptions options = new ChunkingOptions();
+				await directoryNode.CopyFromDirectoryAsync(logDir.ToDirectoryInfo(), options, treeWriter, cancellationToken);
 
-				foreach (FileReference reportFile in artifactFileToId.Keys.Where(x => x.HasExtension(".report.json")))
+				await treeWriter.WriteRefAsync(refName, directoryNode, cancellationToken: cancellationToken);
+			}
+			else
+			{
+				if (DirectoryReference.Exists(logDir))
 				{
-					try
+					Dictionary<FileReference, string> artifactFileToId = new Dictionary<FileReference, string>();
+					foreach (FileReference artifactFile in DirectoryReference.EnumerateFiles(logDir, "*", SearchOption.AllDirectories))
 					{
-						await CreateReportAsync(step.StepId, reportFile, artifactFileToId, logger);
+						string artifactName = artifactFile.MakeRelativeTo(logDir);
+
+						string? artifactId = await ArtifactUploader.UploadAsync(RpcConnection, _jobId, _batchId, step.StepId, artifactName, artifactFile, logger, cancellationToken);
+						if (artifactId != null)
+						{
+							artifactFileToId[artifactFile] = artifactId;
+						}
 					}
-					catch(Exception ex)
+
+					foreach (FileReference reportFile in artifactFileToId.Keys.Where(x => x.HasExtension(".report.json")))
 					{
-						logger.LogWarning("Unable to upload report: {Message}", ex.Message);
+						try
+						{
+							await CreateReportAsync(step.StepId, reportFile, artifactFileToId, logger);
+						}
+						catch (Exception ex)
+						{
+							logger.LogWarning("Unable to upload report: {Message}", ex.Message);
+						}
 					}
 				}
 			}
diff --git a/Engine/Source/Programs/Horde/Horde.Agent/Utility/TempStorage.cs b/Engine/Source/Programs/Horde/Horde.Agent/Utility/TempStorage.cs
--- a/Engine/Source/Programs/Horde/Horde.Agent/Utility/TempStorage.cs
+++ b/Engine/Source/Programs/Horde/Horde.Agent/Utility/TempStorage.cs
@@ -671,24 +671,26 @@ namespace Horde.Storage.Utility
 		/// <summary>
 		/// Gets the ref name for a particular node
 		/// </summary>
+		/// <param name="refPrefix">Prefix for refs in this job</param>
 		/// <param name="nodeName"></param>
 		/// <returns></returns>
-		static RefName GetRefNameForNode(string nodeName)
+		public static RefName GetRefNameForNode(string refPrefix, string nodeName)
 		{
-			return new RefName(RefName.Sanitize(nodeName));
+			return new RefName(RefName.Sanitize($"{refPrefix}/steps/{nodeName}"));
 		}
 
 		/// <summary>
 		/// Reads a set of tagged files from disk
 		/// </summary>
 		/// <param name="store"></param>
+		/// <param name="refPrefix">Prefix for ref names</param>
 		/// <param name="nodeName">Name of the node which produced the tag set</param>
 		/// <param name="tagName">Name of the tag, with a '#' prefix</param>
 		/// <param name="manifestDir">The local directory containing manifests</param>
 		/// <param name="logger">Logger for output</param>
 		/// <param name="cancellationToken"></param>
 		/// <returns>The set of files</returns>
-		public static async Task<TempStorageTagManifest> RetrieveTagAsync(IStorageClient store, string nodeName, string tagName, DirectoryReference manifestDir, ILogger logger, CancellationToken cancellationToken)
+		public static async Task<TempStorageTagManifest> RetrieveTagAsync(IStorageClient store, string refPrefix, string nodeName, string tagName, DirectoryReference manifestDir, ILogger logger, CancellationToken cancellationToken)
 		{
 			TempStorageTagManifest fileList;
 
@@ -701,7 +703,7 @@ namespace Horde.Storage.Utility
 			}
 			else
 			{
-				RefName refName = GetRefNameForNode(nodeName);
+				RefName refName = GetRefNameForNode(refPrefix, nodeName);
 				logger.LogInformation("Reading node \"{NodeName}\" tag \"{TagName}\" from temp storage (ref: {RefName}, localFile: {LocalFile})", nodeName, tagName, refName, localFileListLocation);
 
 				TempStorageNode node = await store.ReadNodeAsync<TempStorageNode>(refName, cancellationToken: cancellationToken);
@@ -806,6 +808,7 @@ namespace Horde.Storage.Utility
 		/// Retrieve an output of the given node. Fetches and decompresses the files from shared storage if necessary, or validates the local files.
 		/// </summary>
 		/// <param name="store">Store to read data from</param>
+		/// <param name="refPrefix">Prefix for ref names</param>
 		/// <param name="nodeName">The node which created the storage block</param>
 		/// <param name="blockName">Name of the block to retrieve.</param>
 		/// <param name="rootDir">Local directory for extracting data to</param>
@@ -813,7 +816,7 @@ namespace Horde.Storage.Utility
 		/// <param name="logger">Logger for output</param>
 		/// <param name="cancellationToken"></param>
 		/// <returns>Manifest of the files retrieved</returns>
-		public static async Task<TempStorageBlockManifest> RetrieveBlockAsync(IStorageClient store, string nodeName, string blockName, DirectoryReference rootDir, DirectoryReference manifestDir, ILogger logger, CancellationToken cancellationToken)
+		public static async Task<TempStorageBlockManifest> RetrieveBlockAsync(IStorageClient store, string refPrefix, string nodeName, string blockName, DirectoryReference rootDir, DirectoryReference manifestDir, ILogger logger, CancellationToken cancellationToken)
 		{
 			// Get the path to the local manifest
 			FileReference localManifestFile = GetBlockManifestLocation(manifestDir, nodeName, blockName);
@@ -829,7 +832,7 @@ namespace Horde.Storage.Utility
 			else
 			{
 				// Read the shared manifest
-				RefName refName = GetRefNameForNode(nodeName);
+				RefName refName = GetRefNameForNode(refPrefix, nodeName);
 				logger.LogInformation("Reading node \"{NodeName}\" block \"{BlockName}\" from temp storage (ref: {RefName}, local: {LocalFile})", nodeName, blockName, refName, localManifestFile);
 
 				TempStorageNode node = await store.ReadNodeAsync<TempStorageNode>(refName, cancellationToken: cancellationToken);
