Horde: Log warnings and errors from agent's gRPC client

diff --git a/Engine/Source/Programs/Horde/Horde.Agent/Commands/Compute/AwsLambdaListenerCommand.cs b/Engine/Source/Programs/Horde/Horde.Agent/Commands/Compute/AwsLambdaListenerCommand.cs
--- a/Engine/Source/Programs/Horde/Horde.Agent/Commands/Compute/AwsLambdaListenerCommand.cs
+++ b/Engine/Source/Programs/Horde/Horde.Agent/Commands/Compute/AwsLambdaListenerCommand.cs
@@ -150,7 +150,7 @@ namespace Horde.Agent.Commands.Compute
 				agentSettings.PerforceExecutor.RunConform = false;
 				
 				OptionsWrapper<AgentSettings> agentSettingsOptions = new (agentSettings);
-				GrpcService grpcService = new (agentSettingsOptions, grpcServiceLogger);
+				GrpcService grpcService = new (agentSettingsOptions, grpcServiceLogger, loggerFactory);
 
 				AppDomain.CurrentDomain.ProcessExit += (s, e) => 
 				{
diff --git a/Engine/Source/Programs/Horde/Horde.Agent/Logging.cs b/Engine/Source/Programs/Horde/Horde.Agent/Logging.cs
--- a/Engine/Source/Programs/Horde/Horde.Agent/Logging.cs
+++ b/Engine/Source/Programs/Horde/Horde.Agent/Logging.cs
@@ -92,6 +92,7 @@ namespace Horde.Agent
 				.MinimumLevel.Override("Microsoft.AspNetCore.Hosting.Diagnostics", LogEventLevel.Warning)
 				.MinimumLevel.Override("Microsoft.AspNetCore.Mvc.Infrastructure.ControllerActionInvoker", LogEventLevel.Warning)
 				.MinimumLevel.Override("Serilog.AspNetCore.RequestLoggingMiddleware", LogEventLevel.Warning)
+				.MinimumLevel.Override("Grpc.Net.Client.Internal.GrpcCall", LogEventLevel.Warning)
 				.MinimumLevel.ControlledBy(LogLevelSwitch)
 				.Enrich.FromLogContext()
 				.Enrich.With<DatadogLogEnricher>()
diff --git a/Engine/Source/Programs/Horde/Horde.Agent/Services/GrpcService.cs b/Engine/Source/Programs/Horde/Horde.Agent/Services/GrpcService.cs
--- a/Engine/Source/Programs/Horde/Horde.Agent/Services/GrpcService.cs
+++ b/Engine/Source/Programs/Horde/Horde.Agent/Services/GrpcService.cs
@@ -5,11 +5,8 @@ using System.Net.Http;
 using System.Net.Http.Headers;
 using Grpc.Net.Client;
 using Horde.Agent.Utility;
-using Microsoft.Extensions.Http;
 using Microsoft.Extensions.Logging;
 using Microsoft.Extensions.Options;
-using Polly;
-using Polly.Extensions.Http;
 
 namespace Horde.Agent.Services
 {
@@ -20,16 +17,19 @@ namespace Horde.Agent.Services
 	{
 		private readonly ServerProfile _serverProfile;
 		private readonly ILogger _logger;
+		private readonly ILoggerFactory _loggerFactory;
 
 		/// <summary>
 		/// Constructor
 		/// </summary>
 		/// <param name="settings"></param>
 		/// <param name="logger"></param>
-		public GrpcService(IOptions<AgentSettings> settings, ILogger<GrpcService> logger)
+		/// <param name="loggerFactory"></param>
+		public GrpcService(IOptions<AgentSettings> settings, ILogger<GrpcService> logger, ILoggerFactory loggerFactory)
 		{
 			_serverProfile = settings.Value.GetCurrentServerProfile();
 			_logger = logger;
+			_loggerFactory = loggerFactory;
 		}
 
 		/// <summary>
@@ -86,6 +86,7 @@ namespace Horde.Agent.Services
 				MaxReceiveMessageSize = 1024 * 1024 * 1024, // 1 GB
 				MaxSendMessageSize = 1024 * 1024 * 1024, // 1 GB
 				
+				LoggerFactory = _loggerFactory,
 				HttpClient = httpClient,
 				DisposeHttpClient = true
 			});
