Horde: Remove trailing slashes from ref prefix, causing exception.

diff --git a/Engine/Source/Programs/Horde/Horde.Agent/Execution/JobExecutor.cs b/Engine/Source/Programs/Horde/Horde.Agent/Execution/JobExecutor.cs
--- a/Engine/Source/Programs/Horde/Horde.Agent/Execution/JobExecutor.cs
+++ b/Engine/Source/Programs/Horde/Horde.Agent/Execution/JobExecutor.cs
@@ -1007,7 +1007,7 @@ namespace Horde.Agent.Execution
 			}
 		}
 
-		string RefPrefix => $"{_job.StreamId}/{_job.Change}-{_jobId}/";
+		string RefPrefix => $"{_job.StreamId}/{_job.Change}-{_jobId}";
 
 		async Task<int> ExecuteCommandAsync(BeginStepResponse step, DirectoryReference workspaceDir, string fileName, string arguments, IStorageClient? store, ILogger logger, CancellationToken cancellationToken)
 		{
