Horde: Fix task name in build script.

diff --git a/Engine/Source/Programs/Horde/StorageTests.xml b/Engine/Source/Programs/Horde/StorageTests.xml
--- a/Engine/Source/Programs/Horde/StorageTests.xml
+++ b/Engine/Source/Programs/Horde/StorageTests.xml
@@ -6,7 +6,7 @@
 	<!-- Small Chunks -->
 	<Agent Name="Agent A" Type="CompileWin64;Win64">
 		<Node Name="Create Small Chunks" Produces="#SmallChunks">
-			<StorageTest Size="1024" Count="20000" Different="true" Tag="#SmallChunks"/>
+			<RandomData Size="1024" Count="20000" Different="true" Tag="#SmallChunks"/>
 			<Log Message="Output Files:" Files="#SmallChunks"/>
 		</Node>
 	</Agent>
@@ -19,7 +19,7 @@
 	<!-- Large Chunks -->
 	<Agent Name="Agent C" Type="CompileWin64;Win64">
 		<Node Name="Create Large Chunks" Produces="#LargeChunks">
-			<StorageTest Size="1073741824" Count="20" Different="true" Tag="#LargeChunks"/>
+			<RandomData Size="1073741824" Count="20" Different="true" Tag="#LargeChunks"/>
 			<Log Message="Output Files:" Files="#LargeChunks"/>
 		</Node>
 	</Agent>
@@ -32,7 +32,7 @@
 	<!-- Small Identical Chunks -->
 	<Agent Name="Agent E" Type="CompileWin64;Win64">
 		<Node Name="Create Small Identical Chunks" Produces="#SmallIdenticalChunks">
-			<StorageTest Size="1024" Count="20000" Different="false" Tag="#SmallIdenticalChunks"/>
+			<RandomData Size="1024" Count="20000" Different="false" Tag="#SmallIdenticalChunks"/>
 			<Log Message="Output Files:" Files="#SmallIdenticalChunks"/>
 		</Node>
 	</Agent>
@@ -45,7 +45,7 @@
 	<!-- Large Identical Chunks -->
 	<Agent Name="Agent G" Type="CompileWin64;Win64">
 		<Node Name="Create Large Identical Chunks" Produces="#LargeIdenticalChunks">
-			<StorageTest Size="1073741824" Count="20" Different="true" Tag="#LargeIdenticalChunks"/>
+			<RandomData Size="1073741824" Count="20" Different="true" Tag="#LargeIdenticalChunks"/>
 			<Log Message="Output Files:" Files="#LargeIdenticalChunks"/>
 		</Node>
 	</Agent>
