Horde: Clean up log output for storage tests.

diff --git a/Engine/Source/Programs/AutomationTool/BuildGraph/Tasks/RandomDataTask.cs b/Engine/Source/Programs/AutomationTool/BuildGraph/Tasks/RandomDataTask.cs
--- a/Engine/Source/Programs/AutomationTool/BuildGraph/Tasks/RandomDataTask.cs
+++ b/Engine/Source/Programs/AutomationTool/BuildGraph/Tasks/RandomDataTask.cs
@@ -81,6 +81,8 @@ namespace AutomationTool.Tasks
 		/// <param name="TagNameToFileSet">Mapping from tag names to the set of files they include</param>
 		public override async Task ExecuteAsync(JobContext Job, HashSet<FileReference> BuildProducts, Dictionary<string, HashSet<FileReference>> TagNameToFileSet)
 		{
+			DirectoryReference OutputDir = ResolveDirectory(Parameters.OutputDir);
+
 			byte[] buffer = Array.Empty<byte>();
 			for (int idx = 0; idx < Parameters.Count; idx++)
 			{
@@ -89,11 +91,13 @@ namespace AutomationTool.Tasks
 					buffer = RandomNumberGenerator.GetBytes(Parameters.Size);
 				}
 
-				FileReference file = FileReference.Combine(ResolveDirectory(Parameters.OutputDir), $"test-{Parameters.Size}-{idx}.dat");
+				FileReference file = FileReference.Combine(OutputDir, $"test-{Parameters.Size}-{idx}.dat");
 				await FileReference.WriteAllBytesAsync(file, buffer);
 				BuildProducts.Add(file);
 			}
 
+			Job.OwnerCommand.Logger.LogInformation("Created {NumFiles} of {Size:n0} bytes in {OutputDir} (Different={Different})", Parameters.Count, Parameters.Size, OutputDir, Parameters.Different);
+
 			// Apply the optional output tag to them
 			foreach (string TagName in FindTagNamesFromList(Parameters.Tag))
 			{
diff --git a/Engine/Source/Programs/Horde/Horde.Agent/Execution/JobExecutor.cs b/Engine/Source/Programs/Horde/Horde.Agent/Execution/JobExecutor.cs
--- a/Engine/Source/Programs/Horde/Horde.Agent/Execution/JobExecutor.cs
+++ b/Engine/Source/Programs/Horde/Horde.Agent/Execution/JobExecutor.cs
@@ -14,6 +14,7 @@ using System.Text.Json;
 using System.Text.Json.Serialization;
 using System.Threading;
 using System.Threading.Tasks;
+using System.Timers;
 using EpicGames.Core;
 using EpicGames.Horde.Storage;
 using EpicGames.Horde.Storage.Backends;
@@ -693,6 +694,7 @@ namespace Horde.Agent.Execution
 			Dictionary<TempStorageBlockRef, TempStorageBlockManifest> inputManifests = new Dictionary<TempStorageBlockRef, TempStorageBlockManifest>();
 			using (IScope scope = GlobalTracer.Instance.BuildSpan("TempStorage").WithTag("resource", "read").StartActive())
 			{
+				Stopwatch timer = Stopwatch.StartNew();
 				scope.Span.SetTag("blocks", inputStorageBlocks.Count);
 				foreach (TempStorageBlockRef inputStorageBlock in inputStorageBlocks)
 				{
@@ -700,6 +702,7 @@ namespace Horde.Agent.Execution
 					inputManifests[inputStorageBlock] = manifest;
 				}
 				scope.Span.SetTag("size", inputManifests.Sum(x => x.Value.GetTotalSize()));
+				logger.LogInformation("Download took {Time:n1}s", timer.Elapsed.TotalSeconds);
 			}
 
 			// Read all the input storage blocks, keeping track of which block each file came from
@@ -709,8 +712,6 @@ namespace Horde.Agent.Execution
 				TempStorageBlockRef inputStorageBlock = pair.Key;
 				foreach (FileReference file in pair.Value.Files.Select(x => x.ToFileReference(workspaceDir)))
 				{
-					logger.LogInformation("Reading input block: {File}", file);
-
 					TempStorageBlockRef? currentStorageBlock;
 					if (fileToStorageBlock.TryGetValue(file, out currentStorageBlock) && !TempStorage.IsDuplicateBuildProduct(file))
 					{
@@ -824,6 +825,7 @@ namespace Horde.Agent.Execution
 			// Write all the storage blocks, and update the mapping from file to storage block
 			using (GlobalTracer.Instance.BuildSpan("TempStorage").WithTag("resource", "Write").StartActive())
 			{
+				Stopwatch timer = Stopwatch.StartNew();
 				RefName refName = TempStorage.GetRefNameForNode(RefPrefix, step.Name);
 
 				TreeOptions treeOptions = new TreeOptions();
@@ -865,6 +867,7 @@ namespace Horde.Agent.Execution
 
 				// Write the final node
 				await treeWriter.WriteRefAsync(refName, outputNode, new RefOptions { Lifetime = TimeSpan.FromDays(3.0) }, cancellationToken);
+				logger.LogInformation("Upload took {Time:n1}s", timer.Elapsed.TotalSeconds);
 			}
 
 			return true;
diff --git a/Engine/Source/Programs/Horde/StorageTests.xml b/Engine/Source/Programs/Horde/StorageTests.xml
--- a/Engine/Source/Programs/Horde/StorageTests.xml
+++ b/Engine/Source/Programs/Horde/StorageTests.xml
@@ -7,7 +7,6 @@
 	<Agent Name="Agent A" Type="CompileWin64;Win64">
 		<Node Name="Create Small Chunks" Produces="#SmallChunks">
 			<RandomData Size="1024" Count="20000" Different="true" Tag="#SmallChunks"/>
-			<Log Message="Output Files:" Files="#SmallChunks"/>
 		</Node>
 	</Agent>
 	<Agent Name="Agent B" Type="CompileWin64;Win64">
@@ -20,7 +19,6 @@
 	<Agent Name="Agent C" Type="CompileWin64;Win64">
 		<Node Name="Create Large Chunks" Produces="#LargeChunks">
 			<RandomData Size="1073741824" Count="20" Different="true" Tag="#LargeChunks"/>
-			<Log Message="Output Files:" Files="#LargeChunks"/>
 		</Node>
 	</Agent>
 	<Agent Name="Agent D" Type="CompileWin64;Win64">
@@ -33,7 +31,6 @@
 	<Agent Name="Agent E" Type="CompileWin64;Win64">
 		<Node Name="Create Small Identical Chunks" Produces="#SmallIdenticalChunks">
 			<RandomData Size="1024" Count="20000" Different="false" Tag="#SmallIdenticalChunks"/>
-			<Log Message="Output Files:" Files="#SmallIdenticalChunks"/>
 		</Node>
 	</Agent>
 	<Agent Name="Agent F" Type="CompileWin64;Win64">
@@ -46,7 +43,6 @@
 	<Agent Name="Agent G" Type="CompileWin64;Win64">
 		<Node Name="Create Large Identical Chunks" Produces="#LargeIdenticalChunks">
 			<RandomData Size="1073741824" Count="20" Different="true" Tag="#LargeIdenticalChunks"/>
-			<Log Message="Output Files:" Files="#LargeIdenticalChunks"/>
 		</Node>
 	</Agent>
 	<Agent Name="Agent H" Type="CompileWin64;Win64">
diff --git a/Engine/Source/Programs/Shared/EpicGames.Horde/Storage/Backends/HttpStorageClient.cs b/Engine/Source/Programs/Shared/EpicGames.Horde/Storage/Backends/HttpStorageClient.cs
--- a/Engine/Source/Programs/Shared/EpicGames.Horde/Storage/Backends/HttpStorageClient.cs
+++ b/Engine/Source/Programs/Shared/EpicGames.Horde/Storage/Backends/HttpStorageClient.cs
@@ -178,7 +178,7 @@ namespace EpicGames.Horde.Storage.Backends
 		/// <inheritdoc/>
 		public override async Task WriteRefTargetAsync(RefName name, NodeLocator target, RefOptions? options = null, CancellationToken cancellationToken = default)
 		{
-			_logger.LogDebug("Setting ref {RefName} -> {Blob}#{ExportIdx}", name, target.Blob, target.ExportIdx);
+			_logger.LogDebug("Writing ref {RefName} -> {Blob}#{ExportIdx}", name, target.Blob, target.ExportIdx);
 			using (HttpResponseMessage response = await _httpClient.PutAsync($"api/v1/storage/{_namespaceId}/refs/{name}", new { locator = target.Blob, exportIdx = target.ExportIdx, options }, cancellationToken))
 			{
 				response.EnsureSuccessStatusCode();
