Horde: Add another class of blobs for storage tests (20,000 x 1mb).

diff --git a/Engine/Source/Programs/AutomationTool/BuildGraph/Tasks/RandomDataTask.cs b/Engine/Source/Programs/AutomationTool/BuildGraph/Tasks/RandomDataTask.cs
--- a/Engine/Source/Programs/AutomationTool/BuildGraph/Tasks/RandomDataTask.cs
+++ b/Engine/Source/Programs/AutomationTool/BuildGraph/Tasks/RandomDataTask.cs
@@ -96,7 +96,7 @@ namespace AutomationTool.Tasks
 				BuildProducts.Add(file);
 			}
 
-			Job.OwnerCommand.Logger.LogInformation("Created {NumFiles} of {Size:n0} bytes in {OutputDir} (Different={Different})", Parameters.Count, Parameters.Size, OutputDir, Parameters.Different);
+			Job.OwnerCommand.Logger.LogInformation("Created {NumFiles:n0} files of {Size:n0} bytes in {OutputDir} (Different={Different})", Parameters.Count, Parameters.Size, OutputDir, Parameters.Different);
 
 			// Apply the optional output tag to them
 			foreach (string TagName in FindTagNamesFromList(Parameters.Tag))
diff --git a/Engine/Source/Programs/Horde/StorageTests.xml b/Engine/Source/Programs/Horde/StorageTests.xml
--- a/Engine/Source/Programs/Horde/StorageTests.xml
+++ b/Engine/Source/Programs/Horde/StorageTests.xml
@@ -4,54 +4,74 @@
 	<Option Name="PreflightChange" DefaultValue="" Description="Preflight changelist number"/>
 	
 	<!-- Small Chunks -->
-	<Agent Name="Agent A" Type="CompileWin64;Win64">
+	<Agent Name="Small Agent A" Type="CompileWin64;Win64">
 		<Node Name="Create Small Chunks" Produces="#SmallChunks">
-			<RandomData Size="1024" Count="20000" Different="true" Tag="#SmallChunks"/>
+			<RandomData Size="1024" Count="20000000" Different="true" Tag="#SmallChunks"/>
 		</Node>
 	</Agent>
-	<Agent Name="Agent B" Type="CompileWin64;Win64">
+	<Agent Name="Small Agent B" Type="CompileWin64;Win64">
 		<Node Name="Read Small Chunks" Requires="#SmallChunks">
-			<Log Message="Output Files:" Files="#SmallChunks"/>
+		</Node>
+	</Agent>
+
+	<!-- Medium Chunks -->
+	<Agent Name="Medium Agent A" Type="CompileWin64;Win64">
+		<Node Name="Create Medium Chunks" Produces="#MediumChunks">
+			<RandomData Size="1048576" Count="20000" Different="true" Tag="#SmallChunks"/>
+		</Node>
+	</Agent>
+	<Agent Name="Medium Agent B" Type="CompileWin64;Win64">
+		<Node Name="Read Medium Chunks" Requires="#MediumChunks">
 		</Node>
 	</Agent>
 
 	<!-- Large Chunks -->
-	<Agent Name="Agent C" Type="CompileWin64;Win64">
+	<Agent Name="Large Agent A" Type="CompileWin64;Win64">
 		<Node Name="Create Large Chunks" Produces="#LargeChunks">
 			<RandomData Size="1073741824" Count="20" Different="true" Tag="#LargeChunks"/>
 		</Node>
 	</Agent>
-	<Agent Name="Agent D" Type="CompileWin64;Win64">
+	<Agent Name="Large Agent B" Type="CompileWin64;Win64">
 		<Node Name="Read Large Chunks" Requires="#LargeChunks">
 			<Log Message="Output Files:" Files="#LargeChunks"/>
 		</Node>
 	</Agent>
 	
 	<!-- Small Identical Chunks -->
-	<Agent Name="Agent E" Type="CompileWin64;Win64">
+	<Agent Name="Small Identical Agent A" Type="CompileWin64;Win64">
 		<Node Name="Create Small Identical Chunks" Produces="#SmallIdenticalChunks">
-			<RandomData Size="1024" Count="20000" Different="false" Tag="#SmallIdenticalChunks"/>
+			<RandomData Size="1024" Count="20000000" Different="false" Tag="#SmallIdenticalChunks"/>
 		</Node>
 	</Agent>
-	<Agent Name="Agent F" Type="CompileWin64;Win64">
+	<Agent Name="Small Identical Agent B" Type="CompileWin64;Win64">
 		<Node Name="Read Small Identical Chunks" Requires="#SmallIdenticalChunks">
-			<Log Message="Output Files:" Files="#SmallIdenticalChunks"/>
+		</Node>
+	</Agent>
+
+	<!-- Medium Identical Chunks -->
+	<Agent Name="Medium Identical Agent A" Type="CompileWin64;Win64">
+		<Node Name="Create Medium Identical Chunks" Produces="#MediumIdenticalChunks">
+			<RandomData Size="1048576" Count="20000" Different="false" Tag="#SmallChunks"/>
+		</Node>
+	</Agent>
+	<Agent Name="Medium Identical Agent B" Type="CompileWin64;Win64">
+		<Node Name="Read Medium Identical Chunks" Requires="#MediumIdenticalChunks">
 		</Node>
 	</Agent>
 
 	<!-- Large Identical Chunks -->
-	<Agent Name="Agent G" Type="CompileWin64;Win64">
+	<Agent Name="Large Identical Agent A" Type="CompileWin64;Win64">
 		<Node Name="Create Large Identical Chunks" Produces="#LargeIdenticalChunks">
 			<RandomData Size="1073741824" Count="20" Different="true" Tag="#LargeIdenticalChunks"/>
 		</Node>
 	</Agent>
-	<Agent Name="Agent H" Type="CompileWin64;Win64">
+	<Agent Name="Large Identical Agent B" Type="CompileWin64;Win64">
 		<Node Name="Read Large Identical Chunks" Requires="#LargeIdenticalChunks">
 			<Log Message="Output Files:" Files="#LargeIdenticalChunks"/>
 		</Node>
 	</Agent>
 
 	<!-- Aggregate for the whole build -->
-	<Aggregate Name="Storage Tests" Requires="Read Small Chunks;Read Large Chunks;Read Small Identical Chunks;Read Large Identical Chunks"/>
+	<Aggregate Name="Storage Tests" Requires="Read Small Chunks;Read Medium Chunks;Read Large Chunks;Read Small Identical Chunks;Read Medium Identical Chunks;Read Large Identical Chunks"/>
 
 </BuildGraph>
