Horde: Prevent illegal characters being included in Perforce client names.

diff --git a/Engine/Source/Programs/Horde/Horde.Build/Perforce/PerforceService.cs b/Engine/Source/Programs/Horde/Horde.Build/Perforce/PerforceService.cs
--- a/Engine/Source/Programs/Horde/Horde.Build/Perforce/PerforceService.cs
+++ b/Engine/Source/Programs/Horde/Horde.Build/Perforce/PerforceService.cs
@@ -630,7 +630,16 @@ namespace Horde.Build.Perforce
 				clientName += "-write";
 			}
 
-			return clientName;
+			return SanitizeClientName(clientName);
+		}
+
+		static string SanitizeClientName(string clientName)
+		{
+			string result = Regex.Replace(clientName, @"%%|\.\.\.|[#/\\@ ]", "_");
+			result = Regex.Replace(result, "_+", "_");
+			result = Regex.Replace(result, "^_", "");
+			result = Regex.Replace(result, "_$", "");
+			return result;
 		}
 
 		async ValueTask<Commit> CreateCommitInternalAsync(IStream stream, int number, string author, string description, string? basePath, DateTime dateUtc, CancellationToken cancellationToken)
