Horde: Fix invalid tag name in storage tests.

diff --git a/Engine/Source/Programs/Horde/StorageTests.xml b/Engine/Source/Programs/Horde/StorageTests.xml
--- a/Engine/Source/Programs/Horde/StorageTests.xml
+++ b/Engine/Source/Programs/Horde/StorageTests.xml
@@ -17,7 +17,7 @@
 	<!-- Medium Chunks -->
 	<Agent Name="Medium Agent A" Type="CompileWin64;Win64">
 		<Node Name="Create Medium Chunks" Produces="#MediumChunks">
-			<RandomData Size="1048576" Count="20000" Different="true" Tag="#SmallChunks"/>
+			<RandomData Size="1048576" Count="20000" Different="true" Tag="#MediumChunks"/>
 		</Node>
 	</Agent>
 	<Agent Name="Medium Agent B" Type="CompileWin64;Win64">
@@ -51,7 +51,7 @@
 	<!-- Medium Identical Chunks -->
 	<Agent Name="Medium Identical Agent A" Type="CompileWin64;Win64">
 		<Node Name="Create Medium Identical Chunks" Produces="#MediumIdenticalChunks">
-			<RandomData Size="1048576" Count="20000" Different="false" Tag="#SmallChunks"/>
+			<RandomData Size="1048576" Count="20000" Different="false" Tag="#MediumIdenticalChunks"/>
 		</Node>
 	</Agent>
 	<Agent Name="Medium Identical Agent B" Type="CompileWin64;Win64">
