Horde: Add a retry policy for all regular HTTP storage requests.

diff --git a/Engine/Source/Programs/Horde/Horde.Agent/Execution/JobExecutor.cs b/Engine/Source/Programs/Horde/Horde.Agent/Execution/JobExecutor.cs
--- a/Engine/Source/Programs/Horde/Horde.Agent/Execution/JobExecutor.cs
+++ b/Engine/Source/Programs/Horde/Horde.Agent/Execution/JobExecutor.cs
@@ -14,7 +14,6 @@ using System.Text.Json;
 using System.Text.Json.Serialization;
 using System.Threading;
 using System.Threading.Tasks;
-using System.Timers;
 using EpicGames.Core;
 using EpicGames.Horde.Storage;
 using EpicGames.Horde.Storage.Backends;
@@ -36,6 +35,8 @@ namespace Horde.Agent.Execution
 {
 	abstract class JobExecutor
 	{
+		public const string StorageHttpClientName = "Storage";
+
 		protected class ExportedNode
 		{
 			public string Name { get; set; } = String.Empty;
@@ -658,11 +659,11 @@ namespace Horde.Agent.Execution
 			// Create the storage client
 			using MemoryCache cache = new MemoryCache(new MemoryCacheOptions { });
 
-			using HttpClient httpClient = _httpClientFactory.CreateClient();
+			using HttpClient httpClient = _httpClientFactory.CreateClient(StorageHttpClientName);
 			httpClient.BaseAddress = _session.ServerUrl;
 			httpClient.DefaultRequestHeaders.Authorization = new AuthenticationHeaderValue("Bearer", _session.Token);
 
-			using HttpClient redirectHttpClient = _httpClientFactory.CreateClient();
+			using HttpClient redirectHttpClient = _httpClientFactory.CreateClient(StorageHttpClientName);
 
 			HttpStorageClient storage = new HttpStorageClient(new NamespaceId("default"), httpClient, redirectHttpClient, cache, logger);
 //			FileStorageClient storage = new FileStorageClient(DirectoryReference.Combine(sharedStorageDir, "bundles"), cache, logger);
diff --git a/Engine/Source/Programs/Horde/Horde.Agent/Program.cs b/Engine/Source/Programs/Horde/Horde.Agent/Program.cs
--- a/Engine/Source/Programs/Horde/Horde.Agent/Program.cs
+++ b/Engine/Source/Programs/Horde/Horde.Agent/Program.cs
@@ -171,6 +171,12 @@ namespace Horde.Agent
 				return builder.WaitAndRetryAsync(new[] { TimeSpan.FromSeconds(1), TimeSpan.FromSeconds(5), TimeSpan.FromSeconds(10) });
 			});
 
+			services.AddHttpClient(JobExecutor.StorageHttpClientName)
+				.AddTransientHttpErrorPolicy(builder =>
+				{
+					return builder.WaitAndRetryAsync(new[] { TimeSpan.FromSeconds(1), TimeSpan.FromSeconds(5), TimeSpan.FromSeconds(10) });
+				});
+
 			services.AddHordeStorage(settings => configSection.GetCurrentServerProfile().GetSection(nameof(serverProfile.Storage)).Bind(settings));
 
 			services.AddSingleton<GrpcService>();
