Horde: Improve logging for pool sizing

diff --git a/Engine/Source/Programs/Horde/Horde.Build/Agents/Fleet/FleetService.cs b/Engine/Source/Programs/Horde/Horde.Build/Agents/Fleet/FleetService.cs
--- a/Engine/Source/Programs/Horde/Horde.Build/Agents/Fleet/FleetService.cs
+++ b/Engine/Source/Programs/Horde/Horde.Build/Agents/Fleet/FleetService.cs
@@ -180,15 +180,17 @@ namespace Horde.Build.Agents.Fleet
 			int deltaAgentCount = desiredAgentCount - currentAgentCount;
 
 			IFleetManager fleetManager = CreateFleetManager(poolSizeResult.Pool);
-
-			Dictionary<string, object> logScope = new() { ["FleetManager"] = fleetManager.GetType().Name };
-			logScope["PoolSizeStrategy"] = poolSizeResult.Status ?? new Dictionary<string, object>();
-
-			using (_logger.BeginScope(logScope))
+			Dictionary<string, object> logScopeMetadata = new()
 			{
-				_logger.LogInformation("{PoolName} Current={Current} Target={Target} Delta={Delta}",
-					pool.Name, currentAgentCount, desiredAgentCount, deltaAgentCount);
-			}
+				["FleetManager"] = fleetManager.GetType().Name,
+				["PoolSizeStrategy"] = poolSizeResult.Status ?? new Dictionary<string, object>(),
+				["PoolId"] = poolSizeResult.Pool.Id,
+			};
+
+			using IDisposable logScope = _logger.BeginScope(logScopeMetadata);
+			
+			_logger.LogInformation("{PoolName} Current={Current} Target={Target} Delta={Delta}",
+				pool.Name, currentAgentCount, desiredAgentCount, deltaAgentCount);
 
 			try
 			{
diff --git a/Engine/Source/Programs/Horde/Horde.Build/Agents/Fleet/Providers/AwsRecyclingFleetManager.cs b/Engine/Source/Programs/Horde/Horde.Build/Agents/Fleet/Providers/AwsRecyclingFleetManager.cs
--- a/Engine/Source/Programs/Horde/Horde.Build/Agents/Fleet/Providers/AwsRecyclingFleetManager.cs
+++ b/Engine/Source/Programs/Horde/Horde.Build/Agents/Fleet/Providers/AwsRecyclingFleetManager.cs
@@ -111,8 +111,11 @@ public sealed class AwsRecyclingFleetManager : IFleetManager
 		Dictionary<string, List<Instance>> stoppedInstancesPerAz = GetInstancesToLaunch(candidatesPerAz, requestCountPerAz);
 		List<InstanceType>? instanceTypePriority = Settings.InstanceTypes?.Select(InstanceType.FindValue).ToList();
 
-		var info = new { MissingInstanceCount = remainingInstancesCount, InstancesToStartCount = requestCountPerAz.Values.Sum(x => x)};
-		_logger.LogWarning("Not enough stopped instances to accommodate the full pool scale-out. {@info}", info);
+		if (remainingInstancesCount > 0)
+		{
+			var info = new { StoppedInstancesMissingCount = remainingInstancesCount, InstancesToStartCount = requestCountPerAz.Values.Sum(x => x)};
+			_logger.LogWarning("Not enough stopped instances to accommodate the full pool scale-out. {@info}", info);			
+		}
 
 		foreach ((string az, List<Instance> instances) in stoppedInstancesPerAz)
 		{
