Horde: Add more logging metadata to AWS fleet manager

diff --git a/Engine/Source/Programs/Horde/Horde.Build/Agents/Fleet/Providers/AwsRecyclingFleetManager.cs b/Engine/Source/Programs/Horde/Horde.Build/Agents/Fleet/Providers/AwsRecyclingFleetManager.cs
--- a/Engine/Source/Programs/Horde/Horde.Build/Agents/Fleet/Providers/AwsRecyclingFleetManager.cs
+++ b/Engine/Source/Programs/Horde/Horde.Build/Agents/Fleet/Providers/AwsRecyclingFleetManager.cs
@@ -9,6 +9,7 @@ using System.Threading.Tasks;
 using Amazon.EC2;
 using Amazon.EC2.Model;
 using Amazon.Runtime;
+using EpicGames.Core;
 using Horde.Build.Agents.Pools;
 using Microsoft.Extensions.Logging;
 using OpenTracing;
@@ -107,19 +108,32 @@ public sealed class AwsRecyclingFleetManager : IFleetManager
 
 		Dictionary<string, List<Instance>> candidatesPerAz = await GetCandidateInstancesAsync(pool, cancellationToken);
 		Dictionary<string, int> instanceCountPerAz = candidatesPerAz.ToDictionary(x => x.Key, x => x.Value.Count);
-		Dictionary<string, int> requestCountPerAz = DistributeInstanceRequestsOverAzs(requestedInstancesCount, instanceCountPerAz, out int remainingInstancesCount);
+		Dictionary<string, int> requestCountPerAz = DistributeInstanceRequestsOverAzs(requestedInstancesCount, instanceCountPerAz, out int stoppedInstancesMissingCount);
 		Dictionary<string, List<Instance>> stoppedInstancesPerAz = GetInstancesToLaunch(candidatesPerAz, requestCountPerAz);
 		List<InstanceType>? instanceTypePriority = Settings.InstanceTypes?.Select(InstanceType.FindValue).ToList();
+		int instancesToStartCount = requestCountPerAz.Values.Sum(x => x);
 
-		if (remainingInstancesCount > 0)
+		if (stoppedInstancesMissingCount > 0)
 		{
-			var info = new { StoppedInstancesMissingCount = remainingInstancesCount, InstancesToStartCount = requestCountPerAz.Values.Sum(x => x)};
-			_logger.LogWarning("Not enough stopped instances to accommodate the full pool scale-out. {@info}", info);			
+			var incompleteData = new { StoppedInstancesMissingCount = stoppedInstancesMissingCount, InstancesToStartCount = instancesToStartCount};
+			_logger.LogWarning("Not enough stopped instances to accommodate the full pool scale-out. {@Data}", incompleteData);			
 		}
 
+		var data = new
+		{
+			RequestedInstancesCount = requestedInstancesCount,
+			CurrentAgentCount = agents.Count,
+			StoppedInstancesMissingCount = stoppedInstancesMissingCount,
+			CandidatesPerAz = candidatesPerAz.ToDictionary(x => x.Key, x => x.Value.Select(y => y.InstanceId)),
+			RequestCountPerAz = requestCountPerAz,
+			StoppedInstancesPerAz = stoppedInstancesPerAz.ToDictionary(x => x.Key, x => x.Value.Select(y => y.InstanceId)),
+			InstanceTypePriority = Settings.InstanceTypes,
+		};
+		_logger.LogInformation("Starting {InstancesToStartCount} instance(s) {@Data}", instancesToStartCount, data);
+
 		foreach ((string az, List<Instance> instances) in stoppedInstancesPerAz)
 		{
-			using (_logger.BeginScope(new Dictionary<string, object> { ["Az"] = az }))
+			using (_logger.WithProperty("Az", az).BeginScope())
 			{
 				_logger.LogDebug("Trying to start {InstanceCount} instance(s)", instances.Count);
 				await StartInstancesWithRetriesAsync(instances, instanceTypePriority, cancellationToken);					
