Horde: Reduce and consolidate logging in AWS fleet manager

diff --git a/Engine/Source/Programs/Horde/Horde.Build/Agents/Fleet/Providers/AwsRecyclingFleetManager.cs b/Engine/Source/Programs/Horde/Horde.Build/Agents/Fleet/Providers/AwsRecyclingFleetManager.cs
--- a/Engine/Source/Programs/Horde/Horde.Build/Agents/Fleet/Providers/AwsRecyclingFleetManager.cs
+++ b/Engine/Source/Programs/Horde/Horde.Build/Agents/Fleet/Providers/AwsRecyclingFleetManager.cs
@@ -112,24 +112,34 @@ public sealed class AwsRecyclingFleetManager : IFleetManager
 		Dictionary<string, List<Instance>> stoppedInstancesPerAz = GetInstancesToLaunch(candidatesPerAz, requestCountPerAz);
 		List<InstanceType>? instanceTypePriority = Settings.InstanceTypes?.Select(InstanceType.FindValue).ToList();
 		int instancesToStartCount = requestCountPerAz.Values.Sum(x => x);
-
-		if (stoppedInstancesMissingCount > 0)
+		
+		Dictionary<string, object?> logScopeMetadata = new()
 		{
-			var incompleteData = new { StoppedInstancesMissingCount = stoppedInstancesMissingCount, InstancesToStartCount = instancesToStartCount};
-			_logger.LogWarning("Not enough stopped instances to accommodate the full pool scale-out. {@Data}", incompleteData);			
-		}
+			["InstancesToStartCount"] = instancesToStartCount,
+			["RequestedInstancesCount"] = requestedInstancesCount,
+			["CurrentAgentCount"] = agents.Count,
+			["StoppedInstancesMissingCount"] = stoppedInstancesMissingCount,
+			["CandidatesPerAz"] = candidatesPerAz.ToDictionary(x => x.Key, x => x.Value.Select(y => y.InstanceId)),
+			["RequestCountPerAz"] = requestCountPerAz,
+			["StoppedInstancesPerAz"] = stoppedInstancesPerAz.ToDictionary(x => x.Key, x => x.Value.Select(y => y.InstanceId)),
+			["InstanceTypePriority"] = Settings.InstanceTypes,
+		};
 
-		var data = new
+		using (_logger.BeginScope(logScopeMetadata))
 		{
-			RequestedInstancesCount = requestedInstancesCount,
-			CurrentAgentCount = agents.Count,
-			StoppedInstancesMissingCount = stoppedInstancesMissingCount,
-			CandidatesPerAz = candidatesPerAz.ToDictionary(x => x.Key, x => x.Value.Select(y => y.InstanceId)),
-			RequestCountPerAz = requestCountPerAz,
-			StoppedInstancesPerAz = stoppedInstancesPerAz.ToDictionary(x => x.Key, x => x.Value.Select(y => y.InstanceId)),
-			InstanceTypePriority = Settings.InstanceTypes,
-		};
-		_logger.LogInformation("Starting {InstancesToStartCount} instance(s) {@Data}", instancesToStartCount, data);
+			if (instancesToStartCount == 0)
+			{
+				_logger.LogWarning("Unable to start any instance(s)");
+			}
+			else if (stoppedInstancesMissingCount > 0)
+			{
+				_logger.LogWarning("Starting {InstancesToStartCount} instance(s) but not enough stopped instances to accommodate the full pool scale-out", instancesToStartCount);
+			}
+			else
+			{
+				_logger.LogInformation("Starting {InstancesToStartCount} instance(s)", instancesToStartCount);	
+			}
+		}
 
 		foreach ((string az, List<Instance> instances) in stoppedInstancesPerAz)
 		{
@@ -348,18 +358,6 @@ public sealed class AwsRecyclingFleetManager : IFleetManager
 		List<IGrouping<string, Instance>> azInstancePairs = response.Reservations
 			.SelectMany(x => x.Instances)
 			.GroupBy(x => x.Placement.AvailabilityZone).ToList();
-
-		Dictionary<string, object> logScope = new()
-		{
-			["InstanceIds"] = azInstancePairs.ToDictionary(
-				x => x.Key,
-				y => y.Select(x => x.InstanceId).ToList()
-			)
-		};
-		using (_logger.BeginScope(logScope))
-		{
-			_logger.LogDebug("Found {NumCandidates} candidate instances", azInstancePairs.Count);
-		}
 		
 		return azInstancePairs.ToDictionary(x => x.Key, y => y.ToList());;
 	}
