Horde: Clarify warning message about missing stopped instances

diff --git a/Engine/Source/Programs/Horde/Horde.Build/Agents/Fleet/Providers/AwsRecyclingFleetManager.cs b/Engine/Source/Programs/Horde/Horde.Build/Agents/Fleet/Providers/AwsRecyclingFleetManager.cs
--- a/Engine/Source/Programs/Horde/Horde.Build/Agents/Fleet/Providers/AwsRecyclingFleetManager.cs
+++ b/Engine/Source/Programs/Horde/Horde.Build/Agents/Fleet/Providers/AwsRecyclingFleetManager.cs
@@ -129,7 +129,7 @@ public sealed class AwsRecyclingFleetManager : IFleetManager
 		{
 			if (instancesToStartCount == 0)
 			{
-				_logger.LogWarning("Unable to start any instance(s)");
+				_logger.LogWarning("Unable to start any instance(s). No stopped instances are available.");
 			}
 			else if (stoppedInstancesMissingCount > 0)
 			{
