Horde: Ensure AWS EC2 exceptions are rethrown

diff --git a/Engine/Source/Programs/Horde/Horde.Build.Tests/Fleet/AwsRecyclingFleetManagerTest.cs b/Engine/Source/Programs/Horde/Horde.Build.Tests/Fleet/AwsRecyclingFleetManagerTest.cs
--- a/Engine/Source/Programs/Horde/Horde.Build.Tests/Fleet/AwsRecyclingFleetManagerTest.cs
+++ b/Engine/Source/Programs/Horde/Horde.Build.Tests/Fleet/AwsRecyclingFleetManagerTest.cs
@@ -20,6 +20,16 @@ namespace Horde.Build.Tests.Fleet
 	[TestClass]
 	public class AwsRecyclingFleetManagerTest : TestSetup
 	{
+		[TestMethod]
+		public async Task ExpandAmazonEc2ExceptionsArePropagated()
+		{
+			FakeAmazonEc2 ec2 = new ();
+			ec2.AddInstance(FakeAmazonEc2.StateStopped, InstanceType.M5Large, FakeAmazonEc2.AzUsEast1A);
+			ec2.SetCapacity(FakeAmazonEc2.AzUsEast1A, InstanceType.M5Large, 1);
+			ec2.OnStartInstances = (_, _) => throw new AmazonEC2Exception("Something bad happened");
+			await Assert.ThrowsExceptionAsync<AmazonEC2Exception>(() => ExpandPoolAsync(ec2.Get(), 1, new()));
+		}
+		
 		[TestMethod]
 		public async Task ExpandOneAzPartialCapacityWithFallback()
 		{
diff --git a/Engine/Source/Programs/Horde/Horde.Build.Tests/Fleet/FakeAmazonEc2.cs b/Engine/Source/Programs/Horde/Horde.Build.Tests/Fleet/FakeAmazonEc2.cs
--- a/Engine/Source/Programs/Horde/Horde.Build.Tests/Fleet/FakeAmazonEc2.cs
+++ b/Engine/Source/Programs/Horde/Horde.Build.Tests/Fleet/FakeAmazonEc2.cs
@@ -27,6 +27,11 @@ public class FakeAmazonEc2
 	private readonly Mock<IAmazonEC2> _mock;
 	private readonly Dictionary<string, Instance> _instances = new();
 	private int _instanceIdCounter;
+
+	/// <summary>
+	/// Callback when StartInstances method is called. Can be used for mocking exception throwing.
+	/// </summary>
+	internal Action<StartInstancesRequest, CancellationToken> OnStartInstances { get; set; } = (_, _) => { };
 	
 	internal IReadOnlyDictionary<string, Instance> Instances => _instances;
 
@@ -118,6 +123,8 @@ public class FakeAmazonEc2
 				throw new AmazonEC2Exception("Insufficient capacity.");
 			}
 
+			OnStartInstances(request, cancellationToken);
+
 			azInstanceTypeCapacities[az][instanceType] = capacity - 1;
 		}
 		
diff --git a/Engine/Source/Programs/Horde/Horde.Build/Agents/Fleet/Providers/AwsRecyclingFleetManager.cs b/Engine/Source/Programs/Horde/Horde.Build/Agents/Fleet/Providers/AwsRecyclingFleetManager.cs
--- a/Engine/Source/Programs/Horde/Horde.Build/Agents/Fleet/Providers/AwsRecyclingFleetManager.cs
+++ b/Engine/Source/Programs/Horde/Horde.Build/Agents/Fleet/Providers/AwsRecyclingFleetManager.cs
@@ -283,6 +283,8 @@ public sealed class AwsRecyclingFleetManager : IFleetManager
 			{
 				throw new AwsInsufficientCapacityException();
 			}
+
+			throw;
 		}
 
 		if (!IsResponseSuccessful(startResponse))
