Horde: Fix bug where start instances not breaking loop on success in AWS fleet manager

And add more logging.

diff --git a/Engine/Source/Programs/Horde/Horde.Build/Agents/Fleet/Providers/AwsRecyclingFleetManager.cs b/Engine/Source/Programs/Horde/Horde.Build/Agents/Fleet/Providers/AwsRecyclingFleetManager.cs
--- a/Engine/Source/Programs/Horde/Horde.Build/Agents/Fleet/Providers/AwsRecyclingFleetManager.cs
+++ b/Engine/Source/Programs/Horde/Horde.Build/Agents/Fleet/Providers/AwsRecyclingFleetManager.cs
@@ -53,7 +53,13 @@ public class AwsFleetManagerException : Exception
 /// <summary>
 /// Exception when AWS API returns insufficient capacity
 /// </summary>
-public class AwsInsufficientCapacityException : AwsFleetManagerException { }
+public class AwsInsufficientCapacityException : AwsFleetManagerException
+{
+	/// <inheritdoc />
+	public AwsInsufficientCapacityException(string message) : base(message)
+	{
+	}
+}
 
 /// <summary>
 /// Generic HTTP response error from AWS API
@@ -221,6 +227,8 @@ public sealed class AwsRecyclingFleetManager : IFleetManager
 	{
 		if (instances.Count == 0) return;
 		
+		_logger.LogInformation("StartInstancesWithRetriesAsync()");
+		
 		List<InstanceType?> instanceTypes = new();
 		if (instanceTypePriority == null || instanceTypePriority.Count == 0)
 		{
@@ -244,10 +252,11 @@ public sealed class AwsRecyclingFleetManager : IFleetManager
 			{
 				await StartInstancesAsync(instances, instanceType, cancellationToken);
 				success = true;
+				break;
 			}
 			catch (AwsInsufficientCapacityException)
 			{
-				_logger.LogDebug("Insufficient capacity for {InstanceType}", instanceType);
+				_logger.LogInformation("Insufficient capacity for {InstanceType}", instanceType);
 			}
 		}
 
@@ -260,6 +269,7 @@ public sealed class AwsRecyclingFleetManager : IFleetManager
 	internal async Task StartInstancesAsync(List<Instance> instances, InstanceType? instanceType, CancellationToken cancellationToken)
 	{
 		if (instances.Count == 0) return;
+		_logger.LogInformation("StartInstancesAsync()");
 		if (instanceType != null)
 		{
 			await ChangeInstanceTypeAsync(instances, instanceType, cancellationToken);
@@ -281,7 +291,7 @@ public sealed class AwsRecyclingFleetManager : IFleetManager
 		{
 			if (e.Message.Contains("insufficient capacity", StringComparison.OrdinalIgnoreCase))
 			{
-				throw new AwsInsufficientCapacityException();
+				throw new AwsInsufficientCapacityException($"Not enough capacity in AZ when using instance type {instanceType}");
 			}
 
 			throw;
@@ -314,6 +324,7 @@ public sealed class AwsRecyclingFleetManager : IFleetManager
 	private async Task ChangeInstanceTypeAsync(List<Instance> instances, InstanceType newInstanceType, CancellationToken cancellationToken)
 	{
 		using IScope scope = GlobalTracer.Instance.BuildSpan("ChangeInstanceType").StartActive();
+		_logger.LogInformation("ChangeInstanceTypeAsync()");
 		
 		foreach (Instance instance in instances)
 		{
