Horde: Make pool resizing spans the entry/top-level span

The parallelism made traces hard to read in previous setup. Letting each span for pool resizing be the top-most makes them easier to read as there's one for each pool.

diff --git a/Engine/Source/Programs/Horde/Horde.Build/Agents/Fleet/FleetService.cs b/Engine/Source/Programs/Horde/Horde.Build/Agents/Fleet/FleetService.cs
--- a/Engine/Source/Programs/Horde/Horde.Build/Agents/Fleet/FleetService.cs
+++ b/Engine/Source/Programs/Horde/Horde.Build/Agents/Fleet/FleetService.cs
@@ -111,32 +111,36 @@ namespace Horde.Build.Agents.Fleet
 
 		internal async ValueTask TickLeaderAsync(CancellationToken stoppingToken)
 		{
-			Stopwatch stopwatch = Stopwatch.StartNew();
-			using IScope _ = GlobalTracer.Instance.BuildSpan("AutoscaleService.TickAsync").StartActive();
-			
-			List<PoolSizeResult> input = await GetPoolSizeDataAsync();
-			List<PoolSizeResult> output = await CalculatePoolSizesAsync(input, stoppingToken);
-			await ResizePoolsAsync(output, stoppingToken);
-
-			stopwatch.Stop();
-			_logger.LogDebug("Autoscaling pools took {ElapsedTime} ms", stopwatch.ElapsedMilliseconds);
+			ISpan span = GlobalTracer.Instance.BuildSpan("FleetService.Tick").Start();
+			try
+			{
+				List<PoolSizeResult> input = await GetPoolSizeDataAsync();
+				List<PoolSizeResult> output = await CalculatePoolSizesAsync(input, stoppingToken);
+				await ResizePoolsAsync(output, stoppingToken);
+			}
+			finally
+			{
+				span.Finish();
+			}
 		}
 
 		internal async ValueTask TickHighFrequencyAsync(CancellationToken stoppingToken)
 		{
-			Stopwatch stopwatch = Stopwatch.StartNew();
-			using IScope _ = GlobalTracer.Instance.BuildSpan("AutoscaleService.TickHighFrequency").StartActive();
-			
-			// TODO: Re-enable high frequency scaling (only used for experimental scaling of remote execution agents at the moment)
-			await Task.Delay(0, stoppingToken);
-
-			stopwatch.Stop();
-			_logger.LogDebug("Autoscaling pools (high frequency) took {ElapsedTime} ms", stopwatch.ElapsedMilliseconds);
+			ISpan span = GlobalTracer.Instance.BuildSpan("FleetService.TickHighFrequency").Start();
+			try
+			{
+				// TODO: Re-enable high frequency scaling (only used for experimental scaling of remote execution agents at the moment)
+				await Task.Delay(0, stoppingToken);
+			}
+			finally
+			{
+				span.Finish();
+			}
 		}
 
 		internal async Task<List<PoolSizeResult>> CalculatePoolSizesAsync(List<PoolSizeResult> poolSizeDatas, CancellationToken cancellationToken)
 		{
-			using IScope _ = GlobalTracer.Instance.BuildSpan("AutoscaleService.CalculatePoolSizesAsync").StartActive();
+			using IScope _ = GlobalTracer.Instance.BuildSpan("FleetService.CalculatePoolSizes").StartActive();
 			ConcurrentQueue<PoolSizeResult> results = new ();
 			ParallelOptions options = new () { MaxDegreeOfParallelism = MaxParallelTasks, CancellationToken = cancellationToken };
 
@@ -179,6 +183,14 @@ namespace Horde.Build.Agents.Fleet
 			int desiredAgentCount = poolSizeResult.DesiredAgentCount.Value;
 			int deltaAgentCount = desiredAgentCount - currentAgentCount;
 
+			using IScope scope = GlobalTracer.Instance
+				.BuildSpan("FleetService.ResizePool")
+				.WithTag(Datadog.Trace.OpenTracing.DatadogTags.ResourceName, pool.Id.ToString())
+				.WithTag("CurrentAgentCount", currentAgentCount)
+				.WithTag("DesiredAgentCount", desiredAgentCount)
+				.WithTag("DeltaAgentCount", deltaAgentCount)
+				.StartActive();
+			
 			IFleetManager fleetManager = CreateFleetManager(poolSizeResult.Pool);
 			Dictionary<string, object> logScopeMetadata = new()
 			{
@@ -194,12 +206,6 @@ namespace Horde.Build.Agents.Fleet
 
 			try
 			{
-				using IScope scope = GlobalTracer.Instance.BuildSpan("ScalingPool").StartActive();
-				scope.Span.SetTag("poolName", pool.Name);
-				scope.Span.SetTag("current", currentAgentCount);
-				scope.Span.SetTag("desired", desiredAgentCount);
-				scope.Span.SetTag("delta", deltaAgentCount);
-
 				if (deltaAgentCount > 0)
 				{
 					TimeSpan scaleOutCooldown = pool.ScaleOutCooldown ?? _defaultScaleOutCooldown;
@@ -246,8 +252,6 @@ namespace Horde.Build.Agents.Fleet
 		
 		internal async Task ResizePoolsAsync(List<PoolSizeResult> poolSizeDatas, CancellationToken cancellationToken = default)
 		{
-			using IScope _ = GlobalTracer.Instance.BuildSpan("AutoscaleService.ResizePoolsAsync").StartActive();
-			
 			ParallelOptions options = new () { MaxDegreeOfParallelism = MaxParallelTasks, CancellationToken = cancellationToken };
 			await Parallel.ForEachAsync(poolSizeDatas.OrderByDescending(x => x.Agents.Count), options, async (poolSizeData, innerCt) =>
 			{
