Horde: Add more logging and spans to AwsRecyclingFleetManager

diff --git a/Engine/Source/Programs/Horde/Horde.Build/Agents/Fleet/Providers/AwsRecyclingFleetManager.cs b/Engine/Source/Programs/Horde/Horde.Build/Agents/Fleet/Providers/AwsRecyclingFleetManager.cs
--- a/Engine/Source/Programs/Horde/Horde.Build/Agents/Fleet/Providers/AwsRecyclingFleetManager.cs
+++ b/Engine/Source/Programs/Horde/Horde.Build/Agents/Fleet/Providers/AwsRecyclingFleetManager.cs
@@ -106,9 +106,10 @@ public sealed class AwsRecyclingFleetManager : IFleetManager
 	public async Task ExpandPoolAsync(IPool pool, IReadOnlyList<IAgent> agents, int requestedInstancesCount, CancellationToken cancellationToken)
 	{
 		using IScope scope = GlobalTracer.Instance.BuildSpan("ExpandPool").StartActive();
-		scope.Span.SetTag("poolName", pool.Name);
-		scope.Span.SetTag("numAgents", agents.Count);
-		scope.Span.SetTag("count", requestedInstancesCount);
+		scope.Span.SetTag("PoolId", pool.Id.ToString());
+		scope.Span.SetTag("PoolName", pool.Name);
+		scope.Span.SetTag("NumAgents", agents.Count);
+		scope.Span.SetTag("Count", requestedInstancesCount);
 
 		using IDisposable logScope = _logger.BeginScope(new Dictionary<string, object> { ["PoolId"] = pool.Id });
 
@@ -225,9 +226,14 @@ public sealed class AwsRecyclingFleetManager : IFleetManager
 	/// <param name="cancellationToken">Cancellation token for the call</param>
 	private async Task StartInstancesWithRetriesAsync(List<Instance> instances, IReadOnlyList<InstanceType>? instanceTypePriority, CancellationToken cancellationToken)
 	{
-		if (instances.Count == 0) return;
+		List<string> instanceIds = instances.Select(x => x.InstanceId).ToList();
+		string instanceTypePriorityString = instanceTypePriority == null ? "" : String.Join(',', instanceTypePriority.Select(x => x.Value));
+		using IScope scope = GlobalTracer.Instance.BuildSpan("StartInstancesWithRetries")
+			.WithTag("Instances", String.Join(',', instanceIds))
+			.WithTag("InstanceTypePriority", instanceTypePriorityString)
+			.StartActive();
 		
-		_logger.LogInformation("StartInstancesWithRetriesAsync()");
+		if (instances.Count == 0) return;
 		
 		List<InstanceType?> instanceTypes = new();
 		if (instanceTypePriority == null || instanceTypePriority.Count == 0)
@@ -241,7 +247,7 @@ public sealed class AwsRecyclingFleetManager : IFleetManager
 		
 		using IDisposable logScope = _logger.BeginScope(new Dictionary<string, object>
 		{
-			["InstanceIds"] = instances.Select(x => x.InstanceId),
+			["InstanceIds"] = instanceIds,
 			["InstanceTypes"] = instanceTypes
 		});
 
@@ -260,6 +266,7 @@ public sealed class AwsRecyclingFleetManager : IFleetManager
 			}
 		}
 
+		scope.Span.SetTag("Success", success);
 		if (!success)
 		{
 			_logger.LogError("Unable to start instances. Insufficient capacity for all instance types tried.");
@@ -268,6 +275,12 @@ public sealed class AwsRecyclingFleetManager : IFleetManager
 
 	internal async Task StartInstancesAsync(List<Instance> instances, InstanceType? instanceType, CancellationToken cancellationToken)
 	{
+		List<string> instanceIds = instances.Select(x => x.InstanceId).ToList();
+		using IScope scope = GlobalTracer.Instance.BuildSpan("StartInstances")
+			.WithTag("Instances", String.Join(',', instanceIds))
+			.WithTag("InstanceType", instanceType)
+			.StartActive();
+		
 		if (instances.Count == 0) return;
 		if (instanceType != null)
 		{
@@ -275,17 +288,16 @@ public sealed class AwsRecyclingFleetManager : IFleetManager
 		}
 
 		StartInstancesRequest startRequest = new ();
-		startRequest.InstanceIds.AddRange(instances.Select(x => x.InstanceId));
+		startRequest.InstanceIds.AddRange(instanceIds);
 
-		StartInstancesResponse startResponse = new();
+		StartInstancesResponse startResponse;
 		try
 		{
-			_logger.LogInformation("StartInstancesAsync()");
-			using IScope scope = GlobalTracer.Instance.BuildSpan("StartInstances").StartActive();
-			scope.Span.SetTag("req.instanceIds", String.Join(",", startRequest.InstanceIds));
+			using IScope ec2Scope = GlobalTracer.Instance.BuildSpan("StartInstancesEc2").StartActive();
+			ec2Scope.Span.SetTag("Request.InstanceIds", String.Join(",", startRequest.InstanceIds));
 			startResponse = await _ec2.StartInstancesAsync(startRequest, cancellationToken);
-			scope.Span.SetTag("res.statusCode", (int)startResponse.HttpStatusCode);
-			scope.Span.SetTag("res.numInstances", startResponse.StartingInstances.Count);
+			ec2Scope.Span.SetTag("Response.StatusCode", (int)startResponse.HttpStatusCode);
+			ec2Scope.Span.SetTag("Response.NumInstances", startResponse.StartingInstances.Count);
 		}
 		catch (AmazonEC2Exception e)
 		{
@@ -323,19 +335,23 @@ public sealed class AwsRecyclingFleetManager : IFleetManager
 	/// <returns>List of instances grouped by availability zone</returns>
 	private async Task ChangeInstanceTypeAsync(List<Instance> instances, InstanceType newInstanceType, CancellationToken cancellationToken)
 	{
-		using IScope scope = GlobalTracer.Instance.BuildSpan("ChangeInstanceType").StartActive();
-		_logger.LogInformation("ChangeInstanceTypeAsync()");
-		
 		foreach (Instance instance in instances)
 		{
+			using IScope scope = GlobalTracer.Instance.BuildSpan("ChangeInstanceType")
+				.WithTag("FromInstanceType", instance.InstanceType)
+				.WithTag("ToInstanceType", newInstanceType)
+				.StartActive();
+			
 			if (instance.InstanceType == newInstanceType) { continue; }
-			using IScope modifyScope = GlobalTracer.Instance.BuildSpan("ModifyInstanceAttribute").StartActive();
 
 			ModifyInstanceAttributeRequest request = new () { InstanceId = instance.InstanceId, InstanceType = newInstanceType };
 			ModifyInstanceAttributeResponse response = await _ec2.ModifyInstanceAttributeAsync(request, cancellationToken);
+			scope.Span.SetTag("Response.StatusCode", (int)response.HttpStatusCode);
+			
 			if (IsResponseSuccessful(response))
 			{
 				_logger.LogInformation("Instance type for {InstanceId} modified ({PrevInstanceType} -> {CurrentInstanceType})", instance.InstanceId, instance.InstanceType, newInstanceType);
+				instance.InstanceType = newInstanceType; // Update local copy with instance type for sake of logging
 			}
 			else
 			{
@@ -352,7 +368,8 @@ public sealed class AwsRecyclingFleetManager : IFleetManager
 	/// <returns>List of instances grouped by availability zone</returns>
 	private async Task<Dictionary<string, List<Instance>>> GetCandidateInstancesAsync(IPool pool, CancellationToken cancellationToken)
 	{
-		using IScope scope = GlobalTracer.Instance.BuildSpan("DescribeInstances").StartActive();
+		using IScope scope = GlobalTracer.Instance.BuildSpan("GetCandidateInstances")
+			.WithTag("PoolId", pool.Id.ToString()).StartActive();
 		
 		// Find stopped instances in the correct pool
 		DescribeInstancesRequest request = new()
