Horde: Lower tick interval for fleet service (auto-scaling) and add more trace spans for strategies

diff --git a/Engine/Source/Programs/Horde/Horde.Build/Agents/Fleet/FleetService.cs b/Engine/Source/Programs/Horde/Horde.Build/Agents/Fleet/FleetService.cs
--- a/Engine/Source/Programs/Horde/Horde.Build/Agents/Fleet/FleetService.cs
+++ b/Engine/Source/Programs/Horde/Horde.Build/Agents/Fleet/FleetService.cs
@@ -82,7 +82,7 @@ namespace Horde.Build.Agents.Fleet
 			_clock = clock;
 			_cache = cache;
 			_logger = logger;
-			_ticker = clock.AddSharedTicker<FleetService>(TimeSpan.FromMinutes(5.0), TickLeaderAsync, _logger);
+			_ticker = clock.AddSharedTicker<FleetService>(TimeSpan.FromMinutes(3.0), TickLeaderAsync, _logger);
 			_tickerHighFrequency = clock.AddSharedTicker("FleetService.TickHighFrequency", TimeSpan.FromSeconds(30), TickHighFrequencyAsync, _logger);
 			_settings = settings;
 			_defaultScaleOutCooldown = TimeSpan.FromSeconds(settings.Value.AgentPoolScaleOutCooldownSeconds);
@@ -140,23 +140,30 @@ namespace Horde.Build.Agents.Fleet
 
 		internal async Task<List<PoolSizeResult>> CalculatePoolSizesAsync(List<PoolSizeResult> poolSizeDatas, CancellationToken cancellationToken)
 		{
-			using IScope _ = GlobalTracer.Instance.BuildSpan("FleetService.CalculatePoolSizes").StartActive();
+			ISpan span = GlobalTracer.Instance.BuildSpan("FleetService.CalculatePoolSizes").WithTag("MaxParallelTasks", MaxParallelTasks).Start();
 			ConcurrentQueue<PoolSizeResult> results = new ();
-			ParallelOptions options = new () { MaxDegreeOfParallelism = MaxParallelTasks, CancellationToken = cancellationToken };
-
-			await Parallel.ForEachAsync(poolSizeDatas, options, async (input, innerCt) =>
+			
+			try
 			{
-				try
+				ParallelOptions options = new () { MaxDegreeOfParallelism = MaxParallelTasks, CancellationToken = cancellationToken };
+				await Parallel.ForEachAsync(poolSizeDatas, options, async (input, innerCt) =>
 				{
-					IPoolSizeStrategy sizeStrategy = CreatePoolSizeStrategy(input.Pool);
-					PoolSizeResult output = await sizeStrategy.CalculatePoolSizeAsync(input.Pool, input.Agents);
-					results.Enqueue(output);
-				}
-				catch (Exception e)
-				{
-					_logger.LogError(e, "Failed calculating pool size for pool {PoolId}", input.Pool?.Id);
-				}
-			});
+					try
+					{
+						IPoolSizeStrategy sizeStrategy = CreatePoolSizeStrategy(input.Pool);
+						PoolSizeResult output = await sizeStrategy.CalculatePoolSizeAsync(input.Pool, input.Agents);
+						results.Enqueue(output);
+					}
+					catch (Exception e)
+					{
+						_logger.LogError(e, "Failed calculating pool size for pool {PoolId}", input.Pool?.Id);
+					}
+				});
+			}
+			finally
+			{
+				span.Finish();
+			}
 
 			return results.ToList();
 		}
diff --git a/Engine/Source/Programs/Horde/Horde.Build/Agents/Fleet/JobQueueStrategy.cs b/Engine/Source/Programs/Horde/Horde.Build/Agents/Fleet/JobQueueStrategy.cs
--- a/Engine/Source/Programs/Horde/Horde.Build/Agents/Fleet/JobQueueStrategy.cs
+++ b/Engine/Source/Programs/Horde/Horde.Build/Agents/Fleet/JobQueueStrategy.cs
@@ -13,6 +13,8 @@ using Horde.Build.Streams;
 using Horde.Build.Utilities;
 using HordeCommon;
 using Microsoft.Extensions.Caching.Memory;
+using OpenTracing;
+using OpenTracing.Util;
 
 namespace Horde.Build.Agents.Fleet
 {
@@ -168,6 +170,8 @@ namespace Horde.Build.Agents.Fleet
 
 		internal async Task<Dictionary<PoolId, int>> GetPoolQueueSizesAsync(DateTimeOffset jobsCreatedAfter)
 		{
+			using IScope scope = GlobalTracer.Instance.BuildSpan("JobQueueStrategy.GetPoolQueueSizes").StartActive();
+			
 			List<IStream> streamsList = await _streamService.GetStreamsAsync();
 			Dictionary<StreamId, IStream> streams = streamsList.ToDictionary(x => x.Id, x => x);
 			List<IJob> recentJobs = await _jobs.FindAsync(minCreateTime: jobsCreatedAfter);
@@ -179,12 +183,20 @@ namespace Horde.Build.Agents.Fleet
 			}
 
 			List<(PoolId PoolId, int QueueSize)> poolsWithQueueSize = jobBatches.GroupBy(t => t.PoolId).Select(t => (t.Key, t.Count())).ToList();
+
+			scope.Span.SetTag("NumPools", poolsWithQueueSize.Count);
 			return poolsWithQueueSize.ToDictionary(x => x.PoolId, x => x.QueueSize);
 		}
 
 		/// <inheritdoc/>
 		public async Task<PoolSizeResult> CalculatePoolSizeAsync(IPool pool, List<IAgent> agents)
 		{
+			using IScope scope = GlobalTracer.Instance
+				.BuildSpan("JobQueueStrategy.CalculatePoolSize")
+				.WithTag(Datadog.Trace.OpenTracing.DatadogTags.ResourceName, pool.Id.ToString())
+				.WithTag("CurrentAgentCount", agents.Count)
+				.StartActive();
+			
 			DateTimeOffset minCreateTime = _clock.UtcNow - TimeSpan.FromMinutes(Settings.SamplePeriodMin);
 
 			// Cache pool queue sizes for a short while for faster runs when many pools are scaled
diff --git a/Engine/Source/Programs/Horde/Horde.Build/Agents/Fleet/LeaseUtilizationStrategy.cs b/Engine/Source/Programs/Horde/Horde.Build/Agents/Fleet/LeaseUtilizationStrategy.cs
--- a/Engine/Source/Programs/Horde/Horde.Build/Agents/Fleet/LeaseUtilizationStrategy.cs
+++ b/Engine/Source/Programs/Horde/Horde.Build/Agents/Fleet/LeaseUtilizationStrategy.cs
@@ -154,7 +154,7 @@ namespace Horde.Build.Agents.Fleet
 
 		private async Task<Dictionary<AgentId, AgentData>> GetAgentDataAsync()
 		{
-			using IScope _ = GlobalTracer.Instance.BuildSpan("GetAgentDataAsync").StartActive();
+			using IScope scope = GlobalTracer.Instance.BuildSpan("LeaseUtilizationStrategy.GetAgentDataAsync").StartActive();
 			
 			// Find all the current agents
 			List<IAgent> agents = await _agentCollection.FindAsync(status: AgentStatus.Ok);
@@ -197,12 +197,13 @@ namespace Horde.Build.Agents.Fleet
 				}
 			}
 
+			scope.Span.SetTag("AgentDataCount", agentIdToData.Count);
 			return agentIdToData;
 		}
 		
 		private async Task<Dictionary<PoolId, PoolData>> GetPoolDataAsync()
 		{
-			using IScope _ = GlobalTracer.Instance.BuildSpan("GetPoolDataAsync").StartActive();
+			using IScope scope = GlobalTracer.Instance.BuildSpan("LeaseUtilizationStrategy.GetPoolDataAsync").StartActive();
 
 			Dictionary<AgentId, AgentData> agentIdToData = await GetAgentDataAsync();
 			
@@ -223,6 +224,7 @@ namespace Horde.Build.Agents.Fleet
 				}
 			}
 
+			scope.Span.SetTag("PoolDataCount", agentIdToData.Count);
 			return poolToData;
 		}
 
@@ -232,6 +234,12 @@ namespace Horde.Build.Agents.Fleet
 		/// <inheritdoc/>
 		public async Task<PoolSizeResult> CalculatePoolSizeAsync(IPool pool, List<IAgent> agents)
 		{
+			using IScope scope = GlobalTracer.Instance
+				.BuildSpan("LeaseUtilizationStrategy.CalculatePoolSize")
+				.WithTag(Datadog.Trace.OpenTracing.DatadogTags.ResourceName, pool.Id.ToString())
+				.WithTag("CurrentAgentCount", agents.Count)
+				.StartActive();
+			
 			Dictionary<PoolId, PoolData> poolToData;
 			
 			// Cache pool data for a short while for faster runs when many pools are scaled
