Horde: Add ability to use an alternative sync on agent without have-table through a feature flag

diff --git a/Engine/Source/Programs/Horde/Horde.Agent/Execution/PerforceExecutor.cs b/Engine/Source/Programs/Horde/Horde.Agent/Execution/PerforceExecutor.cs
--- a/Engine/Source/Programs/Horde/Horde.Agent/Execution/PerforceExecutor.cs
+++ b/Engine/Source/Programs/Horde/Horde.Agent/Execution/PerforceExecutor.cs
@@ -32,6 +32,7 @@ namespace Horde.Agent.Execution
 
 		protected WorkspaceInfo? _autoSdkWorkspace;
 		protected WorkspaceInfo _workspace;
+		private bool _syncWithoutHaveTable = false;
 
 		public PerforceExecutor(ISession session, string jobId, string batchId, string agentTypeName, AgentWorkspace? autoSdkWorkspaceInfo, AgentWorkspace workspaceInfo, DirectoryReference rootDir, IHttpClientFactory httpClientFactory)
 			: base(session, jobId, batchId, agentTypeName, httpClientFactory)
@@ -46,11 +47,19 @@ namespace Horde.Agent.Execution
 		public override async Task InitializeAsync(ILogger logger, CancellationToken cancellationToken)
 		{
 			await base.InitializeAsync(logger, cancellationToken);
+			
+			int index = _additionalArguments.FindIndex(x => x.Contains("-SyncWithoutHaveTable", StringComparison.OrdinalIgnoreCase));
+			if (index >= 0)
+			{
+				logger.LogInformation("Syncing without have table");
+				_additionalArguments.RemoveAt(index);
+				_syncWithoutHaveTable = true;
+			}
 
 			// Setup and sync the autosdk workspace
 			if (_autoSdkWorkspaceInfo != null)
 			{
-				using (IScope ccope = GlobalTracer.Instance.BuildSpan("Workspace").WithResourceName("AutoSDK").StartActive())
+				using (IScope scope = GlobalTracer.Instance.BuildSpan("Workspace").WithResourceName("AutoSDK").StartActive())
 				{
 					_autoSdkWorkspace = await Utility.WorkspaceInfo.SetupWorkspaceAsync(_autoSdkWorkspaceInfo, _rootDir, logger, cancellationToken);
 
@@ -78,7 +87,7 @@ namespace Horde.Agent.Execution
 
 						FileReference autoSdkCacheFile = FileReference.Combine(_autoSdkWorkspace.MetadataDir, "Contents.dat");
 						await WorkspaceInfo.UpdateLocalCacheMarker(autoSdkCacheFile, autoSdkChangeNumber, -1);
-						await _autoSdkWorkspace.SyncAsync(autoSdkChangeNumber, -1, autoSdkCacheFile, cancellationToken);
+						await _autoSdkWorkspace.SyncAsync(autoSdkChangeNumber, -1, autoSdkCacheFile, true, cancellationToken);
 
 						await FileReference.WriteAllTextAsync(syncFile, syncText);
 					}
@@ -104,7 +113,14 @@ namespace Horde.Agent.Execution
 
 				// Sync the workspace
 				int syncPreflightChange = (_job.ClonedPreflightChange != 0) ? _job.ClonedPreflightChange : _job.PreflightChange;
-				await _workspace.SyncAsync(_job.Change, syncPreflightChange, null, cancellationToken);
+				if (_syncWithoutHaveTable)
+				{
+					await _workspace.SyncAsync(_job.Change, syncPreflightChange, null, false, cancellationToken);					
+				}
+				else
+				{
+					await _workspace.SyncAsync(_job.Change, syncPreflightChange, null, true, cancellationToken);
+				}
 
 				// Remove any cached BuildGraph manifests
 				DirectoryReference manifestDir = DirectoryReference.Combine(_workspace.WorkspaceDir, "Engine", "Saved", "BuildGraph");
@@ -368,7 +384,7 @@ namespace Horde.Agent.Execution
 					if (populateRequests.Count == 1 && !firstWorkspace.RemoveUntrackedFiles && !removeUntrackedFiles)
 					{
 						await firstWorkspace.CleanAsync(cancellationToken);
-						syncFuncs.Add(() => firstWorkspace.SyncAsync(-1, -1, null, cancellationToken));
+						syncFuncs.Add(() => firstWorkspace.SyncAsync(-1, -1, null, true, cancellationToken));
 					}
 					else
 					{
diff --git a/Engine/Source/Programs/Horde/Horde.Agent/Utility/WorkspaceInfo.cs b/Engine/Source/Programs/Horde/Horde.Agent/Utility/WorkspaceInfo.cs
--- a/Engine/Source/Programs/Horde/Horde.Agent/Utility/WorkspaceInfo.cs
+++ b/Engine/Source/Programs/Horde/Horde.Agent/Utility/WorkspaceInfo.cs
@@ -295,11 +295,19 @@ namespace Horde.Agent.Utility
 		/// <param name="change">The changelist to sync to</param>
 		/// <param name="preflightChange">Change to preflight, or 0</param>
 		/// <param name="cacheFile">Path to the cache file to use</param>
+		/// <param name="useHaveTable">Use client's have table when resolving contents</param>
 		/// <param name="cancellationToken">Cancellation token</param>
 		/// <returns>Async task</returns>
-		public async Task SyncAsync(int change, int preflightChange, FileReference? cacheFile, CancellationToken cancellationToken)
+		public async Task SyncAsync(int change, int preflightChange, FileReference? cacheFile, bool useHaveTable, CancellationToken cancellationToken)
 		{
-			await Repository.SyncAsync(PerforceClient, StreamName, change, View, RemoveUntrackedFiles, false, cacheFile, cancellationToken);
+			if (useHaveTable)
+			{
+				await Repository.SyncAsync(PerforceClient, StreamName, change, View, RemoveUntrackedFiles, false, cacheFile, cancellationToken);
+			}
+			else
+			{
+				await Repository.SyncWithoutHaveTableAsync(PerforceClient, StreamName, change, View, RemoveUntrackedFiles, false, cacheFile, cancellationToken);
+			}
 
 			// Purge the cache for incremental workspaces
 			if (!RemoveUntrackedFiles)
