Horde: Remove requirement to specify a target for a change query.

diff --git a/Engine/Source/Programs/Horde/Horde.Build/Jobs/JobService.cs b/Engine/Source/Programs/Horde/Horde.Build/Jobs/JobService.cs
--- a/Engine/Source/Programs/Horde/Horde.Build/Jobs/JobService.cs
+++ b/Engine/Source/Programs/Horde/Horde.Build/Jobs/JobService.cs
@@ -439,7 +439,7 @@ namespace Horde.Build.Jobs
 				}
 
 				// Find the job with the given result
-				if (query.Target != null && query.TemplateId != null)
+				if (query.TemplateId != null)
 				{
 					List<JobStepOutcome> outcomes = query.Outcomes ?? new List<JobStepOutcome> { JobStepOutcome.Success };
 
