Horde: Initial cut of automation test data v2 (feature flagged)
#rnx

diff --git a/Engine/Source/Programs/Horde/Horde.Build.Tests/TestDataTests.cs b/Engine/Source/Programs/Horde/Horde.Build.Tests/TestDataTests.cs
new file mode 100644
--- /dev/null
+++ b/Engine/Source/Programs/Horde/Horde.Build.Tests/TestDataTests.cs
@@ -0,0 +1,557 @@
+// Copyright Epic Games, Inc. All Rights Reserved.
+
+using System;
+using System.Collections.Generic;
+using System.Threading.Tasks;
+using EpicGames.Core;
+using Horde.Build.Jobs.Graphs;
+using Horde.Build.Projects;
+using Horde.Build.Jobs;
+using Horde.Build.Users;
+using Horde.Build.Streams;
+using Horde.Build.Utilities;
+using HordeCommon;
+using Microsoft.VisualStudio.TestTools.UnitTesting;
+using MongoDB.Bson;
+using Moq;
+using Horde.Build.Jobs.Templates;
+using Horde.Build.Jobs.TestData;
+using Microsoft.AspNetCore.Http;
+using Microsoft.AspNetCore.Mvc;
+using System.Security.Claims;
+using EpicGames.BuildGraph.Expressions;
+using Horde.Build.Logs;
+
+namespace Horde.Build.Tests
+{
+	using ProjectId = StringId<IProject>;
+	using StreamId = StringId<IStream>;
+	using TemplateId = StringId<ITemplateRef>;
+	using JobId = ObjectId<IJob>;
+
+	/// <summary>
+	/// Tests for the device service
+	/// </summary>
+	[TestClass]
+	public class TestDataTests : TestSetup
+	{
+		const string MainStreamName = "//UE5/Main";
+		readonly StreamId _mainStreamId = StreamId.Sanitize(MainStreamName);
+
+		const string ReleaseStreamName = "//UE5/Release";
+		readonly StreamId _releaseStreamId = StreamId.Sanitize(ReleaseStreamName);
+
+		readonly IGraph _graph;
+
+		TestDataController? _testDataController;
+
+		// override DeviceController with valid user
+		private new TestDataController TestDataController
+		{
+			get
+			{
+				if (_testDataController == null)
+				{
+					IUser user = UserCollection.FindOrAddUserByLoginAsync("TestUser").Result;
+					_testDataController = base.TestDataController;
+					ControllerContext controllerContext = new ControllerContext();
+					controllerContext.HttpContext = new DefaultHttpContext();
+					controllerContext.HttpContext.User = new ClaimsPrincipal(new ClaimsIdentity(
+						new List<Claim> {new Claim(ServerSettings.AdminClaimType, ServerSettings.AdminClaimValue),
+					new Claim(ClaimTypes.Name, "TestUser"),
+					new Claim(HordeClaimTypes.UserId, user.Id.ToString()) }
+						, "TestAuthType"));
+					_testDataController.ControllerContext = controllerContext;
+
+				}
+				return _testDataController;
+			}
+		}
+
+		public static INode MockNode(string name, IReadOnlyNodeAnnotations annotations)
+		{
+			Mock<INode> node = new Mock<INode>(MockBehavior.Strict);
+			node.SetupGet(x => x.Name).Returns(name);
+			node.SetupGet(x => x.Annotations).Returns(annotations);
+			return node.Object;
+		}
+
+		async Task<IStream> CreateStreamAsync(ProjectId projectId, StreamId streamId, string streamName)
+		{
+			string revision = $"config:{streamId}";
+
+			StreamConfig streamConfig = new StreamConfig { Name = streamName };
+			streamConfig.Tabs.Add(new JobsTabConfig { Title = "General", Templates = new List<TemplateId> { new TemplateId("test-template") } });
+			streamConfig.Templates.Add(new TemplateRefConfig { Id = new TemplateId("test-template") });
+			await ConfigCollection.AddConfigAsync(revision, streamConfig);
+
+			return Deref(await StreamCollection.TryCreateOrReplaceAsync(streamId, null, revision, projectId));
+		}
+
+		public TestDataTests()
+		{
+			const string ProjectConfigRevision = "projectconfig";
+			ConfigCollection.AddConfigAsync(ProjectConfigRevision, new ProjectConfig { Name = "UE4" }).Wait();
+			IProject project = ProjectCollection.AddOrUpdateAsync(new ProjectId("ue4"), ProjectConfigRevision, 0).Result!;
+
+			IStream mainStream = CreateStreamAsync(project.Id, _mainStreamId, MainStreamName).Result;
+			IStream releaseStream = CreateStreamAsync(project.Id, _releaseStreamId, ReleaseStreamName).Result;
+
+			List<INode> nodes = new List<INode>();
+			nodes.Add(MockNode("Update Version Files", NodeAnnotations.Empty));
+			nodes.Add(MockNode("Compile UnrealHeaderTool Win64", NodeAnnotations.Empty));
+			nodes.Add(MockNode("Compile ShooterGameEditor Win64", NodeAnnotations.Empty));
+			nodes.Add(MockNode("Cook ShooterGame Win64", NodeAnnotations.Empty));
+
+			Mock<INodeGroup> grp = new Mock<INodeGroup>(MockBehavior.Strict);
+			grp.SetupGet(x => x.Nodes).Returns(nodes);
+
+			Mock<IGraph> graphMock = new Mock<IGraph>(MockBehavior.Strict);
+			graphMock.SetupGet(x => x.Groups).Returns(new List<INodeGroup> { grp.Object });
+			_graph = graphMock.Object;
+		}
+
+		public IJob CreateJob(StreamId streamId, int change, string name, IGraph graph, TimeSpan time = default)
+		{
+			JobId jobId = JobId.GenerateNewId();
+
+			List<IJobStepBatch> batches = new List<IJobStepBatch>();
+			for (int groupIdx = 0; groupIdx < graph.Groups.Count; groupIdx++)
+			{
+				INodeGroup @group = graph.Groups[groupIdx];
+
+				List<IJobStep> steps = new List<IJobStep>();
+				for (int nodeIdx = 0; nodeIdx < @group.Nodes.Count; nodeIdx++)
+				{
+					SubResourceId stepId = new SubResourceId((ushort)((groupIdx * 100) + nodeIdx));
+
+					ILogFile logFile = LogFileService.CreateLogFileAsync(jobId, null, LogType.Json).Result;
+
+					Mock<IJobStep> step = new Mock<IJobStep>(MockBehavior.Strict);
+					step.SetupGet(x => x.Id).Returns(stepId);
+					step.SetupGet(x => x.NodeIdx).Returns(nodeIdx);
+					step.SetupGet(x => x.LogId).Returns(logFile.Id);
+					step.SetupGet(x => x.StartTimeUtc).Returns(DateTime.UtcNow + time);
+
+					steps.Add(step.Object);
+				}
+
+				SubResourceId batchId = new SubResourceId((ushort)(groupIdx * 100));
+
+				Mock<IJobStepBatch> batch = new Mock<IJobStepBatch>(MockBehavior.Strict);
+				batch.SetupGet(x => x.Id).Returns(batchId);
+				batch.SetupGet(x => x.GroupIdx).Returns(groupIdx);
+				batch.SetupGet(x => x.Steps).Returns(steps);
+				batches.Add(batch.Object);
+			}
+
+			Mock<IJob> job = new Mock<IJob>(MockBehavior.Strict);
+			job.SetupGet(x => x.Id).Returns(jobId);
+			job.SetupGet(x => x.Name).Returns(name);
+			job.SetupGet(x => x.StreamId).Returns(streamId);
+			job.SetupGet(x => x.TemplateId).Returns(new TemplateId("test-template"));
+			job.SetupGet(x => x.Change).Returns(change);
+			job.SetupGet(x => x.Batches).Returns(batches);
+			job.SetupGet(x => x.ShowUgsBadges).Returns(false);
+			job.SetupGet(x => x.ShowUgsAlerts).Returns(false);
+			job.SetupGet(x => x.PromoteIssuesByDefault).Returns(false);
+			job.SetupGet(x => x.UpdateIssues).Returns(false);
+			job.SetupGet(x => x.NotificationChannel).Returns("#devtools-horde-slack-testing");
+			return job.Object;
+		}
+
+
+
+		[TestMethod]
+		public async Task SimpleReportTest()
+		{
+			if (ServerSettings.FeatureFlags.EnableTestDataV2)
+			{
+				string[] streamIds = new string[] { _mainStreamId.ToString(), _releaseStreamId.ToString() };
+
+				IJob job = CreateJob(_mainStreamId, 105, "Test Build", _graph);
+				IJobStep step = job.Batches[0].Steps[0];
+				IJob job2 = CreateJob(_releaseStreamId, 105, "Test Build", _graph);
+				IJobStep step2 = job2.Batches[0].Steps[0];
+
+				BsonDocument testData = BsonDocument.Parse(String.Join('\n', _simpleTestDataLines));
+
+				await TestDataCollection.AddAsync(job, step, "Simple Report Key", testData);
+				await TestDataCollection.AddAsync(job2, step2, "Simple Report Key", testData);
+
+
+				ActionResult<List<GetTestStreamResponse>> streamResult = await TestDataController.GetTestStreamsAsync(streamIds);
+				Assert.IsNotNull(streamResult);
+				Assert.IsNotNull(streamResult.Value);
+				List<GetTestStreamResponse> streams = streamResult.Value;
+
+				Assert.AreEqual(2, streams.Count);
+				Assert.AreEqual(1, streams[0].Tests.Count);
+				Assert.AreEqual(1, streams[0].TestMetadata.Count);
+				Assert.AreEqual(0, streams[0].TestSuites.Count);
+				Assert.AreEqual(1, streams[1].Tests.Count);
+				Assert.AreEqual(1, streams[1].TestMetadata.Count);
+				Assert.AreEqual(0, streams[1].TestSuites.Count);
+
+				Assert.AreEqual(streams[0].TestMetadata[0].Id, streams[1].TestMetadata[0].Id);
+
+				ActionResult<List<GetTestDataRefResponse>> refResult = await TestDataController.GetTestDataRefAsync(streamIds);
+				Assert.IsNotNull(refResult);
+				Assert.IsNotNull(refResult.Value);
+				List<GetTestDataRefResponse> refs = refResult.Value;
+
+				Assert.AreEqual(2, refs.Count);
+			}
+		}
+
+		[TestMethod]
+		public async Task SessionReportTest()
+		{
+			if (ServerSettings.FeatureFlags.EnableTestDataV2)
+			{
+				string[] streamIds = new string[] { _mainStreamId.ToString(), _releaseStreamId.ToString() };
+
+				IJob job = CreateJob(_mainStreamId, 105, "Test Build", _graph);
+				IJobStep step = job.Batches[0].Steps[0];
+				IJob job2 = CreateJob(_releaseStreamId, 105, "Test Build", _graph);
+				IJobStep step2 = job2.Batches[0].Steps[0];
+
+				BsonDocument testData = BsonDocument.Parse(String.Join('\n', _testSessionDataLines));
+
+				await TestDataCollection.AddAsync(job, step, "Session Report Key", testData);
+				await TestDataCollection.AddAsync(job2, step2, "Session Report Key", testData);
+
+				ActionResult<List<GetTestStreamResponse>> streamResult = await TestDataController.GetTestStreamsAsync(streamIds);
+				Assert.IsNotNull(streamResult);
+				Assert.IsNotNull(streamResult.Value);
+				List<GetTestStreamResponse> streams = streamResult.Value;
+
+				Assert.AreEqual(2, streams.Count);
+
+				Assert.AreEqual(4, streams[0].Tests.Count);				
+				Assert.AreEqual(1, streams[0].TestMetadata.Count);
+				Assert.AreEqual(1, streams[0].TestSuites.Count);
+				Assert.AreEqual(4, streams[0].TestSuites[0].Tests.Count);
+
+				Assert.AreEqual(4, streams[1].Tests.Count);
+				Assert.AreEqual(1, streams[1].TestMetadata.Count);
+				Assert.AreEqual(1, streams[1].TestSuites.Count);
+				Assert.AreEqual(4, streams[1].TestSuites[0].Tests.Count);
+
+				Assert.AreEqual(streams[0].TestMetadata[0].Id, streams[1].TestMetadata[0].Id);
+
+				ActionResult<List<GetTestDataRefResponse>> refResult = await TestDataController.GetTestDataRefAsync(streamIds);
+				Assert.IsNotNull(refResult);
+				Assert.IsNotNull(refResult.Value);
+				List<GetTestDataRefResponse> refs = refResult.Value;
+
+				Assert.AreEqual(2, refs.Count);
+			}
+			
+		}
+
+		private readonly string[] _simpleTestDataLines =
+		{
+			@"{",
+			@"  ""Items"": [",
+			@"    {",
+			@"      ""Key"": ""Simple Report::UE.BootTest EngineTest Editor Win64"",",
+			@"      ""Data"": {",
+			@"		""Version"" : 1,",
+			@"        ""Type"": ""Simple Report"",",
+			@"        ""TestName"": ""EditorBootTest"",",
+			@"        ""Description"": ""Win64 Development EditorGame"",",
+			@"        ""ReportCreatedOn"": ""10/31/2022 11:44:56 AM"",",
+			@"        ""TotalDurationSeconds"": 35.33694,",
+			@"        ""HasSucceeded"": true,",
+			@"        ""Status"": ""Passed"",",
+			@"        ""URLLink"": """",",
+			@"        ""BuildChangeList"": 22815797,",
+			@"        ""MainRole"": {",
+			@"          ""Type"": ""Editor"",",
+			@"          ""Platform"": ""Win64"",",
+			@"          ""Configuration"": ""Development""",
+			@"        },",
+			@"        ""Roles"": [",
+			@"          {",
+			@"            ""Type"": ""Editor"",",
+			@"            ""Platform"": ""Win64"",",
+			@"            ""Configuration"": ""Development""",
+			@"          }",
+			@"        ],",
+			@"        ""TestResult"": ""Passed"",",
+			@"        ""Logs"": [],",
+			@"        ""Errors"": [],",
+			@"        ""Warnings"": [],",
+			@"        ""Metadata"": {",
+			@"          ""Platform"": ""Win64"",",
+			@"          ""BuildTarget"": ""Editor"",",
+			@"          ""Configuration"": ""Development"",",
+			@"          ""Project"": ""EngineTest""",
+			@"        }",
+			@"      }",
+			@"    }",
+			@"  ],",
+			@"  ""Version"": 1",
+			@"}"
+		};
+
+		private readonly string[] _testSessionDataLines =
+		{
+			@"{",
+			@"	""Version"": 1,",
+			@"    ""Items"": [	",
+			@"        {",
+			@"            ""Key"": ""Automated Test Session"",",
+			@"            ""Data"": {",
+			@"                ""Type"": ""Automated Test Session"",",
+			@"                ""Name"": ""UE.Automation(Group:UI) EngineTest"",",
+			@"                ""PreFlightChange"": """",",
+			@"                ""TestSessionInfo"": {",
+			@"                    ""DateTime"": ""2022.10.31-01.35.07"",",
+			@"                    ""TimeElapseSec"": 17.63286,",
+			@"                    ""Tests"": {",
+			@"                        ""cbdb55ea"": {",
+			@"                            ""Name"": ""Project.Functional Tests.Tests.UI.Clipping.WidgetClipping.Simple UI"",",
+			@"                            ""TestUID"": ""cbdb55ea"",",
+			@"                            ""Suite"": ""UI"",",
+			@"                            ""State"": ""Skipped"",",
+			@"                            ""DeviceAppInstanceName"": [",
+			@"                                ""3da3633f-c82a-5c25-9a60-1e852705b65e-C0218E40D9AC4B4FBD07326E10179042""",
+			@"                            ],",
+			@"                            ""ErrorCount"": 0,",
+			@"                            ""WarningCount"": 0,",
+			@"                            ""ErrorHashAggregate"": """",",
+			@"                            ""DateTime"": ""2022.10.31-01.34.41"",",
+			@"                            ""TimeElapseSec"": 0",
+			@"                        },",
+			@"                        ""eafa1362"": {",
+			@"                            ""Name"": ""Project.Functional Tests.Tests.UI.Effects.UI_Effects.Blur"",",
+			@"                            ""TestUID"": ""eafa1362"",",
+			@"                            ""Suite"": ""UI"",",
+			@"                            ""State"": ""Success"",",
+			@"                            ""DeviceAppInstanceName"": [",
+			@"                                ""3da3633f-c82a-5c25-9a60-1e852705b65e-C0218E40D9AC4B4FBD07326E10179042""",
+			@"                            ],",
+			@"                            ""ErrorCount"": 0,",
+			@"                            ""WarningCount"": 0,",
+			@"                            ""ErrorHashAggregate"": """",",
+			@"                            ""DateTime"": ""2022.10.31-01.34.49"",",
+			@"                            ""TimeElapseSec"": 10.5334",
+			@"                        },",
+			@"                        ""1521079c"": {",
+			@"                            ""Name"": ""Project.Functional Tests.Tests.UI.Fonts.FontOutlineTestUI.FontOutlineTest"",",
+			@"                            ""TestUID"": ""1521079c"",",
+			@"                            ""Suite"": ""UI"",",
+			@"                            ""State"": ""Success"",",
+			@"                            ""DeviceAppInstanceName"": [",
+			@"                                ""3da3633f-c82a-5c25-9a60-1e852705b65e-C0218E40D9AC4B4FBD07326E10179042""",
+			@"                            ],",
+			@"                            ""ErrorCount"": 0,",
+			@"                            ""WarningCount"": 0,",
+			@"                            ""ErrorHashAggregate"": """",",
+			@"                            ""DateTime"": ""2022.10.31-01.35.00"",",
+			@"                            ""TimeElapseSec"": 3.53325",
+			@"                        },",
+			@"                        ""c2638213"": {",
+			@"                            ""Name"": ""Project.Functional Tests.Tests.UI.Fonts.FontRenderingTestUI.FontRenderingTest"",",
+			@"                            ""TestUID"": ""c2638213"",",
+			@"                            ""Suite"": ""UI"",",
+			@"                            ""State"": ""Fail"",",
+			@"                            ""DeviceAppInstanceName"": [",
+			@"                                ""3da3633f-c82a-5c25-9a60-1e852705b65e-C0218E40D9AC4B4FBD07326E10179042""",
+			@"                            ],",
+			@"                            ""ErrorCount"": 1,",
+			@"                            ""WarningCount"": 2,",
+			@"                            ""ErrorHashAggregate"": """",",
+			@"                            ""DateTime"": ""2022.10.31-01.35.03"",",
+			@"                            ""TimeElapseSec"": 3.56621",
+			@"                        }",
+			@"                    },",
+			@"                    ""TestResultsTestDataUID"": ""b420bdde-c030-4add-81d2-3a8404ab3e45""",
+			@"                },",
+			@"                ""Devices"": [",
+			@"                    {",
+			@"                        ""Name"": ""RDU-WIN64-12"",",
+			@"                        ""AppInstanceName"": ""3da3633f-c82a-5c25-9a60-1e852705b65e-C0218E40D9AC4B4FBD07326E10179042"",",
+			@"                        ""AppInstanceLog"": ""UI_Win64_Game/Client/ClientOutput.log"",",
+			@"                        ""Metadata"": {",
+			@"                            ""platform"": ""Win64"",",
+			@"                            ""os_version"": ""14 3"",",
+			@"                            ""model"": ""Default"",",
+			@"                            ""gpu"": ""Win64 GPU"",",
+			@"                            ""cpumodel"": ""Win64 CPU"",",
+			@"                            ""ram_in_gb"": ""5"",",
+			@"                            ""render_mode"": ""ES3_1"",",
+			@"                            ""rhi"": """"",
+			@"                        }",
+			@"                    }",
+			@"                ],",
+			@"                ""IndexedErrors"": {},",
+			@"                ""Metadata"": {",
+			@"                    ""Platform"": ""Win64"",",
+			@"                    ""BuildTarget"": ""Client"",",
+			@"                    ""Configuration"": ""Development"",",
+			@"                    ""Project"": ""EngineTest"",",
+			@"                    ""RHI"": ""default""",
+			@"                }",
+			@"            }",
+			@"        },",
+			@"        {",
+			@"            ""Key"": ""Unreal Automated Tests::UE.TargetAutomation(RunTest=UI) Win64"",",
+			@"            ""Data"": {",
+			@"                ""Type"": ""Unreal Automated Tests"",",
+			@"                ""Devices"": [",
+			@"                    {",
+			@"                        ""DeviceName"": ""RDU-WIN64-12"",",
+			@"                        ""Instance"": ""3da3633f-c82a-5c25-9a60-1e852705b65e-C0218E40D9AC4B4FBD07326E10179042"",",
+			@"                        ""Platform"": ""Win64"",",
+			@"                        ""OSVersion"": ""14 3"",",
+			@"                        ""Model"": ""Default"",",
+			@"                        ""GPU"": ""Win64 GPU"",",
+			@"                        ""CPUModel"": ""Win64 CPU"",",
+			@"                        ""RAMInGB"": 5,",
+			@"                        ""RenderMode"": ""ES3_1"",",
+			@"                        ""RHI"": """"",
+			@"                    }",
+			@"                ],",
+			@"                ""ReportCreatedOn"": ""2022.10.31-01.35.07"",",
+			@"                ""ReportURL"": """",",
+			@"                ""SucceededCount"": 3,",
+			@"                ""SucceededWithWarningsCount"": 0,",
+			@"                ""FailedCount"": 0,",
+			@"                ""NotRunCount"": 0,",
+			@"                ""InProcessCount"": 0,",
+			@"                ""TotalDurationSeconds"": 17.63286,",
+			@"                ""Tests"": [",
+			@"                    {",
+			@"                        ""TestDisplayName"": ""Simple UI"",",
+			@"                        ""FullTestPath"": ""Project.Functional Tests.Tests.UI.Clipping.WidgetClipping.Simple UI"",",
+			@"                        ""State"": ""Skipped"",",
+			@"                        ""DeviceInstance"": ""3da3633f-c82a-5c25-9a60-1e852705b65e-C0218E40D9AC4B4FBD07326E10179042"",",
+			@"                        ""Errors"": 0,",
+			@"                        ""Warnings"": 0,",
+			@"                        ""ArtifactName"": ""33ec1c33-71b7-41a1-84ba-d5ee2f42af8c.json""",
+			@"                    },",
+			@"                    {",
+			@"                        ""TestDisplayName"": ""Blur"",",
+			@"                        ""FullTestPath"": ""Project.Functional Tests.Tests.UI.Effects.UI_Effects.Blur"",",
+			@"                        ""State"": ""Success"",",
+			@"                        ""DeviceInstance"": ""3da3633f-c82a-5c25-9a60-1e852705b65e-C0218E40D9AC4B4FBD07326E10179042"",",
+			@"                        ""Errors"": 0,",
+			@"                        ""Warnings"": 0,",
+			@"                        ""ArtifactName"": ""450130ac-18cb-4bee-b529-c1d8ae943a61.json""",
+			@"                    },",
+			@"                    {",
+			@"                        ""TestDisplayName"": ""FontOutlineTest"",",
+			@"                        ""FullTestPath"": ""Project.Functional Tests.Tests.UI.Fonts.FontOutlineTestUI.FontOutlineTest"",",
+			@"                        ""State"": ""Success"",",
+			@"                        ""DeviceInstance"": ""3da3633f-c82a-5c25-9a60-1e852705b65e-C0218E40D9AC4B4FBD07326E10179042"",",
+			@"                        ""Errors"": 0,",
+			@"                        ""Warnings"": 0,",
+			@"                        ""ArtifactName"": ""12752762-6332-462f-8ef3-254c7502b57a.json""",
+			@"                    },",
+			@"                    {",
+			@"                        ""TestDisplayName"": ""FontRenderingTest"",",
+			@"                        ""FullTestPath"": ""Project.Functional Tests.Tests.UI.Fonts.FontRenderingTestUI.FontRenderingTest"",",
+			@"                        ""State"": ""Success"",",
+			@"                        ""DeviceInstance"": ""3da3633f-c82a-5c25-9a60-1e852705b65e-C0218E40D9AC4B4FBD07326E10179042"",",
+			@"                        ""Errors"": 0,",
+			@"                        ""Warnings"": 0,",
+			@"                        ""ArtifactName"": ""711a5bd9-26dd-470b-a3c6-98f04a522e98.json""",
+			@"                    }",
+			@"                ],",
+			@"                ""Metadata"": {",
+			@"                    ""Platform"": ""Win64"",",
+			@"                    ""BuildTarget"": ""Client"",",
+			@"                    ""Configuration"": ""Development"",",
+			@"                    ""Project"": ""EngineTest""",
+			@"                }",
+			@"            }",
+			@"        },",
+			@"        {",
+			@"            ""Key"": ""Automated Test Session Result Details::b420bdde-c030-4add-81d2-3a8404ab3e45"",",
+			@"            ""Data"": {",
+			@"                ""cbdb55ea"": {",
+			@"                    ""Events"": [",
+			@"                        {",
+			@"                            ""Message"": ""Skipping test because of exclude list: Sporadic failure with animated UI elements. Jira UE-113396"",",
+			@"                            ""Context"": """",",
+			@"                            ""Type"": ""Info"",",
+			@"                            ""Tag"": ""entry"",",
+			@"                            ""Hash"": """",",
+			@"                            ""DateTime"": ""0001.01.01-00.00.00"",",
+			@"                            ""Artifacts"": []",
+			@"                        }",
+			@"                    ]",
+			@"                },",
+			@"                ""eafa1362"": {",
+			@"                    ""Events"": [",
+			@"                        {",
+			@"                            ""Message"": ""LogGauntlet: GauntletHeartbeat: Idle "",",
+			@"                            ""Context"": """",",
+			@"                            ""Type"": ""Info"",",
+			@"                            ""Tag"": ""entry"",",
+			@"                            ""Hash"": """",",
+			@"                            ""DateTime"": ""0001.01.01-00.00.00"",",
+			@"                            ""Artifacts"": []",
+			@"                        },",
+			@"                        {",
+			@"                            ""Message"": ""Screenshot 'Blur' was similar!  Global Difference = 0.000303, Max Local Difference = 0.005257"",",
+			@"                            ""Context"": ""UI_Effects/Blur"",",
+			@"                            ""Type"": ""Info"",",
+			@"                            ""Tag"": ""image comparison"",",
+			@"                            ""Hash"": """",",
+			@"                            ""DateTime"": ""0001.01.01-00.00.00"",",
+			@"                            ""Artifacts"": [",
+			@"                                {",
+			@"                                    ""Tag"": ""unapproved"",",
+			@"                                    ""ReferencePath"": ""imageCompare/UI_Effects/Blur/Win64/ES3_1/Incoming.png""",
+			@"                                }",
+			@"                            ]",
+			@"                        }",
+			@"                    ]",
+			@"                },",
+			@"                ""1521079c"": {",
+			@"                    ""Events"": [",
+			@"                        {",
+			@"                            ""Message"": ""Screenshot 'FontOutlineTest' was similar!  Global Difference = 0.000083, Max Local Difference = 0.003520"",",
+			@"                            ""Context"": ""FontOutlineTestUI/FontOutlineTest"",",
+			@"                            ""Type"": ""Info"",",
+			@"                            ""Tag"": ""image comparison"",",
+			@"                            ""Hash"": """",",
+			@"                            ""DateTime"": ""0001.01.01-00.00.00"",",
+			@"                            ""Artifacts"": [",
+			@"                                {",
+			@"                                    ""Tag"": ""unapproved"",",
+			@"                                    ""ReferencePath"": ""imageCompare/FontOutlineTestUI/FontOutlineTest/Win64/ES3_1/Incoming.png""",
+			@"                                }",
+			@"                            ]",
+			@"                        }",
+			@"                    ]",
+			@"                },",
+			@"                ""c2638213"": {",
+			@"                    ""Events"": [",
+			@"                        {",
+			@"                            ""Message"": ""Screenshot 'FontRenderingTest' was similar!  Global Difference = 0.000000, Max Local Difference = 0.000000"",",
+			@"                            ""Context"": ""FontRenderingTestUI/FontRenderingTest"",",
+			@"                            ""Type"": ""Info"",",
+			@"                            ""Tag"": ""image comparison"",",
+			@"                            ""Hash"": """",",
+			@"                            ""DateTime"": ""0001.01.01-00.00.00"",",
+			@"                            ""Artifacts"": [",
+			@"                                {",
+			@"                                    ""Tag"": ""unapproved"",",
+			@"                                    ""ReferencePath"": ""imageCompare/FontRenderingTestUI/FontRenderingTest/Win64/ES3_1/Incoming.png""",
+			@"                                }",
+			@"                            ]",
+			@"                        }",
+			@"                    ]",
+			@"                }",
+			@"            }",
+			@"        }",
+			@"    ]",
+			@"}"
+		};
+	}
+}
diff --git a/Engine/Source/Programs/Horde/Horde.Build.Tests/TestSetup.cs b/Engine/Source/Programs/Horde/Horde.Build.Tests/TestSetup.cs
--- a/Engine/Source/Programs/Horde/Horde.Build.Tests/TestSetup.cs
+++ b/Engine/Source/Programs/Horde/Horde.Build.Tests/TestSetup.cs
@@ -124,6 +124,7 @@ namespace Horde.Build.Tests
 		public ScheduleService ScheduleService => ServiceProvider.GetRequiredService<ScheduleService>();
 		public DeviceService DeviceService => ServiceProvider.GetRequiredService<DeviceService>();
 		public StorageService StorageService => ServiceProvider.GetRequiredService<StorageService>();
+		public TestDataService TestDataService => ServiceProvider.GetRequiredService<TestDataService>();
 
 		public ServerSettings ServerSettings => ServiceProvider.GetRequiredService<IOptions<ServerSettings>>().Value;
 		public IOptionsMonitor<ServerSettings> ServerSettingsMon => ServiceProvider.GetRequiredService<IOptionsMonitor<ServerSettings>>();
@@ -133,6 +134,7 @@ namespace Horde.Build.Tests
 		public PoolsController PoolsController => GetPoolsController();
 		public LeasesController LeasesController => GetLeasesController();
 		public DevicesController DevicesController => GetDevicesController();
+		public TestDataController TestDataController => GetTestDataController();
 
 		private static bool s_datadogWriterPatched;
 
@@ -238,6 +240,7 @@ namespace Horde.Build.Tests
 			services.AddSingleton<ScheduleService>();
 			services.AddSingleton<StreamService>();
 			services.AddSingleton<DeviceService>();
+			services.AddSingleton<TestDataService>();
 
 			services.AddSingleton<ConformTaskSource>();
 			services.AddSingleton<ICommitService, CommitService>();
@@ -291,6 +294,13 @@ namespace Horde.Build.Tests
 			return devicesCtrl;
 		}
 
+		private TestDataController GetTestDataController()
+		{
+			TestDataController dataCtrl = new TestDataController(TestDataService, StreamService, JobService, TestDataCollection);
+			dataCtrl.ControllerContext = GetControllerContext();
+			return dataCtrl;
+		}
+
 		private AgentsController GetAgentsController()
 		{
 			AgentsController agentCtrl = new AgentsController(AclService, AgentService);
diff --git a/Engine/Source/Programs/Horde/Horde.Build/Jobs/TestData/ITestData.cs b/Engine/Source/Programs/Horde/Horde.Build/Jobs/TestData/ITestData.cs
--- a/Engine/Source/Programs/Horde/Horde.Build/Jobs/TestData/ITestData.cs
+++ b/Engine/Source/Programs/Horde/Horde.Build/Jobs/TestData/ITestData.cs
@@ -1,5 +1,7 @@
 // Copyright Epic Games, Inc. All Rights Reserved.
 
+using System.Collections.Generic;
+using System;
 using Horde.Build.Streams;
 using Horde.Build.Utilities;
 using MongoDB.Bson;
@@ -10,6 +12,192 @@ namespace Horde.Build.Jobs.TestData
 	using StreamId = StringId<IStream>;
 	using TemplateId = StringId<ITemplateRef>;
 
+	using TestId = ObjectId<ITest>;
+	using TestSuiteId = ObjectId<ITestSuite>;
+	using TestMetaId = ObjectId<ITestMeta>;
+
+	/// <summary>
+	/// Test outcome
+	/// </summary>
+	public enum TestOutcome
+	{
+		/// <summary>
+		/// The test was successful
+		/// </summary>
+		Success,
+		/// <summary>
+		/// The test failed
+		/// </summary>
+		Failure,
+		/// <summary>
+		/// The test was skipped
+		/// </summary>
+		Skipped,
+		/// <summary>
+		/// The test had an unspecified result
+		/// </summary>
+		Unspecified
+	}
+
+	/// <summary>
+	/// Defines a testing environment based on platforms, configurations, targets, etc.
+	/// </summary>
+	public interface ITestMeta
+	{
+		/// <summary>
+		/// The test meta id
+		/// </summary>
+		TestMetaId Id { get; }
+		/// <summary>
+		/// The name of the test platform
+		/// </summary>
+		IReadOnlyList<string> Platforms { get; }
+
+		/// <summary>
+		/// The configuration the test was run on
+		/// </summary>
+		IReadOnlyList<string> Configurations { get; }
+
+		/// <summary>
+		/// The build target, editor, server, client, etc
+		/// </summary>
+		IReadOnlyList<string> BuildTargets { get; }
+
+		/// <summary>
+		/// The uproject name associated with this test, note: may not be directly related to Horde project
+		/// </summary>
+		string ProjectName { get; }
+
+		/// <summary>
+		/// The rendering hardware interface used for the test
+		/// </summary>
+		string RHI { get; }
+	}
+
+	/// <summary>
+	/// A test that runs in a stream
+	/// </summary>
+	public interface ITest
+	{
+		/// <summary>
+		/// The associated stream
+		/// </summary>
+		StreamId StreamId { get; }
+
+		/// <summary>
+		/// The fully qualified name of the test 
+		/// </summary>
+		string Name { get; }
+
+		/// <summary>
+		/// The display name of the test 
+		/// </summary>
+		string? DisplayName { get; }
+
+		/// <summary>
+		/// The meta data for the test 
+		/// </summary>
+		IReadOnlyList<TestMetaId> Metadata { get; }
+	}
+
+	/// <summary>
+	/// A test suite that runs in a stream
+	/// </summary>
+	public interface ITestSuite
+	{
+		/// <summary>
+		/// The associated stream
+		/// </summary>
+		StreamId StreamId { get; }
+
+		/// <summary>
+		/// The name of the test suite
+		/// </summary>
+		string Name { get; }
+
+		/// <summary>
+		/// The tests that compose the suite
+		/// </summary>
+		IReadOnlyList<TestId> Tests { get; }
+	}
+
+	/// <summary>
+	/// Suite test data
+	/// </summary>
+	public interface ISuiteTestData
+	{
+		/// <summary>
+		/// The test id
+		/// </summary>
+		TestId Id { get; }
+
+		/// <summary>
+		/// The ourcome of the suite test
+		/// </summary>
+		TestOutcome Outcome { get; }
+
+		/// <summary>
+		/// How long the suite test ran
+		/// </summary>
+		TimeSpan Duration { get; }
+
+		/// <summary>
+		/// Test UID for looking up in test details
+		/// </summary>
+		string UID { get; }
+	}
+
+	/// <summary>
+	/// Data ref 
+	/// </summary>
+	public interface ITestDataRef
+	{
+		/// <summary>
+		/// The associated stream
+		/// </summary>
+		StreamId StreamId { get; }
+
+		/// <summary>
+		/// The full details test data for this ref
+		/// </summary>
+		ObjectId TestDataId { get; }
+
+		/// <summary>
+		/// How long the test ran
+		/// </summary>
+		TimeSpan Duration { get; }
+
+		/// <summary>
+		/// The build changelist upon which the test ran, may not correspond to the job changelist
+		/// </summary>
+		int BuildChangeList { get; }
+
+		/// <summary>
+		/// The environment the test ran in
+		/// </summary>
+		TestMetaId Metadata { get; }
+
+		/// <summary>
+		/// The ITest in stream
+		/// </summary>
+		TestId? TestId { get; }
+
+		/// <summary>
+		/// The outcome of the test
+		/// </summary>
+		TestOutcome? Outcome { get; }
+
+		/// <summary>
+		/// The ITestSuite in stream
+		/// </summary>
+		TestSuiteId? SuiteId { get; }
+
+		/// <summary>
+		/// List for suite test data
+		/// </summary>
+		IReadOnlyList<ISuiteTestData>? SuiteTests { get; }
+	}
+	
 	/// <summary>
 	/// Stores information about the results of a test
 	/// </summary>
diff --git a/Engine/Source/Programs/Horde/Horde.Build/Jobs/TestData/ITestDataCollection.cs b/Engine/Source/Programs/Horde/Horde.Build/Jobs/TestData/ITestDataCollection.cs
--- a/Engine/Source/Programs/Horde/Horde.Build/Jobs/TestData/ITestDataCollection.cs
+++ b/Engine/Source/Programs/Horde/Horde.Build/Jobs/TestData/ITestDataCollection.cs
@@ -1,5 +1,6 @@
 // Copyright Epic Games, Inc. All Rights Reserved.
 
+using System;
 using System.Collections.Generic;
 using System.Threading.Tasks;
 using Horde.Build.Streams;
@@ -7,15 +8,39 @@ using Horde.Build.Utilities;
 using MongoDB.Bson;
 
 namespace Horde.Build.Jobs.TestData
-{
+{	
 	using JobId = ObjectId<IJob>;
 	using StreamId = StringId<IStream>;
+	using TestId = ObjectId<ITest>;
+	using TestSuiteId = ObjectId<ITestSuite>;
+	using TestMetaId = ObjectId<ITestMeta>;
 
 	/// <summary>
 	/// Collection of test data documents
 	/// </summary>
 	public interface ITestDataCollection
 	{
+		/// <summary>
+		/// Gets the tests running in provided streams
+		/// </summary>
+		/// <param name="streamIds"></param>
+		/// <returns></returns>
+		Task<List<TestStream>> FindTestStreams(StreamId[] streamIds);
+
+		/// <summary>
+		/// 
+		/// </summary>
+		/// <param name="streamIds"></param>
+		/// <param name="metaIds"></param>
+		/// <param name="testIds"></param>
+		/// <param name="suiteIds"></param>
+		/// <param name="minCreateTime"></param>
+		/// <param name="maxCreateTime"></param>
+		/// <param name="minChange"></param>
+		/// <param name="maxChange"></param>
+		/// <returns></returns>
+		Task<List<ITestDataRef>> FindTestRefs(StreamId[] streamIds, TestMetaId[]? metaIds = null, TestId[]? testIds = null, TestSuiteId[]? suiteIds = null, DateTime? minCreateTime = null, DateTime? maxCreateTime = null, int? minChange = null, int? maxChange = null);
+
 		/// <summary>
 		/// Creates a new test data document
 		/// </summary>
diff --git a/Engine/Source/Programs/Horde/Horde.Build/Jobs/TestData/TestDataCollection.cs b/Engine/Source/Programs/Horde/Horde.Build/Jobs/TestData/TestDataCollection.cs
--- a/Engine/Source/Programs/Horde/Horde.Build/Jobs/TestData/TestDataCollection.cs
+++ b/Engine/Source/Programs/Horde/Horde.Build/Jobs/TestData/TestDataCollection.cs
@@ -7,8 +7,16 @@ using System.Threading.Tasks;
 using Horde.Build.Server;
 using Horde.Build.Streams;
 using Horde.Build.Utilities;
+using Microsoft.Extensions.Configuration;
+using Microsoft.Extensions.Logging;
 using MongoDB.Bson;
+using MongoDB.Bson.Serialization;
+using MongoDB.Bson.Serialization.Attributes;
 using MongoDB.Driver;
+using OpenTracing.Util;
+using OpenTracing;
+using Microsoft.CodeAnalysis;
+using Microsoft.Extensions.Options;
 
 namespace Horde.Build.Jobs.TestData
 {
@@ -16,11 +24,235 @@ namespace Horde.Build.Jobs.TestData
 	using StreamId = StringId<IStream>;
 	using TemplateId = StringId<ITemplateRef>;
 
+	using TestId = ObjectId<ITest>;
+	using TestSuiteId = ObjectId<ITestSuite>;
+	using TestMetaId = ObjectId<ITestMeta>;
+	using static System.Net.Mime.MediaTypeNames;
+
+	/// <summary>
+	/// Tests which are run in a  stream
+	/// </summary>
+	public class TestStream
+	{
+		/// <summary>
+		/// The associated stream tests are running in
+		/// </summary>
+		public StreamId StreamId { get; set; }
+
+		/// <summary>
+		/// The individual tests running in the stream
+		/// </summary>
+		public List<ITest> Tests { get; set; }
+
+		/// <summary>
+		/// The tests suites running in the stream
+		/// </summary>
+		public List<ITestSuite> TestSuites { get; set; }
+
+		/// <summary>
+		/// The test meta data (environments) running in the stream
+		/// </summary>
+		public List<ITestMeta> TestMeta { get; set; }
+
+		/// <summary>
+		/// Constructor
+		/// </summary>
+		/// <param name="streamId"></param>
+		/// <param name="tests"></param>
+		/// <param name="testSuites"></param>
+		/// <param name="testMeta"></param>
+		public TestStream(StreamId streamId, List<ITest> tests, List<ITestSuite> testSuites, List<ITestMeta> testMeta)
+		{
+			StreamId = streamId;
+			Tests = tests;
+			TestSuites = testSuites;
+			TestMeta = testMeta;
+		}
+	}
+
 	/// <summary>
 	/// Collection of test data documents
 	/// </summary>
 	public class TestDataCollection : ITestDataCollection
 	{
+
+		class TestMetaDocument : ITestMeta
+		{
+			[BsonRequired, BsonId]
+			public TestMetaId Id { get; set; }
+
+			[BsonRequired, BsonElement("p")]
+			public List<string> Platforms { get; set; }
+			IReadOnlyList<string> ITestMeta.Platforms => Platforms;
+
+			[BsonRequired, BsonElement("c")]
+			public List<string> Configurations { get; set; }
+			IReadOnlyList<string> ITestMeta.Configurations => Configurations;
+
+			[BsonRequired, BsonElement("bt")]
+			public List<string> BuildTargets { get; set; }
+			IReadOnlyList<string> ITestMeta.BuildTargets => BuildTargets;
+
+			[BsonRequired, BsonElement("pn")]
+			public string ProjectName { get; set; }
+
+			[BsonRequired, BsonElement("r")]
+			public string RHI { get; set; }
+
+			private TestMetaDocument()
+			{
+				Platforms = new List<string>();
+				Configurations = new List<string>();
+				BuildTargets = new List<string>();
+
+				RHI = String.Empty;
+				ProjectName = String.Empty;
+			}
+
+			public TestMetaDocument(List<string> platforms, List<string> configurations, List<string> buildTargets, string projectName, string? rhi)
+			{
+				Id = TestMetaId.GenerateNewId();
+				Platforms = platforms;
+				Configurations = configurations;
+				BuildTargets = buildTargets;
+				ProjectName = projectName;
+				RHI = rhi ?? "default";
+			}
+		}
+
+		class TestDocument : ITest
+		{
+			[BsonRequired, BsonId]
+			public TestId Id { get; set; }
+
+			[BsonRequired, BsonElement("sid")]
+			public StreamId StreamId { get; set; }
+			
+			[BsonRequired, BsonElement("n")]
+			public string Name { get; set; }
+
+			[BsonIgnoreIfNull, BsonElement("dn")]
+			public string? DisplayName { get; set; }
+
+			[BsonRequired, BsonElement("m")]
+			public List<TestMetaId> Metadata { get; set; }
+			IReadOnlyList<TestMetaId> ITest.Metadata => Metadata;
+
+			private TestDocument()
+			{
+				Name = String.Empty;
+				Metadata = new List<TestMetaId>();
+			}
+
+			public TestDocument(StreamId streamId, string name, List<TestMetaId> meta, string? displayName = null)
+			{
+				Id = TestId.GenerateNewId();
+				StreamId = streamId;
+				Name = name;
+				Metadata = meta;
+				DisplayName = displayName;
+			}
+		}
+
+		class TestSuiteDocument : ITestSuite
+		{
+			[BsonRequired, BsonId]
+			public TestSuiteId Id { get; set; }
+
+			[BsonRequired, BsonElement("sid")]
+			public StreamId StreamId { get; set; }
+
+			[BsonRequired, BsonElement("n")]
+			public string Name { get; set; }
+
+			[BsonRequired, BsonElement("t")]
+			public List<TestId> Tests { get; set; }
+			IReadOnlyList<TestId> ITestSuite.Tests => Tests;
+
+			private TestSuiteDocument()
+			{
+				Name = String.Empty;
+				Tests = new List<TestId>();
+			}
+
+			public TestSuiteDocument(StreamId streamId, string name, List<TestId> tests)
+			{
+				Id = TestSuiteId.GenerateNewId();
+				StreamId = streamId;
+				Name = name;
+				Tests = tests;
+			}
+		}
+
+		class SuiteTestData : ISuiteTestData
+		{
+			[BsonRequired, BsonElement("tid")]
+			public TestId Id { get; set; }
+
+			[BsonRequired, BsonElement("o")]
+			public TestOutcome Outcome { get; set; }
+
+			[BsonRequired, BsonElement("d")]
+			public TimeSpan Duration { get; set; }
+
+			[BsonRequired, BsonElement("uid")]
+			public string UID { get; set; }
+
+			public SuiteTestData(TestId id, TestOutcome outcome, TimeSpan duration, string uid)
+			{
+				Id = id;
+				Outcome = outcome;
+				Duration = duration;
+				UID = uid;
+			}
+		}
+
+		class TestDataRefDocument : ITestDataRef
+		{
+			[BsonRequired, BsonId]
+			public ObjectId Id { get; set; }
+
+			[BsonRequired,BsonElement("sid")]
+			public StreamId StreamId { get; set; }
+
+			[BsonRequired, BsonElement("m")]
+			public TestMetaId Metadata { get; set; }
+
+			[BsonRequired, BsonElement("bcl")]
+			public int BuildChangeList { get; set; }
+
+			[BsonRequired, BsonElement("tdid")]
+			public ObjectId TestDataId { get; set; }
+
+			[BsonRequired, BsonElement("d")]
+			public TimeSpan Duration { get; set; }
+
+			[BsonIgnoreIfNull, BsonElement("tid")]
+			public TestId? TestId { get; set; }
+
+			[BsonIgnoreIfNull, BsonElement("o")]
+			public TestOutcome? Outcome { get; set; }
+
+			[BsonIgnoreIfNull, BsonElement("suid")]
+			public TestSuiteId? SuiteId { get; set; }
+
+			[BsonIgnoreIfNull, BsonElement("sts")]
+			public List<SuiteTestData>? SuiteTests { get; set; }
+
+			IReadOnlyList<ISuiteTestData>? ITestDataRef.SuiteTests => SuiteTests;
+
+			private TestDataRefDocument()
+			{
+
+			}
+
+			public TestDataRefDocument(StreamId streamId)
+			{
+				Id = ObjectId.GenerateNewId();
+				StreamId = streamId;
+			}
+		}
+
 		/// <summary>
 		/// Information about a test data document
 		/// </summary>
@@ -55,27 +287,160 @@ namespace Horde.Build.Jobs.TestData
 		}
 
 		/// <summary>
-		/// The stream collection
+		/// The detailed test data collection
 		/// </summary>
 		readonly IMongoCollection<TestDataDocument> _testDataDocuments;
 
+		/// <summary>
+		/// Test platform collection
+		/// </summary>
+		readonly IMongoCollection<TestMetaDocument> _testMeta;
+
+		/// <summary>
+		/// Test collection
+		/// </summary>
+		readonly IMongoCollection<TestDocument> _tests;
+
+		/// <summary>
+		/// Test suite collection
+		/// </summary>
+		readonly IMongoCollection<TestSuiteDocument> _testSuites;
+
+		/// <summary>
+		/// Test suite collection
+		/// </summary>
+		readonly IMongoCollection<TestDataRefDocument> _testRefs;
+
+		readonly ILogger _logger;
+
+		readonly IOptionsMonitor<ServerSettings> _settings;
+
 		/// <summary>
 		/// Constructor
 		/// </summary>
-		/// <param name="mongoService">The database service instance</param>
-		public TestDataCollection(MongoService mongoService)
+		/// <param name="mongoService"></param>
+		/// <param name="logger"></param>
+		/// <param name="settings"></param>
+#pragma warning disable CS8618 // Disabled for feature flag: Non-nullable field must contain a non-null value when exiting constructor. Consider declaring as nullable.
+		public TestDataCollection(MongoService mongoService, ILogger<TestDataCollection> logger, IOptionsMonitor<ServerSettings> settings)
+#pragma warning restore CS8618 
 		{
+
+			_logger = logger;
+			_settings = settings;
+
 			List<MongoIndex<TestDataDocument>> indexes = new List<MongoIndex<TestDataDocument>>();
 			indexes.Add(keys => keys.Ascending(x => x.StreamId).Ascending(x => x.Change).Ascending(x => x.Key));
 			indexes.Add(keys => keys.Ascending(x => x.JobId).Ascending(x => x.StepId).Ascending(x => x.Key), unique: true);
 			_testDataDocuments = mongoService.GetCollection<TestDataDocument>("TestData", indexes);
+
+			if (_settings.CurrentValue.FeatureFlags.EnableTestDataV2)
+			{
+				List<MongoIndex<TestMetaDocument>> metaIndexes = new List<MongoIndex<TestMetaDocument>>();
+				metaIndexes.Add(keys => keys.Ascending(x => x.ProjectName).Ascending(x => x.Platforms).Ascending(x => x.Configurations));
+				_testMeta = mongoService.GetCollection<TestMetaDocument>("TestData.Meta", metaIndexes);
+
+				List<MongoIndex<TestDocument>> testIndexes = new List<MongoIndex<TestDocument>>();
+				testIndexes.Add(keys => keys.Ascending(x => x.StreamId));
+				_tests = mongoService.GetCollection<TestDocument>("TestData.Tests", testIndexes);
+
+				List<MongoIndex<TestSuiteDocument>> testSuiteIndexes = new List<MongoIndex<TestSuiteDocument>>();
+				testSuiteIndexes.Add(keys => keys.Ascending(x => x.StreamId));
+				_testSuites = mongoService.GetCollection<TestSuiteDocument>("TestData.TestSuites", testSuiteIndexes);
+
+				List<MongoIndex<TestDataRefDocument>> testRefIndexes = new List<MongoIndex<TestDataRefDocument>>();
+				testRefIndexes.Add(keys => keys.Ascending(x => x.StreamId).Ascending(x => x.Metadata).Descending(x => x.BuildChangeList).Ascending(x => x.TestId).Ascending(x => x.SuiteId));
+				_testRefs = mongoService.GetCollection<TestDataRefDocument>("TestData.TestRefs", testRefIndexes);
+			}
+		}
+
+		/// <summary>
+		/// Gets test data refs
+		/// </summary>
+		/// <param name="streamIds"></param>
+		/// <param name="metaIds"></param>
+		/// <param name="testIds"></param>
+		/// <param name="suiteIds"></param>
+		/// <param name="minCreateTime"></param>
+		/// <param name="maxCreateTime"></param>
+		/// <param name="minChange"></param>
+		/// <param name="maxChange"></param>
+		/// <returns></returns>
+		public async Task<List<ITestDataRef>> FindTestRefs(StreamId[] streamIds, TestMetaId[]? metaIds = null, TestId[]? testIds = null, TestSuiteId[]? suiteIds = null, DateTime? minCreateTime = null, DateTime? maxCreateTime = null, int? minChange = null, int? maxChange = null)
+		{
+
+			FilterDefinition<TestDataRefDocument> filter = FilterDefinition<TestDataRefDocument>.Empty;
+			FilterDefinitionBuilder<TestDataRefDocument> filterBuilder = Builders<TestDataRefDocument>.Filter;
+
+			filter &= Builders<TestDataRefDocument>.Filter.In(x => x.StreamId, streamIds);
+
+			if (minCreateTime != null)
+			{
+				ObjectId minTime = ObjectId.GenerateNewId(minCreateTime.Value);
+				filter &= filterBuilder.Gte(x => x.Id!, minTime);
+
+			}
+			if (maxCreateTime != null)
+			{
+				ObjectId maxTime = ObjectId.GenerateNewId(maxCreateTime.Value);
+				filter &= filterBuilder.Lte(x => x.Id!, maxTime);
+			}
+
+			if (metaIds != null && metaIds.Length > 0)
+			{
+				filter &= filterBuilder.In(x => x.Metadata, metaIds);
+			}
+
+			if (testIds != null && testIds.Length > 0)
+			{
+				filter &= filterBuilder.Ne(x => x.TestId, null);
+				filter &= filterBuilder.In(x => x.TestId!.Value, testIds);
+			}
+
+			if (suiteIds != null && suiteIds.Length > 0)
+			{
+				filter &= filterBuilder.Ne(x => x.SuiteId, null);
+				filter &= filterBuilder.In(x => x.SuiteId!.Value, suiteIds);
+			}
+
+			if (minChange != null)
+			{
+				filter &= filterBuilder.Gte(x => x.BuildChangeList, minChange);
+			}
+			if (maxChange != null)
+			{
+				filter &= filterBuilder.Lte(x => x.BuildChangeList, maxChange);
+			}
+
+			List<TestDataRefDocument> results;
+			using (IScope scope = GlobalTracer.Instance.BuildSpan("TestData.FindTestRefs").StartActive())
+			{
+				results = await _testRefs.Find(filter).ToListAsync();
+			}
+
+			return results.ConvertAll<ITestDataRef>(x => x);
+
 		}
 
 		/// <inheritdoc/>
 		public async Task<ITestData> AddAsync(IJob job, IJobStep step, string key, BsonDocument value)
 		{
+			// detailed test data
 			TestDataDocument newTestData = new TestDataDocument(job, step, key, value);
 			await _testDataDocuments.InsertOneAsync(newTestData);
+
+			try
+			{
+				if (_settings.CurrentValue.FeatureFlags.EnableTestDataV2)
+				{
+					await AddTestReportData(job, step, value, newTestData.Id);
+				}									
+			}
+			catch (Exception ex)
+			{
+				_logger.LogWarning(ex, "Exception while adding test data  report, jobId: {JobId} stepId: {StepId}", job.Id, step.Id);
+			}
+
 			return newTestData;
 		}
 
@@ -125,5 +490,506 @@ namespace Horde.Build.Jobs.TestData
 		{
 			await _testDataDocuments.DeleteOneAsync<TestDataDocument>(x => x.Id == id);
 		}
+
+		private async Task<List<TestMetaDocument>> FindTestMeta(TestMetaId[] metaIds)
+		{
+			FilterDefinitionBuilder<TestMetaDocument> filterBuilder = Builders<TestMetaDocument>.Filter;
+			FilterDefinition<TestMetaDocument> filter = FilterDefinition<TestMetaDocument>.Empty;
+			filter &= filterBuilder.In(x => x.Id, metaIds);
+
+			return await _testMeta.Find(filter).ToListAsync();
+		}
+
+		private async Task<List<TestDocument>> FindTests(StreamId[] streamIds)
+		{
+			FilterDefinitionBuilder<TestDocument> filterBuilder = Builders<TestDocument>.Filter;
+			FilterDefinition<TestDocument> filter = FilterDefinition<TestDocument>.Empty;
+			filter &= filterBuilder.In(x => x.StreamId, streamIds);
+
+			return await _tests.Find(filter).ToListAsync();
+		}
+
+		private async Task<List<TestSuiteDocument>> FindTestSuites(StreamId[] streamIds)
+		{
+			FilterDefinitionBuilder<TestSuiteDocument> filterBuilder = Builders<TestSuiteDocument>.Filter;
+			FilterDefinition<TestSuiteDocument> filter = FilterDefinition<TestSuiteDocument>.Empty;
+			filter &= filterBuilder.In(x => x.StreamId, streamIds);
+
+			return await _testSuites.Find(filter).ToListAsync();
+		}
+
+		/// <inheritdoc/>
+		public async Task<List<TestStream>> FindTestStreams(StreamId[] streamIds)
+		{
+			List<TestDocument> tests = await FindTests(streamIds);
+			List<TestSuiteDocument> suites = await FindTestSuites(streamIds);
+
+			List<TestMetaId> metaIds = tests.SelectMany(t => t.Metadata).Distinct().ToList();
+			List<TestMetaDocument> metaData = await FindTestMeta(metaIds.ToArray());
+			
+			List<TestStream> testStreams = new List<TestStream>();
+			for (int i = 0; i < streamIds.Length; i++)
+			{
+				StreamId id = streamIds[i];
+				List<ITest> streamTests = tests.Where(x => x.StreamId == id).ToList<ITest>();
+				List<ITestSuite> streamSuites = suites.Where(x => x.StreamId == id).ToList<ITestSuite>();
+				List<TestMetaId> streamMetaIds = streamTests.SelectMany(t => t.Metadata).Distinct().ToList();
+				List<ITestMeta> streamMeta = metaData.Where(x => streamMetaIds.Contains(x.Id)).ToList<ITestMeta>();
+				testStreams.Add(new TestStream(id, streamTests, streamSuites, streamMeta));
+			}
+
+			return testStreams;
+		}
+
+		/// <summary>
+		/// 
+		/// </summary>
+		/// <param name="metaData"></param>
+		/// <param name="job"></param>
+		/// <param name="step"></param>
+		/// <param name="testName"></param>
+		/// <returns></returns>
+		private async Task<TestMetaId?> AddTestMeta(Dictionary<string, string> metaData, IJob job, IJobStep step, string testName)
+		{
+			string? platform;
+			if (!metaData.TryGetValue("Platform", out platform))
+			{
+				_logger.LogWarning("Test {TestName} is missing platform metadata, jobId: {JobId} stepId: {StepId}", testName, job.Id, step.Id);
+				return null;
+			}
+
+			List<string> platforms = platform.Split('+', StringSplitOptions.TrimEntries).OrderBy(p => p).ToList();
+
+			string? buildTarget;
+			if (!metaData.TryGetValue("BuildTarget", out buildTarget))
+			{
+				_logger.LogWarning("Test {TestName} is missing build target metadata, jobId: {JobId} stepId: {StepId}", testName, job.Id, step.Id);
+				return null;
+			}
+
+			List<string> buildTargets = buildTarget.Split('+', StringSplitOptions.TrimEntries).OrderBy(b => b).ToList();
+
+			string? configuration;
+			if (!metaData.TryGetValue("Configuration", out configuration))
+			{
+				_logger.LogWarning("Test {TestName} is missing configuration metadata, jobId: {JobId} stepId: {StepId}", testName, job.Id, step.Id);
+				return null;
+			}
+
+			List<string> configurations = configuration.Split('+', StringSplitOptions.TrimEntries).OrderBy(c => c).ToList();
+
+			string? projectName;
+			if (!metaData.TryGetValue("Project", out projectName))
+			{
+				_logger.LogWarning("Test {TestName} is missing project metadata, jobId: {JobId} stepId: {StepId}", testName, job.Id, step.Id);
+				return null;
+			}
+
+			string? rhi;
+			metaData.TryGetValue("RHI", out rhi);
+			rhi = rhi ?? "default";
+
+			TestMetaDocument meta = new TestMetaDocument(platforms, configurations, buildTargets, projectName, rhi);
+
+			FilterDefinitionBuilder<TestMetaDocument> filterBuilder = Builders<TestMetaDocument>.Filter;
+			FilterDefinition<TestMetaDocument> filter = filterBuilder.Eq(x => x.ProjectName, meta.ProjectName);
+			filter &= filterBuilder.Eq(x => x.RHI, meta.RHI);
+			filter &= filterBuilder.Eq(x => x.Platforms, meta.Platforms);
+			filter &= filterBuilder.Eq(x => x.BuildTargets, meta.BuildTargets);
+			filter &= filterBuilder.Eq(x => x.Configurations, meta.Configurations);
+
+			TestMetaDocument? result = await _testMeta.Find(filter).FirstOrDefaultAsync();
+
+			if (result == null)
+			{
+				result = meta;
+				await _testMeta.InsertOneAsync(result);
+			}
+			
+			return result.Id;
+
+		}
+
+		async Task<TestDocument?> AddOrUpdateTest(IJob job, TestMetaId metaId, string testName, string? displayName = null)
+		{			
+			// register the test
+			TestDocument? test = await _tests.Find(x => x.StreamId == job.StreamId && x.Name == testName).FirstOrDefaultAsync();
+			if (test == null)
+			{
+				test = new TestDocument(job.StreamId, testName, new List<TestMetaId>() { metaId }, displayName);
+				await _tests.InsertOneAsync(test);
+			}
+			else
+			{
+				if (!test.Metadata.Contains(metaId))
+				{
+					test.Metadata.Add(metaId);
+
+					FilterDefinitionBuilder<TestDocument> ebuilder = Builders<TestDocument>.Filter;
+					FilterDefinition<TestDocument> efilter = ebuilder.Eq(x => x.Id, test.Id);
+
+					UpdateDefinitionBuilder<TestDocument> updateBuilder = Builders<TestDocument>.Update;
+					UpdateDefinition<TestDocument> update = updateBuilder.Set(x => x.Metadata, test.Metadata);
+					await _tests.UpdateOneAsync(efilter, update);
+				}
+			}
+
+			return test;
+		}
+
+		private async Task<TestSuiteDocument?> AddOrUpdateTestSuite(IJob job, string suiteName, List<TestId> testIds)
+		{
+			TestSuiteDocument? suite = await _testSuites.Find(x => x.StreamId == job.StreamId && x.Name == suiteName).FirstOrDefaultAsync();
+
+			if (suite == null)
+			{
+				suite = new TestSuiteDocument(job.StreamId, suiteName, testIds);
+				await _testSuites.InsertOneAsync(suite);
+			}
+			else
+			{
+				List<TestId> newIds = suite.Tests.Union(testIds).ToList();
+
+				if (newIds.Count != suite.Tests.Count)
+				{
+					suite.Tests = newIds;
+
+					FilterDefinitionBuilder<TestSuiteDocument> ebuilder = Builders<TestSuiteDocument>.Filter;
+					FilterDefinition<TestSuiteDocument> efilter = ebuilder.Eq(x => x.Id, suite.Id);
+
+					UpdateDefinitionBuilder<TestSuiteDocument> updateBuilder = Builders<TestSuiteDocument>.Update;
+					UpdateDefinition<TestSuiteDocument> update = updateBuilder.Set(x => x.Tests, suite.Tests);
+					await _testSuites.UpdateOneAsync(efilter, update);
+				}
+			}
+
+			return suite;
+		}
+
+		async Task AddTestReportData(IJob job, IJobStep step, BsonDocument rootDoc, ObjectId testDataId)
+		{
+			int version = rootDoc.GetValue("Version", new BsonInt32(0)).AsInt32;
+			if (version < 1)
+			{
+				return;
+			}
+
+			List<AutomatedTestSessionData> sessions = new List<AutomatedTestSessionData>();
+			List<UnrealAutomatedTestData> tests = new List<UnrealAutomatedTestData>();
+
+			BsonArray items = rootDoc.GetValue("Items").AsBsonArray;
+			foreach (BsonValue item in items.ToList())
+			{
+				if (!item.IsBsonDocument)
+				{
+					continue;
+				}
+
+				BsonDocument testData = item.AsBsonDocument;
+
+				BsonValue? value;
+				if (!testData.TryGetPropertyValue("Data.Type", BsonType.String, out value))
+				{
+					continue;
+				}
+
+				string type = value.AsString;
+				BsonDocument data = testData.GetValue("Data").AsBsonDocument;
+
+				if (type == "Simple Report")
+				{
+					await AddSimpleReportData(job, step, testDataId, data);
+					continue;
+				}
+
+				if (type == "Automated Test Session")
+				{
+					sessions.Add(BsonSerializer.Deserialize<AutomatedTestSessionData>(data));
+				}
+
+				if (type == "Unreal Automated Tests")
+				{
+					tests.Add(BsonSerializer.Deserialize<UnrealAutomatedTestData>(data));
+				}
+			}
+
+			if (sessions.Count > 0)
+			{
+				await AddTestSessionReportData(job, step, testDataId, sessions, tests);
+			}
+		}
+
+		// --- (Minimal) Parsing of Automation framework tests, geared towards generating indexes for aggregate queries
+
+		class AutomatedTestSessionData
+		{
+			public class SessionTest
+			{
+				public string Name { get; set; } = String.Empty;
+				public string TestUID { get; set; } = String.Empty;
+				public string Suite { get; set; } = String.Empty;
+				public string State { get; set; } = String.Empty;
+				public List<string> DeviceAppInstanceName { get; set; } = new List<string>();
+				public int ErrorCount { get; set; } = 0;
+				public int WarningCount { get; set; } = 0;
+				public double TimeElapseSec { get; set; } = 0;
+			}
+
+			public class TestSessionInfoType
+			{
+				public double TimeElapseSec { get; set; } = 0;
+				public Dictionary<string, SessionTest> Tests { get; set; } = new Dictionary<string, SessionTest>();
+				public string TestResultsTestDataUID { get; set; } = String.Empty;
+			}
+
+			public string Type { get; set; } = String.Empty;
+			public string Name { get; set; } = String.Empty;
+
+			public TestSessionInfoType? TestSessionInfo { get; set; }
+			public Dictionary<string, string> Metadata { get; set; } = new Dictionary<string, string>();
+
+		}
+
+		class UnrealAutomatedTestData
+		{
+			public class UnrealTest
+			{
+				public string TestDisplayName { get; set; } = String.Empty;
+				public string FullTestPath { get; set; } = String.Empty;
+				public string State { get; set; } = String.Empty;
+				public string DeviceInstance { get; set; } = String.Empty;
+				public int Errors { get; set; } = 0;
+				public int Warnings { get; set; } = 0;				
+			}
+
+			public string Type { get; set; } = String.Empty;
+			public double TotalDurationSeconds { get; set; } = 0;
+			public List<UnrealTest> Tests { get; set; } = new List<UnrealTest>();
+			public Dictionary<string, string> Metadata { get; set; } = new Dictionary<string, string>();
+
+		}
+
+		// Simple Test Reports
+
+		class TestRole
+		{
+			public string Type { get; set; } = String.Empty;
+			public string Platform { get; set; } = String.Empty;
+			public string Configuration { get; set; } = String.Empty;
+		}
+
+		class SimpleTestData
+		{			
+			public string TestName { get; set; } = String.Empty;
+			public string Description { get; set; } = String.Empty;
+			public string ReportCreatedOn { get; set; } = String.Empty;
+			public double TotalDurationSeconds { get; set; } = 0;
+			public bool HasSucceeded { get; set; } = false;
+			public string Status { get; set; } = String.Empty;
+			public string URLLink { get; set; } = String.Empty;
+			public int BuildChangeList { get; set; } = 0;
+			public TestRole? MainRole { get; set; }
+			public List<TestRole>? Roles { get; set; } = new List<TestRole>();
+			public string TestResult { get; set; } = "Invalid";
+			public List<string> Logs { get; set; } = new List<string>();
+			public List<string> Errors { get; set; } = new List<string>();
+			public List<string> Warnings { get; set; } = new List<string>();
+			public Dictionary<string, string> Metadata { get; set; } = new Dictionary<string, string>();
+		}
+
+		async Task AddSimpleReportData(IJob job, IJobStep step, ObjectId testDataId, BsonDocument testDoc)
+		{
+			try
+			{
+				SimpleTestData testData = BsonSerializer.Deserialize<SimpleTestData>(testDoc);
+
+				TestMetaId? metaId = await AddTestMeta(testData.Metadata, job, step, testData.TestName);
+
+				if (metaId == null)
+				{
+					throw new Exception($"Unable to add or update simple test report for {testData.TestName}, unable to generate meta data");
+				}
+
+				TestDocument? test = await AddOrUpdateTest(job, metaId.Value, testData.TestName);
+
+				if (test == null)
+				{
+					throw new Exception($"Unable to add or update simple test report for {testData.TestName}, unable to generate test document");
+				}
+
+				TestDataRefDocument testRef = new TestDataRefDocument(job.StreamId);
+				testRef.StreamId = test.StreamId;
+				testRef.Metadata = metaId.Value;
+				testRef.BuildChangeList = testData.BuildChangeList;
+				testRef.Duration = TimeSpan.FromSeconds(testData.TotalDurationSeconds);
+				testRef.TestId = test.Id;
+				testRef.TestDataId = testDataId;
+				testRef.Outcome = TestOutcome.Unspecified;
+
+				switch (testData.TestResult)
+				{
+					case "Passed":
+						testRef.Outcome = TestOutcome.Success;
+						break;
+					case "Failed":
+					case "Cancelled":
+					case "TimedOut":
+						testRef.Outcome = TestOutcome.Failure;
+						break;
+				}
+
+				await _testRefs.InsertOneAsync(testRef);
+
+			}
+			catch (Exception ex)
+			{
+				_logger.LogWarning(ex, "Exception while adding test data simple report, jobId: {JobId} stepId: {StepId}", job.Id, step.Id);
+			}
+		}
+
+		async Task AddTestSessionReportData(IJob job, IJobStep step, ObjectId testDataId, List<AutomatedTestSessionData> sessions, List<UnrealAutomatedTestData> tests)
+		{
+			HashSet<string> suites = new HashSet<string>();
+
+			for (int i = 0; i < sessions.Count; i++)
+			{
+				AutomatedTestSessionData session = sessions[i];
+
+				AutomatedTestSessionData.TestSessionInfoType? info = session.TestSessionInfo;
+				if (info == null)
+				{
+					return;
+				}
+
+				foreach (KeyValuePair<string, AutomatedTestSessionData.SessionTest> item in info.Tests)
+				{
+					AutomatedTestSessionData.SessionTest test = item.Value;
+
+					if (!String.IsNullOrEmpty(test.Suite))
+					{
+						suites.Add(test.Suite);
+					}
+				}
+			}
+
+			if (suites.Count == 0)
+			{
+				return;
+			}
+
+			foreach (string suite in suites)
+			{
+				for (int i = 0; i < sessions.Count; i++)
+				{
+					HashSet<TestId> suiteTestIds = new HashSet<TestId>();
+					Dictionary<string, TestDocument> suiteTestDocuments = new Dictionary<string, TestDocument>();
+
+					AutomatedTestSessionData session = sessions[i];
+					AutomatedTestSessionData.TestSessionInfoType? info = session.TestSessionInfo;
+					if (info == null)
+					{
+						continue;
+					}
+
+					TestMetaId? metaId = await AddTestMeta(session.Metadata, job, step, "Automated Test Session");
+
+					if (metaId == null)
+					{
+						throw new Exception($"Unable to add or update simple test report for automated test session, unable to generate meta data");
+					}
+
+					List<AutomatedTestSessionData.SessionTest> suiteTests = info.Tests.Where(x => x.Value.Suite == suite).Select(x => x.Value).ToList();
+
+					foreach (AutomatedTestSessionData.SessionTest test in suiteTests)
+					{
+
+						for (int j = 0; j < tests.Count; j++)
+						{
+							UnrealAutomatedTestData testData = tests[j];
+
+							// attempt to find the unreal test, which holds some optional data
+							UnrealAutomatedTestData.UnrealTest? utest = null;
+							for (int k = 0; k < testData.Tests.Count; k++)
+							{
+								UnrealAutomatedTestData.UnrealTest u = testData.Tests[k];
+								if (test.Name == u.FullTestPath && test.DeviceAppInstanceName.Contains(u.DeviceInstance, StringComparer.OrdinalIgnoreCase))
+								{
+									utest = u;
+									break;
+								}
+							}
+
+							string? displayName = null;
+							if (utest != null && !String.IsNullOrEmpty(utest.TestDisplayName))
+							{
+								displayName = utest.TestDisplayName;
+							}
+
+							TestDocument? testDoc = await AddOrUpdateTest(job, metaId!.Value, test.Name, displayName);
+
+							if (testDoc == null)
+							{
+								_logger.LogWarning("Skipping adding suite test {TestName} metaId: {MetaId} testId: {TestId} jobId: {JobId} stepId: {StepId}", test.Name, metaId?.Value, testDoc?.Id, job.Id, step.Id);
+								continue;
+							}
+
+							suiteTestDocuments.Add(test.Name, testDoc);
+							suiteTestIds.Add(testDoc.Id);
+						}
+					}
+
+					// register test suite
+					if (suiteTestIds.Count > 0)
+					{
+						TestSuiteDocument? testSuite = await AddOrUpdateTestSuite(job, suite, suiteTestIds.ToList());
+
+						if (testSuite == null)
+						{
+							continue;
+						}
+
+						TestDataRefDocument testRef = new TestDataRefDocument(job.StreamId);
+						testRef.StreamId = job.StreamId;
+						testRef.Metadata = metaId.Value;
+						testRef.BuildChangeList = job.Change;
+						testRef.Duration = TimeSpan.FromSeconds(session.TestSessionInfo!.TimeElapseSec);
+						testRef.TestDataId = testDataId;
+						testRef.Outcome = TestOutcome.Unspecified;
+						testRef.SuiteId = testSuite.Id;
+						testRef.SuiteTests = new List<SuiteTestData>();
+
+						// populate the suite tests
+
+						foreach (AutomatedTestSessionData.SessionTest test in suiteTests)
+						{
+
+							TestDocument document = suiteTestDocuments[test.Name];
+							TestOutcome outcome = TestOutcome.Unspecified;
+
+							switch (test.State)
+							{
+								case "Success":
+									outcome = TestOutcome.Success;
+									break;
+								case "Skipped":
+									outcome = TestOutcome.Skipped;
+									break;
+								case "Fail":
+									outcome = TestOutcome.Failure;
+									break;
+							}
+
+							testRef.SuiteTests.Add(new SuiteTestData(document.Id, outcome, TimeSpan.FromSeconds(test.TimeElapseSec), test.TestUID));
+						}
+
+						await _testRefs.InsertOneAsync(testRef);
+					}
+					else
+					{
+						_logger.LogWarning("Skipping suite {SuiteName} with not test ids, jobId: {JobId} stepId: {StepId}", suite, job.Id, step.Id);
+					}
+				}
+			}
+		}
 	}
 }
diff --git a/Engine/Source/Programs/Horde/Horde.Build/Jobs/TestData/TestDataController.Messages.cs b/Engine/Source/Programs/Horde/Horde.Build/Jobs/TestData/TestDataController.Messages.cs
deleted file mode 100644
diff --git a/Engine/Source/Programs/Horde/Horde.Build/Jobs/TestData/TestDataController.cs b/Engine/Source/Programs/Horde/Horde.Build/Jobs/TestData/TestDataController.cs
--- a/Engine/Source/Programs/Horde/Horde.Build/Jobs/TestData/TestDataController.cs
+++ b/Engine/Source/Programs/Horde/Horde.Build/Jobs/TestData/TestDataController.cs
@@ -1,6 +1,8 @@
 // Copyright Epic Games, Inc. All Rights Reserved.
 
+using System;
 using System.Collections.Generic;
+using System.Linq;
 using System.Threading.Tasks;
 using Horde.Build.Acls;
 using Horde.Build.Streams;
@@ -11,8 +13,12 @@ using MongoDB.Bson;
 
 namespace Horde.Build.Jobs.TestData
 {
+	
 	using JobId = ObjectId<IJob>;
 	using StreamId = StringId<IStream>;
+	using TestId = ObjectId<ITest>;
+	using TestSuiteId = ObjectId<ITestSuite>;
+	using TestMetaId = ObjectId<ITestMeta>;
 
 	/// <summary>
 	/// Controller for the /api/v1/testdata endpoint
@@ -32,17 +38,120 @@ namespace Horde.Build.Jobs.TestData
 		/// </summary>
 		private readonly ITestDataCollection _testDataCollection;
 
+		readonly TestDataService _testDataService;
+
+		readonly StreamService _streamService;
+
+
 		/// <summary>
 		/// Constructor
 		/// </summary>
-		/// <param name="jobService">The job service singleton</param>
-		/// <param name="testDataCollection">Collection of test data documents</param>
-		public TestDataController(JobService jobService, ITestDataCollection testDataCollection)
+		/// <param name="testDataService"></param>
+		/// <param name="streamService"></param>
+		/// <param name="jobService"></param>
+		/// <param name="testDataCollection"></param>
+		public TestDataController(TestDataService testDataService, StreamService streamService, JobService jobService, ITestDataCollection testDataCollection)
 		{
 			_jobService = jobService;
 			_testDataCollection = testDataCollection;
+			_testDataService = testDataService;
+			_streamService = streamService;
+		}
+
+		/// <summary>
+		/// Get stream test data for the provided ids
+		/// </summary>
+		/// <param name="streamIds"></param>
+		/// <returns></returns>
+		[HttpGet]
+		[Route("/api/v2/testdata/streams")]
+		[ProducesResponseType(typeof(List<GetTestStreamResponse>), 200)]
+		public async Task<ActionResult<List<GetTestStreamResponse>>> GetTestStreamsAsync([FromQuery(Name = "Id")] string[] streamIds)
+		{
+			StreamId[] streamIdValues = Array.ConvertAll(streamIds, x => new StreamId(x));
+
+			List<StreamId> queryStreams = new List<StreamId>();
+			List<GetTestStreamResponse> responses = new List<GetTestStreamResponse>();
+
+			// authorize streams
+			foreach (StreamId streamId in streamIdValues)
+			{
+				IStream? stream = await _streamService.GetCachedStream(streamId);
+				if (stream != null)
+				{
+					if (await _streamService.AuthorizeAsync(stream, AclAction.ViewJob, User, null))
+					{
+						queryStreams.Add(stream.Id);
+					}
+				}
+			}
+
+			if (queryStreams.Count == 0)
+			{
+				return responses;
+			}
+
+			List<TestStream> streams = await _testDataService.FindTestStreams(queryStreams.ToArray());
+			streams.ForEach(s => responses.Add(new GetTestStreamResponse(s.StreamId, s.Tests, s.TestSuites, s.TestMeta)));
+			return responses;
 		}
 
+		/// <summary>
+		/// Gets test data refs 
+		/// </summary>
+		/// <param name="streamIds"></param>
+		/// <param name="testIds"></param>
+		/// <param name="suiteIds"></param>
+		/// <param name="platformIds"></param>
+		/// <param name="minCreateTime"></param>
+		/// <param name="maxCreateTime"></param>
+		/// <param name="minChange"></param>
+		/// <param name="maxChange"></param>
+		/// <returns></returns>
+		[HttpGet]
+		[Route("/api/v2/testdata/refs")]
+		[ProducesResponseType(typeof(List<GetTestDataRefResponse>), 200)]
+		public async Task<ActionResult<List<GetTestDataRefResponse>>> GetTestDataRefAsync(
+			[FromQuery(Name = "Id")] string[] streamIds,
+			[FromQuery(Name = "Tid")] string[]? testIds = null,
+			[FromQuery(Name = "Sid")] string[]? suiteIds = null,
+			[FromQuery(Name = "Pid")] string[]? platformIds = null,
+			[FromQuery] DateTimeOffset? minCreateTime = null,
+			[FromQuery] DateTimeOffset? maxCreateTime = null,
+			[FromQuery] int? minChange = null, 
+			[FromQuery] int? maxChange = null)
+		{
+			StreamId[] streamIdValues = Array.ConvertAll(streamIds, x => new StreamId(x));
+
+			List<StreamId> queryStreams = new List<StreamId>();
+			List<GetTestDataRefResponse> responses = new List<GetTestDataRefResponse>();
+
+			// authorize streams
+			foreach (StreamId streamId in streamIdValues)
+			{
+				IStream? stream = await _streamService.GetCachedStream(streamId);
+				if (stream != null)
+				{
+					if (await _streamService.AuthorizeAsync(stream, AclAction.ViewJob, User, null))
+					{
+						queryStreams.Add(stream.Id);
+					}
+				}
+			}
+
+			if (queryStreams.Count == 0)
+			{
+				return responses;
+			}
+
+			List<ITestDataRef> dataRefs = await _testDataService.FindTestRefs(queryStreams.ToArray(), platformIds, testIds, suiteIds, minCreateTime?.UtcDateTime, maxCreateTime?.UtcDateTime, minChange, maxChange);
+
+			dataRefs.ForEach(d => responses.Add(new GetTestDataRefResponse(d)));
+
+			return responses;
+		}
+
+
 		/// <summary>
 		/// Creates a new TestData document
 		/// </summary>
diff --git a/Engine/Source/Programs/Horde/Horde.Build/Jobs/TestData/TestDataMessages.cs b/Engine/Source/Programs/Horde/Horde.Build/Jobs/TestData/TestDataMessages.cs
new file mode 100644
--- /dev/null
+++ b/Engine/Source/Programs/Horde/Horde.Build/Jobs/TestData/TestDataMessages.cs
@@ -0,0 +1,392 @@
+// Copyright Epic Games, Inc. All Rights Reserved.
+
+using System;
+using System.Collections.Generic;
+using System.ComponentModel.DataAnnotations;
+using System.Linq;
+using Horde.Build.Streams;
+using Horde.Build.Utilities;
+using MongoDB.Bson.Serialization;
+
+namespace Horde.Build.Jobs.TestData
+{
+	using StreamId = StringId<IStream>;
+	using TestId = ObjectId<ITest>;
+
+	/// <summary>
+	/// Response object describing test data to store
+	/// </summary>
+	public class CreateTestDataRequest
+	{
+		/// <summary>
+		/// The job which produced the data
+		/// </summary>
+		[Required]
+		public string JobId { get; set; } = String.Empty;
+
+		/// <summary>
+		/// The step that ran
+		/// </summary>
+		[Required]
+		public string StepId { get; set; } = String.Empty;
+
+		/// <summary>
+		/// Key used to identify the particular data
+		/// </summary>
+		public string Key { get; set; } = String.Empty;
+
+		/// <summary>
+		/// The data stored for this test
+		/// </summary>
+		[Required]
+		public Dictionary<string, object> Data { get; set; } = new Dictionary<string, object>();
+	}
+
+	/// <summary>
+	/// Response object describing the created document
+	/// </summary>
+	public class CreateTestDataResponse
+	{
+		/// <summary>
+		/// The id for the new document
+		/// </summary>
+		public string Id { get; set; }
+
+		/// <summary>
+		/// Constructor
+		/// </summary>
+		/// <param name="id">Id of the new document</param>
+		public CreateTestDataResponse(string id)
+		{
+			Id = id;
+		}
+	}
+
+	/// <summary>
+	/// Response object describing test results
+	/// </summary>
+	public class GetTestDataResponse
+	{
+		/// <summary>
+		/// Unique id of the test data
+		/// </summary>
+		public string Id { get; set; }
+
+		/// <summary>
+		/// Stream that generated the test data
+		/// </summary>
+		public string StreamId { get; set; }
+
+		/// <summary>
+		/// The template reference id
+		/// </summary>
+		public string TemplateRefId { get; set; }
+
+		/// <summary>
+		/// The job which produced the data
+		/// </summary>
+		public string JobId { get; set; }
+
+		/// <summary>
+		/// The step that ran
+		/// </summary>
+		public string StepId { get; set; }
+
+		/// <summary>
+		/// The changelist number that contained the data
+		/// </summary>
+		public int Change { get; set; }
+
+		/// <summary>
+		/// Key used to identify the particular data
+		/// </summary>
+		public string Key { get; set; }
+
+		/// <summary>
+		/// The data stored for this test
+		/// </summary>
+		public Dictionary<string, object> Data { get; set; }
+
+		/// <summary>
+		/// Constructor
+		/// </summary>
+		/// <param name="testData">Test data to construct from</param>
+		internal GetTestDataResponse(ITestData testData)
+		{
+			Id = testData.Id.ToString();
+			StreamId = testData.StreamId.ToString();
+			TemplateRefId = testData.TemplateRefId.ToString();
+			JobId = testData.JobId.ToString();
+			StepId = testData.StepId.ToString();
+			Change = testData.Change;
+			Key = testData.Key;
+			Data = BsonSerializer.Deserialize<Dictionary<string, object>>(testData.Data);
+		}
+	}
+
+	/// <summary>
+	/// A test emvironment running in a stream
+	/// </summary>
+	public class GetTestMetaResponse
+	{
+		/// <summary>
+		/// Meta unique id for environment 
+		/// </summary>
+		public string Id { get; set; }
+
+		/// <summary>
+		/// The platforms in the environment
+		/// </summary>
+		public List<string> Platforms { get; set; }
+
+		/// <summary>
+		/// The build configurations being tested
+		/// </summary>
+		public List<string> Configurations { get; set; }
+
+		/// <summary>
+		/// The build targets being tested
+		/// </summary>
+		public List<string> BuildTargets { get; set; }
+
+		/// <summary>
+		/// The test project name
+		/// </summary>
+		public string ProjectName { get; set; }
+
+		/// <summary>
+		/// The rendering hardware interface being used with the test
+		/// </summary>
+		public string RHI { get; set; }
+
+		internal GetTestMetaResponse(ITestMeta meta)
+		{
+			Id = meta.Id.ToString();
+			Platforms = meta.Platforms.Select(p => p).ToList();
+			Configurations = meta.Configurations.Select(p => p).ToList();
+			BuildTargets = meta.BuildTargets.Select(p => p).ToList();
+			ProjectName = meta.ProjectName;
+			RHI = meta.RHI;	
+		}
+	}
+
+	/// <summary>
+	/// A test that runs in a stream
+	/// </summary>
+	public class GetTestResponse
+	{
+		/// <summary>
+		/// The name of the test 
+		/// </summary>
+		public string Name { get; set; }
+
+		/// <summary>
+		/// The name of the test 
+		/// </summary>
+		public string? DisplayName { get; set; }
+
+		/// <summary>
+		/// The name of the test 
+		/// </summary>
+		public string StreamId { get; set; }
+
+		/// <summary>
+		/// The platforms the test runs on
+		/// </summary>
+		public List<string> Metadata { get; set; }
+
+		internal GetTestResponse(ITest test)
+		{
+			Name = test.Name;
+			DisplayName = test.DisplayName;
+			StreamId = test.StreamId.ToString();
+			Metadata = test.Metadata.Select(x => x.ToString()).ToList();
+		}
+	}
+
+	/// <summary>
+	/// A test suite that runs in a stream, contain subtests
+	/// </summary>
+	public class GetTestSuiteResponse
+	{
+		/// <summary>
+		/// The name of the test suite
+		/// </summary>
+		public string Name { get; set; }
+
+		/// <summary>
+		/// The name of the test 
+		/// </summary>
+		public string StreamId { get; set; }
+
+		/// <summary>
+		/// The tests in the suite
+		/// </summary>
+		public List<TestId> Tests { get; set; }
+
+		internal GetTestSuiteResponse(ITestSuite suite)
+		{
+			Name = suite.Name;
+			StreamId = suite.StreamId.ToString();
+			Tests = suite.Tests.ToList();
+		}
+	}
+
+	/// <summary>
+	/// Response object describing test results
+	/// </summary>
+	public class GetTestStreamResponse
+	{
+		/// <summary>
+		/// The stream id
+		/// </summary>
+		public string StreamId { get; set; }
+
+		/// <summary>
+		/// Individual tests which run in the stream
+		/// </summary>
+		public List<GetTestResponse> Tests { get; set; }
+
+		/// <summary>
+		/// Test suites that run in the stream
+		/// </summary>
+		public List<GetTestSuiteResponse> TestSuites { get; set; }
+
+		/// <summary>
+		/// Test suites that run in the stream
+		/// </summary>
+		public List<GetTestMetaResponse> TestMetadata { get; set; }
+
+		/// <summary>
+		/// Constructor
+		/// </summary>
+		/// <param name="streamId"></param>
+		/// <param name="tests"></param>
+		/// <param name="suites"></param>
+		/// <param name="metaData"></param>
+		internal GetTestStreamResponse(StreamId streamId, List<ITest> tests, List<ITestSuite> suites, List<ITestMeta> metaData)
+		{
+			StreamId = streamId.ToString();
+
+			Tests = tests.Select(t => new GetTestResponse(t)).ToList();
+			TestSuites = suites.Select(t => new GetTestSuiteResponse(t)).ToList();
+			TestMetadata = metaData.Select(m => new GetTestMetaResponse(m)).ToList();
+		}
+
+	}
+
+	/// <summary>
+	/// Suite test data
+	/// </summary>
+	public class GetSuiteTestDataResponse
+	{
+		/// <summary>
+		/// The test id
+		/// </summary>
+		public string Id { get; set; }
+
+		/// <summary>
+		/// The ourcome of the suite test
+		/// </summary>
+		public TestOutcome Outcome { get; set; }
+
+		/// <summary>
+		/// How long the suite test ran
+		/// </summary>
+		public TimeSpan Duration { get; set; }
+
+		/// <summary>
+		/// Test UID for looking up in test details
+		/// </summary>
+		public string UID { get; set; }
+
+		/// <summary>
+		/// Constructor
+		/// </summary>
+		/// <param name="data"></param>
+		public GetSuiteTestDataResponse(ISuiteTestData data)
+		{
+			Id = data.Id.ToString();
+			Outcome = data.Outcome;
+			Duration = data.Duration;
+			UID =data.UID;
+		}
+	}
+
+
+	/// <summary>
+	/// Data ref 
+	/// </summary>
+	public class GetTestDataRefResponse
+	{
+		/// <summary>
+		/// The associated stream
+		/// </summary>
+		public string StreamId { get; set; }
+
+		/// <summary>
+		/// The full details test data for this ref
+		/// </summary>
+		public string TestDataId { get; set; }
+
+		/// <summary>
+		/// How long the test ran
+		/// </summary>
+		public TimeSpan Duration { get; set; }
+
+		/// <summary>
+		/// The build changelist upon which the test ran, may not correspond to the job changelist
+		/// </summary>
+		public int BuildChangeList { get; set; }
+
+		/// <summary>
+		/// The platform the test ran on 
+		/// </summary>
+		public string MetaId { get; set; }
+
+		/// <summary>
+		/// The test id in stream
+		/// </summary>
+		public string? TestId { get; set; }
+
+		/// <summary>
+		/// The outcome of the test
+		/// </summary>
+		public TestOutcome? Outcome { get; set; }
+
+		/// <summary>
+		/// The if of the stream test suite
+		/// </summary>
+		public string? SuiteId { get; set; }
+
+		/// <summary>
+		/// The suite tests
+		/// </summary>
+		public List<GetSuiteTestDataResponse>? SuiteTests { get; set; }
+
+
+		/// <summary>
+		/// Constructor
+		/// </summary>
+		/// <param name="testData"></param>
+		public GetTestDataRefResponse(ITestDataRef testData)
+		{
+			StreamId = testData.StreamId.ToString();
+			TestDataId = testData.TestDataId.ToString();
+			Duration = testData.Duration;
+			BuildChangeList = testData.BuildChangeList;
+			MetaId = testData.Metadata.ToString();
+			TestId = testData.TestId?.ToString();
+			Outcome = testData.Outcome;
+			SuiteId = testData.SuiteId?.ToString();
+			if (testData.SuiteTests != null)
+			{
+				SuiteTests = new List<GetSuiteTestDataResponse>();
+				foreach (ISuiteTestData test in testData.SuiteTests)
+				{
+					SuiteTests.Add(new GetSuiteTestDataResponse(test));
+				}
+			}			
+		}
+	}
+}
diff --git a/Engine/Source/Programs/Horde/Horde.Build/Jobs/TestData/TestDataService.cs b/Engine/Source/Programs/Horde/Horde.Build/Jobs/TestData/TestDataService.cs
new file mode 100644
--- /dev/null
+++ b/Engine/Source/Programs/Horde/Horde.Build/Jobs/TestData/TestDataService.cs
@@ -0,0 +1,74 @@
+// Copyright Epic Games, Inc. All Rights Reserved.
+
+using System;
+using System.Collections.Generic;
+using System.Threading.Tasks;
+using EpicGames.Core;
+using Horde.Build.Streams;
+using Horde.Build.Utilities;
+using Microsoft.Extensions.Options;
+
+namespace Horde.Build.Jobs.TestData
+{
+	
+	using StreamId = StringId<IStream>;
+	using TestId = ObjectId<ITest>;
+	using TestSuiteId = ObjectId<ITestSuite>;
+	using TestMetaId = ObjectId<ITestMeta>;
+
+
+	/// <summary>
+	/// Device management service
+	/// </summary>
+	public sealed class TestDataService : IDisposable
+	{
+		
+		readonly ITestDataCollection _testData;
+
+
+		/// <summary>
+		/// Device service constructor
+		/// </summary>
+		public TestDataService(ITestDataCollection testData)
+		{
+			_testData = testData;			
+		}
+
+		/// <inheritdoc/>
+		public void Dispose()
+		{
+		}
+
+
+		/// <summary>
+		/// Finds test streams
+		/// </summary>
+		/// <param name="streamIds"></param>
+		/// <returns></returns>
+		public async Task<List<TestStream>> FindTestStreams(StreamId[] streamIds)
+		{
+			return await _testData.FindTestStreams(streamIds);
+		}
+
+		/// <summary>
+		/// Gets test data refs
+		/// </summary>
+		/// <param name="streamIds"></param>
+		/// <param name="metaIds"></param>
+		/// <param name="testIds"></param>
+		/// <param name="suiteIds"></param>
+		/// <param name="minCreateTime"></param>
+		/// <param name="maxCreateTime"></param>
+		/// <param name="minChange"></param>
+		/// <param name="maxChange"></param>
+		/// <returns></returns>
+		public async Task<List<ITestDataRef>> FindTestRefs(StreamId[] streamIds, string[]? metaIds = null, string[]? testIds = null, string[]? suiteIds = null, DateTime? minCreateTime = null, DateTime? maxCreateTime = null, int? minChange = null, int? maxChange = null)
+		{
+			TestMetaId[]? mids = metaIds?.ConvertAll(x => new TestMetaId(x)).ToArray();
+			TestId[]? tids = testIds?.ConvertAll(x => new TestId(x)).ToArray();
+			TestSuiteId[]? sids = suiteIds?.ConvertAll(x => new TestSuiteId(x)).ToArray();
+
+			return await _testData.FindTestRefs(streamIds, mids, tids, sids, minCreateTime, maxCreateTime, minChange, maxChange);
+		}
+	}
+}
diff --git a/Engine/Source/Programs/Horde/Horde.Build/ServerSettings.cs b/Engine/Source/Programs/Horde/Horde.Build/ServerSettings.cs
--- a/Engine/Source/Programs/Horde/Horde.Build/ServerSettings.cs
+++ b/Engine/Source/Programs/Horde/Horde.Build/ServerSettings.cs
@@ -179,6 +179,12 @@ namespace Horde.Build
 		/// Limit concurrent log chunk writes and await them to reduce mem and I/O usage
 		/// </summary>
 		public bool LimitConcurrentLogChunkWriting { get; set; } = false;
+
+		/// <summary>
+		/// Enable new test data collections and service, feature flag added 22/11/7
+		/// </summary>
+		public bool EnableTestDataV2 { get; set; } = false;
+
 	}
 
 	/// <summary>
diff --git a/Engine/Source/Programs/Horde/Horde.Build/Startup.cs b/Engine/Source/Programs/Horde/Horde.Build/Startup.cs
--- a/Engine/Source/Programs/Horde/Horde.Build/Startup.cs
+++ b/Engine/Source/Programs/Horde/Horde.Build/Startup.cs
@@ -425,6 +425,7 @@ namespace Horde.Build
 			services.AddSingleton<DeviceService>();						
 			services.AddSingleton<NoticeService>();
 			services.AddSingleton<StorageService>();
+			services.AddSingleton<TestDataService>();
 
 			if (settings.JiraUrl != null)
 			{
@@ -620,7 +621,7 @@ namespace Horde.Build
 				services.AddHostedService(provider => provider.GetRequiredService<SlackNotificationSink>());
 				services.AddHostedService<ConfigUpdateService>();
 				services.AddHostedService<TelemetryService>();
-				services.AddHostedService(provider => provider.GetRequiredService<DeviceService>());
+				services.AddHostedService(provider => provider.GetRequiredService<DeviceService>());				
 			}
 
 			services.AddHostedService(provider => provider.GetRequiredService<IExternalIssueService>());
