Horde: Remove incorrect using statement
#rnx
#skipci

diff --git a/Engine/Source/Programs/Horde/Horde.Build/Jobs/TestData/TestDataCollection.cs b/Engine/Source/Programs/Horde/Horde.Build/Jobs/TestData/TestDataCollection.cs
--- a/Engine/Source/Programs/Horde/Horde.Build/Jobs/TestData/TestDataCollection.cs
+++ b/Engine/Source/Programs/Horde/Horde.Build/Jobs/TestData/TestDataCollection.cs
@@ -26,8 +26,7 @@ namespace Horde.Build.Jobs.TestData
 
 	using TestId = ObjectId<ITest>;
 	using TestSuiteId = ObjectId<ITestSuite>;
-	using TestMetaId = ObjectId<ITestMeta>;
-	using static System.Net.Mime.MediaTypeNames;
+	using TestMetaId = ObjectId<ITestMeta>;	
 
 	/// <summary>
 	/// Tests which are run in a  stream
