Horde: Re-enable DebugController

diff --git a/Engine/Source/Programs/Horde/Horde.Build/Server/DebugController.cs b/Engine/Source/Programs/Horde/Horde.Build/Server/DebugController.cs
--- a/Engine/Source/Programs/Horde/Horde.Build/Server/DebugController.cs
+++ b/Engine/Source/Programs/Horde/Horde.Build/Server/DebugController.cs
@@ -1,16 +1,33 @@
 // Copyright Epic Games, Inc. All Rights Reserved.
 
 //#define ENABLE_PUBLIC_DEBUG_CONTROLLER
-//#define ENABLE_SECURE_DEBUG_CONTROLLER
-
-#if ENABLE_PUBLIC_DEBUG_CONTROLLER
-
-using Horde.Build.Collections;
-using Horde.Build.Models;
+#define ENABLE_SECURE_DEBUG_CONTROLLER
+
+
+using System.Collections.Generic;
+using System.IO;
+using System.Net;
+using System.Text;
+using System.Text.Json;
+using System.Threading;
+using System.Threading.Tasks;
+using System.Web;
+using EpicGames.Core;
+using Horde.Build.Acls;
+using Horde.Build.Jobs;
+using Horde.Build.Jobs.Graphs;
+using Horde.Build.Jobs.Templates;
+using Horde.Build.Logs;
 using Horde.Build.Utilities;
+using JetBrains.Profiler.SelfApi;
+using Microsoft.AspNetCore.Authorization;
+using Microsoft.AspNetCore.Mvc;
+using Microsoft.Extensions.Options;
+using MongoDB.Driver;
 
-namespace Horde.Build.Controllers
+namespace Horde.Build.Server
 {
+#if ENABLE_PUBLIC_DEBUG_CONTROLLER
 	/// <summary>
 	/// Public endpoints for the debug controller
 	/// </summary>
@@ -226,120 +243,59 @@ namespace Horde.Build.Controllers
 		/// <summary>
 		/// The ACL service singleton
 		/// </summary>
-		AclService AclService;
+		private readonly AclService _aclService;
 
 		/// <summary>
 		/// The database service instance
 		/// </summary>
-		DatabaseService DatabaseService;
-
-		/// <summary>
-		/// The job task source singelton
-		/// </summary>
-		JobTaskSource JobTaskSource;
-
-		/// <summary>
-		/// The globals singleton
-		/// </summary>
-		ISingletonDocument<Globals> GlobalsSingleton;
-
-		/// <summary>
-		/// Collection of pool documents
-		/// </summary>
-		IPoolCollection PoolCollection;
+		private readonly MongoService _mongoService;
 
 		/// <summary>
-		/// Collection of project documents
+		/// The job task source singleton
 		/// </summary>
-		IProjectCollection ProjectCollection;
-
-		/// <summary>
-		/// Collection of agent documents
-		/// </summary>
-		IAgentCollection AgentCollection;
-
-		/// <summary>
-		/// Collection of session documents
-		/// </summary>
-		ISessionCollection SessionCollection;
-
-		/// <summary>
-		/// Collection of lease documents
-		/// </summary>
-		ILeaseCollection LeaseCollection;
+		private readonly JobTaskSource _jobTaskSource;
 
 		/// <summary>
 		/// Collection of template documents
 		/// </summary>
-		ITemplateCollection TemplateCollection;
-
-		/// <summary>
-		/// Collection of stream documents
-		/// </summary>
-		IStreamCollection StreamCollection;
-
+		private readonly ITemplateCollection _templateCollection;
+		
 		/// <summary>
 		/// The graph collection singleton
 		/// </summary>
-		IGraphCollection GraphCollection;
+		private readonly IGraphCollection _graphCollection;
 
 		/// <summary>
 		/// The log file collection singleton
 		/// </summary>
-		ILogFileCollection LogFileCollection;
-
-		/// <summary>
-		/// Perforce client
-		/// </summary>
-		IPerforceService Perforce;
-
-		/// <summary>
-		/// Fleet manager
-		/// </summary>
-		IFleetManager FleetManager;
+		private readonly ILogFileCollection _logFileCollection;
 
 		/// <summary>
 		/// Settings
 		/// </summary>
-		private IOptionsMonitor<ServerSettings> Settings;
+		private readonly IOptionsMonitor<ServerSettings> _settings;
 
 		/// <summary>
 		/// Constructor
 		/// </summary>
-		/// <param name="AclService">The ACL service singleton</param>
-		/// <param name="DatabaseService">The database service instance</param>
-		/// <param name="JobTaskSource">The dispatch service singleton</param>
-		/// <param name="GlobalsSingleton">The globals singleton</param>
-		/// <param name="PoolCollection">Collection of pool documents</param>
-		/// <param name="ProjectCollection">Collection of project documents</param>
-		/// <param name="AgentCollection">Collection of agent documents</param>
-		/// <param name="SessionCollection">Collection of session documents</param>
-		/// <param name="LeaseCollection">Collection of lease documents</param>
-		/// <param name="TemplateCollection">Collection of template documents</param>
-		/// <param name="StreamCollection">Collection of stream documents</param>
-		/// <param name="GraphCollection">The graph collection</param>
-		/// <param name="LogFileCollection">The log file collection</param>
-		/// <param name="Perforce">Perforce client</param>
-		/// <param name="FleetManager">The default fleet manager</param>
-		/// <param name="Settings">Settings</param>
-		public SecureDebugController(AclService AclService, DatabaseService DatabaseService, JobTaskSource JobTaskSource, ISingletonDocument<Globals> GlobalsSingleton, IPoolCollection PoolCollection, IProjectCollection ProjectCollection, IAgentCollection AgentCollection, ISessionCollection SessionCollection, ILeaseCollection LeaseCollection, ITemplateCollection TemplateCollection, IStreamCollection StreamCollection, IGraphCollection GraphCollection, ILogFileCollection LogFileCollection, IPerforceService Perforce, IFleetManager FleetManager, IOptionsMonitor<ServerSettings> Settings)
+		/// <param name="aclService">The ACL service singleton</param>
+		/// <param name="mongoService">The database service instance</param>
+		/// <param name="jobTaskSource">The dispatch service singleton</param>
+		/// <param name="templateCollection">Collection of template documents</param>
+		/// <param name="graphCollection">The graph collection</param>
+		/// <param name="logFileCollection">The log file collection</param>
+		/// <param name="settings">Settings</param>
+		public SecureDebugController(AclService aclService, MongoService mongoService, JobTaskSource jobTaskSource,
+			ITemplateCollection templateCollection, IGraphCollection graphCollection,
+			ILogFileCollection logFileCollection, IOptionsMonitor<ServerSettings> settings)
 		{
-			AclService = AclService;
-			DatabaseService = DatabaseService;
-			JobTaskSource = JobTaskSource;
-			GlobalsSingleton = GlobalsSingleton;
-			PoolCollection = PoolCollection;
-			ProjectCollection = ProjectCollection;
-			AgentCollection = AgentCollection;
-			SessionCollection = SessionCollection;
-			LeaseCollection = LeaseCollection;
-			TemplateCollection = TemplateCollection;
-			StreamCollection = StreamCollection;
-			GraphCollection = GraphCollection;
-			LogFileCollection = LogFileCollection;
-			Perforce = Perforce;
-			FleetManager = FleetManager;
-			Settings = Settings;
+			_aclService = aclService;
+			_mongoService = mongoService;
+			_jobTaskSource = jobTaskSource;
+			_templateCollection = templateCollection;
+			_graphCollection = graphCollection;
+			_logFileCollection = logFileCollection;
+			_settings = settings;
 		}
 
 		/// <summary>
@@ -350,22 +306,22 @@ namespace Horde.Build.Controllers
 		[Route("/api/v1/debug/environment")]
 		public async Task<ActionResult> GetServerEnvVars()
 		{
-			if (!await AclService.AuthorizeAsync(AclAction.AdminRead, User))
+			if (!await _aclService.AuthorizeAsync(AclAction.AdminRead, User))
 			{
 				return Forbid();
 			}
 
-			StringBuilder Content = new StringBuilder();
-			Content.AppendLine("<html><body><pre>");
-			foreach (System.Collections.DictionaryEntry? Pair in System.Environment.GetEnvironmentVariables())
+			StringBuilder content = new StringBuilder();
+			content.AppendLine("<html><body><pre>");
+			foreach (System.Collections.DictionaryEntry? pair in System.Environment.GetEnvironmentVariables())
 			{
-				if (Pair != null)
+				if (pair != null)
 				{
-					Content.AppendLine(HttpUtility.HtmlEncode($"{Pair.Value.Key}={Pair.Value.Value}"));
+					content.AppendLine(HttpUtility.HtmlEncode($"{pair.Value.Key}={pair.Value.Value}"));
 				}
 			}
-			Content.Append("</pre></body></html>");
-			return new ContentResult { ContentType = "text/html", StatusCode = (int)HttpStatusCode.OK, Content = Content.ToString() };
+			content.Append("</pre></body></html>");
+			return new ContentResult { ContentType = "text/html", StatusCode = (int)HttpStatusCode.OK, Content = content.ToString() };
 		}
 
 		/// <summary>
@@ -376,12 +332,12 @@ namespace Horde.Build.Controllers
 		[Route("/api/v1/debug/queue")]
 		public async Task<ActionResult<object>> GetQueueStatusAsync()
 		{
-			if (!await AclService.AuthorizeAsync(AclAction.AdminRead, User))
+			if (!await _aclService.AuthorizeAsync(AclAction.AdminRead, User))
 			{
 				return Forbid();
 			}
 
-			return JobTaskSource.GetStatus();
+			return _jobTaskSource.GetStatus();
 		}
 
 		/// <summary>
@@ -392,17 +348,17 @@ namespace Horde.Build.Controllers
 		[Route("/api/v1/debug/config")]
 		public async Task<ActionResult> GetConfig()
 		{
-			if (!await AclService.AuthorizeAsync(AclAction.AdminRead, User))
+			if (!await _aclService.AuthorizeAsync(AclAction.AdminRead, User))
 			{
 				return Forbid();
 			}
 
-			JsonSerializerOptions Options = new JsonSerializerOptions
+			JsonSerializerOptions options = new JsonSerializerOptions
 			{
 				WriteIndented = true
 			};
 
-			return Ok(JsonSerializer.Serialize(Settings.CurrentValue, Options));
+			return Ok(JsonSerializer.Serialize(_settings.CurrentValue, options));
 		}
 
 		/// <summary>
@@ -412,15 +368,15 @@ namespace Horde.Build.Controllers
 		[HttpGet]
 		[Route("/api/v1/debug/graphs")]
 		[ProducesResponseType(200, Type = typeof(GetGraphResponse))]
-		public async Task<ActionResult<List<object>>> GetGraphsAsync([FromQuery] int? Index = null, [FromQuery] int? Count = null, [FromQuery] PropertyFilter? Filter = null)
+		public async Task<ActionResult<List<object>>> GetGraphsAsync([FromQuery] int? index = null, [FromQuery] int? count = null, [FromQuery] PropertyFilter? filter = null)
 		{
-			if (!await AclService.AuthorizeAsync(AclAction.AdminRead, User))
+			if (!await _aclService.AuthorizeAsync(AclAction.AdminRead, User))
 			{
 				return Forbid();
 			}
 
-			List<IGraph> Graphs = await GraphCollection.FindAllAsync(null, Index, Count);
-			return Graphs.ConvertAll(x => new GetGraphResponse(x).ApplyFilter(Filter));
+			List<IGraph> graphs = await _graphCollection.FindAllAsync(null, index, count);
+			return graphs.ConvertAll(x => new GetGraphResponse(x).ApplyFilter(filter));
 		}
 
 		/// <summary>
@@ -430,83 +386,83 @@ namespace Horde.Build.Controllers
 		[HttpGet]
 		[Route("/api/v1/debug/graphs/{GraphId}")]
 		[ProducesResponseType(200, Type = typeof(GetGraphResponse))]
-		public async Task<ActionResult<object>> GetGraphAsync(string GraphId, [FromQuery] PropertyFilter? Filter = null)
+		public async Task<ActionResult<object>> GetGraphAsync(string graphId, [FromQuery] PropertyFilter? filter = null)
 		{
-			if (!await AclService.AuthorizeAsync(AclAction.AdminRead, User))
+			if (!await _aclService.AuthorizeAsync(AclAction.AdminRead, User))
 			{
 				return Forbid();
 			}
 
-			IGraph Graph = await GraphCollection.GetAsync(ContentHash.Parse(GraphId));
-			return new GetGraphResponse(Graph).ApplyFilter(Filter);
+			IGraph graph = await _graphCollection.GetAsync(ContentHash.Parse(graphId));
+			return new GetGraphResponse(graph).ApplyFilter(filter);
 		}
 
 		/// <summary>
 		/// Query all the job templates.
 		/// </summary>
-		/// <param name="Filter">Filter for properties to return</param>
+		/// <param name="filter">Filter for properties to return</param>
 		/// <returns>Information about all the job templates</returns>
 		[HttpGet]
 		[Route("/api/v1/debug/templates")]
 		[ProducesResponseType(typeof(List<GetTemplateResponse>), 200)]
-		public async Task<ActionResult<object>> GetTemplatesAsync([FromQuery] PropertyFilter? Filter = null)
+		public async Task<ActionResult<object>> GetTemplatesAsync([FromQuery] PropertyFilter? filter = null)
 		{
-			if (!await AclService.AuthorizeAsync(AclAction.AdminRead, User))
+			if (!await _aclService.AuthorizeAsync(AclAction.AdminRead, User))
 			{
 				return Forbid();
 			}
 
-			List<ITemplate> Templates = await TemplateCollection.FindAllAsync();
-			return Templates.ConvertAll(x => new GetTemplateResponse(x).ApplyFilter(Filter));
+			List<ITemplate> templates = await _templateCollection.FindAllAsync();
+			return templates.ConvertAll(x => new GetTemplateResponse(x).ApplyFilter(filter));
 		}
 
 		/// <summary>
 		/// Retrieve information about a specific job template.
 		/// </summary>
-		/// <param name="TemplateHash">Id of the template to get information about</param>
-		/// <param name="Filter">List of properties to return</param>
+		/// <param name="templateHash">Id of the template to get information about</param>
+		/// <param name="filter">List of properties to return</param>
 		/// <returns>Information about the requested template</returns>
 		[HttpGet]
 		[Route("/api/v1/debug/templates/{TemplateHash}")]
 		[ProducesResponseType(typeof(GetTemplateResponse), 200)]
-		public async Task<ActionResult<object>> GetTemplateAsync(string TemplateHash, [FromQuery] PropertyFilter? Filter = null)
+		public async Task<ActionResult<object>> GetTemplateAsync(string templateHash, [FromQuery] PropertyFilter? filter = null)
 		{
-			ContentHash TemplateHashValue = ContentHash.Parse(TemplateHash);
-			if (!await AclService.AuthorizeAsync(AclAction.AdminRead, User))
+			ContentHash templateHashValue = ContentHash.Parse(templateHash);
+			if (!await _aclService.AuthorizeAsync(AclAction.AdminRead, User))
 			{
 				return Forbid();
 			}
 
-			ITemplate? Template = await TemplateCollection.GetAsync(TemplateHashValue);
-			if (Template == null)
+			ITemplate? template = await _templateCollection.GetAsync(templateHashValue);
+			if (template == null)
 			{
 				return NotFound();
 			}
-			return Template.ApplyFilter(Filter);
+			return template.ApplyFilter(filter);
 		}
 
 		/// <summary>
 		/// Retrieve metadata about a specific log file
 		/// </summary>
-		/// <param name="LogFileId">Id of the log file to get information about</param>
-		/// <param name="Filter">Filter for the properties to return</param>
+		/// <param name="logFileId">Id of the log file to get information about</param>
+		/// <param name="filter">Filter for the properties to return</param>
 		/// <returns>Information about the requested project</returns>
 		[HttpGet]
 		[Route("/api/v1/debug/logs/{LogFileId}")]
-		public async Task<ActionResult<object>> GetLogAsync(LogId LogFileId, [FromQuery] PropertyFilter? Filter = null)
+		public async Task<ActionResult<object>> GetLogAsync(ObjectId<ILogFile> logFileId, [FromQuery] PropertyFilter? filter = null)
 		{
-			if (!await AclService.AuthorizeAsync(AclAction.AdminRead, User))
+			if (!await _aclService.AuthorizeAsync(AclAction.AdminRead, User))
 			{
 				return Forbid();
 			}
 
-			ILogFile? LogFile = await LogFileCollection.GetLogFileAsync(LogFileId);
-			if (LogFile == null)
+			ILogFile? logFile = await _logFileCollection.GetLogFileAsync(logFileId, CancellationToken.None);
+			if (logFile == null)
 			{
 				return NotFound();
 			}
 
-			return LogFile.ApplyFilter(Filter);
+			return logFile.ApplyFilter(filter);
 		}
 
 		/// <summary>
@@ -515,325 +471,16 @@ namespace Horde.Build.Controllers
 		/// <returns>Async task</returns>
 		[HttpGet]
 		[Route("/api/v1/debug/collections/{Name}")]
-		public async Task<ActionResult<object>> GetDocumentsAsync(string Name, [FromQuery] string? Filter = null, [FromQuery] int Index = 0, [FromQuery] int Count = 10)
-		{
-			if (!await AclService.AuthorizeAsync(AclAction.AdminRead, User))
-			{
-				return Forbid();
-			}
-
-			IMongoCollection<Dictionary<string, object>> Collection = DatabaseService.GetCollection<Dictionary<string, object>>(Name);
-			List<Dictionary<string, object>> Documents = await Collection.Find(Filter ?? "{}").Skip(Index).Limit(Count).ToListAsync();
-			return Documents;
-		}
-
-		/// <summary>
-		/// Create P4 login ticket for the specified username
-		/// </summary>
-		/// <returns>Perforce ticket</returns>
-		[HttpGet]
-		[Route("/api/v1/debug/p4ticket")]
-		public async Task<ActionResult<string>> CreateTicket([FromQuery] string? Username = null)
-		{
-			if (!await AclService.AuthorizeAsync(AclAction.AdminRead, User))
-			{
-				return Forbid();
-			}
-
-			if (Username == null)
-			{
-				return BadRequest();
-			}
-
-			return await Perforce.CreateTicket(PerforceCluster.DefaultName, Username);
-		}
-
-		/// <summary>
-		/// Debugs perforce service commands
-		/// </summary>
-		/// <returns>Perforce ticket</returns>
-		[HttpGet]
-		[Route("/api/v1/debug/perforce/stress")]
-		public async Task<ActionResult<string>> GetPerforceStress()
-		{
-			if (!await AclService.AuthorizeAsync(AclAction.AdminWrite, User))
-			{
-				return Forbid();
-			}
-
-			Dictionary<string, string> Results = new Dictionary<string, string>();
-			const string DevBuild = "//UE4/Dev-Build";
-			string? PerforceUser = User.GetPerforceUser();
-
-			object UpdateLock = new object();
-
-			// ParallelTest.ForAsync doesn't seem to handle async lambdas
-			IEnumerable<Task> Tasks = Enumerable.Range(0, 16).Select(async (Idx) => {
-
-				int LatestChange = await Perforce.GetLatestChangeAsync(PerforceCluster.DefaultName, DevBuild, PerforceUser);
-				int CodeChange = await Perforce.GetCodeChangeAsync(PerforceCluster.DefaultName, DevBuild, LatestChange);				
-
-				lock (UpdateLock)
-				{
-					Results.Add($"Parallel Perforce Result {Idx}", $"LatestChange: {LatestChange} - CodeChange: {CodeChange}");
-				}
-			});
-
-			await Task.WhenAll(Tasks);			
-
-			Tasks = Enumerable.Range(0, 16).Select(async (Idx) => {
-
-				int LatestChange = await Perforce.GetLatestChangeAsync(PerforceCluster.DefaultName, DevBuild, PerforceUser);
-				int CodeChange = await Perforce.GetCodeChangeAsync(PerforceCluster.DefaultName, DevBuild, LatestChange);				
-
-				lock (UpdateLock)
-				{
-					Results.Add($"Parallel Perforce Result {Idx + 16}", $"LatestChange: {LatestChange} - CodeChange: {CodeChange}");
-				}
-			});
-
-
-			await Task.WhenAll(Tasks);
-
-			StringBuilder Content = new StringBuilder();
-			Content.AppendLine("<html><body><pre>");
-			foreach (KeyValuePair<string, string> Pair in Results)
-			{
-				Content.AppendLine(HttpUtility.HtmlEncode($"{Pair.Key}={Pair.Value}"));
-			}
-			Content.Append("</pre></body></html>");
-			return new ContentResult { ContentType = "text/html", StatusCode = (int)HttpStatusCode.OK, Content = Content.ToString() };
-
-		}
-
-		/// <summary>
-		/// Debugs perforce service commands
-		/// </summary>
-		/// <returns>Perforce ticket</returns>
-		[HttpGet]
-		[Route("/api/v1/debug/perforce/commands")]
-		public async Task<ActionResult<string>> GetPerforceCommands()
-		{
-			if (!await AclService.AuthorizeAsync(AclAction.AdminWrite, User))
-			{
-				return Forbid();
-			}
-
-			StringBuilder Content = new StringBuilder();
-
-			const string DevBuild = "//UE4/Dev-Build";
-			// a test CL I have setup in Dev-Build
-			const int TestCL = 16687685;
-
-			string? PerforceUser = User.GetPerforceUser();
-
-			if (string.IsNullOrEmpty(PerforceUser))
-			{
-				throw new Exception("Perforce user required");
-			}
-
-			Dictionary<string, string> Results = new Dictionary<string, string>();
-
-			// NOTE: This is in Private-Build not DevBuild
-			int NewChangeId = await Perforce.CreateNewChangeAsync(PerforceCluster.DefaultName, "//UE4/Private-Build", "Counter.txt");
-
-			List<ChangeDetails> ChangeDetails = await Perforce.GetChangeDetailsAsync(PerforceCluster.DefaultName, "//UE4/Private-Build", new int[] { NewChangeId }, PerforceUser);
-
-			Results.Add("CreateNewChangeAsync", $"{NewChangeId} : {ChangeDetails[0].Description}");
-
-			await Perforce.UpdateChangelistDescription(PerforceCluster.DefaultName, TestCL, $"Updated Description: New Change CL was {NewChangeId}");
-
-			ChangeDetails = await Perforce.GetChangeDetailsAsync(PerforceCluster.DefaultName, DevBuild, new int[] { TestCL }, PerforceUser);
-
-			Results.Add("UpdateChangelistDescription", $"{TestCL} : {ChangeDetails[0].Description}");
-
-			int LatestChange = await Perforce.GetLatestChangeAsync(PerforceCluster.DefaultName, DevBuild, PerforceUser);
-			Results.Add("GetLatestChangeAsync", LatestChange.ToString(CultureInfo.InvariantCulture));
-
-			int CodeChange = await Perforce.GetCodeChangeAsync(PerforceCluster.DefaultName, DevBuild, LatestChange);
-			Results.Add("GetCodeChangeAsync", CodeChange.ToString(CultureInfo.InvariantCulture));
-
-			PerforceUserInfo? UserInfo = await Perforce.GetUserInfoAsync(PerforceCluster.DefaultName, PerforceUser);
-
-			if (UserInfo == null)
-			{
-				Results.Add("GetUserInfoAsync", "null");
-			}
-			else
-			{
-				Results.Add("GetUserInfoAsync", $"{UserInfo.Login} : {UserInfo.Email}");
-			}
-
-			// changes
-			List<ChangeSummary> Changes = await Perforce.GetChangesAsync(PerforceCluster.DefaultName, DevBuild, null, LatestChange, 10, PerforceUser);
-			int Count = 0;
-			foreach (ChangeSummary Summary in Changes)
-			{
-				Results.Add($"GetChangesAsync Result {Count++}", $"{Summary.Number} : {Summary.Author} - {Summary.Description}");
-			}
-
-			ChangeDetails = await Perforce.GetChangeDetailsAsync(PerforceCluster.DefaultName, DevBuild, Changes.Select(Change => Change.Number).ToArray(), PerforceUser);
-			Count = 0;
-			foreach (ChangeDetails Details in ChangeDetails)
-			{
-				Results.Add($"GetChangeDetailsAsync Result {Count++}", $"{Details.Number} : {Details.Author} - {Details.Files.ToJson()}");
-			}
-
-			String Ticket = "Failed (Expected if not running service account)";
-			try
-			{
-				await Perforce.CreateTicket(PerforceCluster.DefaultName, PerforceUser);
-				Ticket = "Success";
-			}
-			catch
-			{
-
-			}
-
-			Results.Add($"CreateTicket", $"{Ticket}");
-			List<FileSummary> FileSummaries = await Perforce.FindFilesAsync(PerforceCluster.DefaultName, ChangeDetails.SelectMany(Details => Details.Files).Select(File => $"{DevBuild}{File}"));
-			Count = 0;
-			foreach (FileSummary File in FileSummaries)
-			{
-				Results.Add($"FindFilesAsync Result {Count++}", $"{File.Change} : {File.DepotPath} : {File.Exists}");
-			}
-
-			// print async
-
-			byte[] Bytes = await Perforce.PrintAsync(PerforceCluster.DefaultName, $"{DevBuild}/RunUAT.bat");
-			Results.Add($"PrintAsync: {DevBuild}/RunUAT.bat ", $"{System.Text.Encoding.Default.GetString(Bytes)}");
-
-			Bytes = await Perforce.PrintAsync(PerforceCluster.DefaultName, $"{DevBuild}/RunUAT.bat@13916380");
-			Results.Add($"PrintAsync: {DevBuild}/RunUAT.bat@13916380", $"{System.Text.Encoding.Default.GetString(Bytes)}");
-
-			// need to be running as service account for these
-
-			int? DuplicateCL = null;
-
-			int ShelvedChangeOnDevBuild = 16687685;
-
-			try 
-			{
-				DuplicateCL = await Perforce.DuplicateShelvedChangeAsync(PerforceCluster.DefaultName, ShelvedChangeOnDevBuild);
-				Results.Add("DuplicateShelvedChangeAsync", $"Duplicated CL {DuplicateCL} from Shelved CL {ShelvedChangeOnDevBuild}");
-			}
-			catch 
-			{
-				Results.Add($"DuplicateShelvedChangeAsync",  "Failed - expected unless using service account");
-			}
-
-			if (DuplicateCL.HasValue)
-			{
-				// Note change client much exist for the update
-				ChangeDetails = await Perforce.GetChangeDetailsAsync(PerforceCluster.DefaultName, DevBuild, new int[]{ DuplicateCL.Value}, PerforceUser);
-
-				Results.Add("UpdateChangelistDescription - PreUpdate", $"{DuplicateCL} : {ChangeDetails[0].Description}"); 
-
-				try 
-				{
-					await Perforce.UpdateChangelistDescription(PerforceCluster.DefaultName, DuplicateCL.Value, "Updated Description from Horde");
-
-					ChangeDetails = await Perforce.GetChangeDetailsAsync(PerforceCluster.DefaultName, DevBuild, new int[]{ DuplicateCL.Value}, PerforceUser);
-
-					Results.Add("UpdateChangelistDescription - PostUpdate", $"{DuplicateCL} : {ChangeDetails[0].Description}"); 
-				}
-				catch 
-				{
-					Results.Add("UpdateChangelistDescription - Faile", $"Client from change {DuplicateCL} must exist"); 
-				}
-
-				await Perforce.DeleteShelvedChangeAsync(PerforceCluster.DefaultName, DuplicateCL.Value);
-
-				ChangeDetails = await Perforce.GetChangeDetailsAsync(PerforceCluster.DefaultName, DevBuild, new int[] {DuplicateCL.Value} , PerforceUser);
-
-				Results.Add("GetChangeDetailsAsync - After Delete", $"ChangeDetails.Count: {ChangeDetails.Count}"); 
-
-			}
-
-			byte[] ImageBytes = await Perforce.PrintAsync(PerforceCluster.DefaultName, $"{DevBuild}/Samples/Games/ShooterGame/ShooterGame.png");
-
-			Content.AppendLine("<html><body><pre>");
-			foreach (KeyValuePair<string, string> Pair in Results)
-			{
-				Content.AppendLine(HttpUtility.HtmlEncode($"{Pair.Key}={Pair.Value}"));
-			}
-
-			Content.Append($"PrintAsync Binary: {DevBuild}/Samples/Games/ShooterGame/ShooterGame.png\n</pre>");
-			Content.Append($"<img src=\"data:image/png;base64,{Convert.ToBase64String(ImageBytes, Base64FormattingOptions.None)}\"/>");
-
-			Content.Append("</body></html>");
-			return new ContentResult { ContentType = "text/html", StatusCode = (int)HttpStatusCode.OK, Content = Content.ToString() };
-
-		}
-
-		/// <summary>
-		/// Simulates a crash of the native p4 bridge
-		/// </summary>		
-		[HttpGet]		
-		[Route("/api/v1/debug/perforce/debugcrash")]		
-		public async Task<ActionResult> GetPerforceDebugCrash()
-		{
-			if (!await AclService.AuthorizeAsync(AclAction.AdminWrite, User))
-			{
-				return Forbid();
-			}
-
-			P4Debugging.DebugCrash();
-
-			return Ok();
-		}
-
-		/// <summary>
-		/// Debugging fleet managers
-		/// </summary>
-		/// <returns>Nothing</returns>
-		[HttpGet]
-		[Route("/api/v1/debug/fleetmanager")]
-		public async Task<ActionResult<string>> FleetManage([FromQuery] string? PoolId = null, [FromQuery] string? ChangeType = null, [FromQuery] int? Count = null)
+		public async Task<ActionResult<object>> GetDocumentsAsync(string name, [FromQuery] string? filter = null, [FromQuery] int index = 0, [FromQuery] int count = 10)
 		{
-			if (!await AclService.AuthorizeAsync(AclAction.AdminRead, User))
+			if (!await _aclService.AuthorizeAsync(AclAction.AdminRead, User))
 			{
 				return Forbid();
 			}
 
-			if (PoolId == null)
-			{
-				return BadRequest("Missing 'PoolId' query parameter");
-			}
-
-			IPool? Pool = await PoolCollection.GetAsync(new PoolId(PoolId));
-			if (Pool == null)
-			{
-				return BadRequest($"Pool with ID '{PoolId}' not found");
-			}
-
-			if (!(ChangeType == "expand" || ChangeType == "shrink"))
-			{
-				return BadRequest("Missing 'ChangeType' query parameter and/or must be set to 'expand' or 'shrink'");
-			}
-
-			if (Count == null || Count.Value <= 0)
-			{
-				return BadRequest("Missing 'Count' query parameter and/or must be greater than 0");
-			}
-
-			List<IAgent> Agents = new List<IAgent>();
-			string Message = "No operation";
-
-			if (ChangeType == "expand")
-			{
-				await FleetManager.ExpandPool(Pool, Agents, Count.Value);
-				Message = $"Expanded pool {Pool.Name} by {Count}";
-			}
-			else if (ChangeType == "shrink")
-			{
-				Agents = await AgentCollection.FindAsync();
-				Agents = Agents.FindAll(a => a.IsInPool(Pool.Id)).ToList();
-				await FleetManager.ShrinkPool(Pool, Agents, Count.Value);
-				Message = $"Shrank pool {Pool.Name} by {Count}";
-			}
-
-			return new ContentResult { ContentType = "text/plain", StatusCode = (int)HttpStatusCode.OK, Content = Message };
+			IMongoCollection<Dictionary<string, object>> collection = _mongoService.GetCollection<Dictionary<string, object>>(name);
+			List<Dictionary<string, object>> documents = await collection.Find(filter ?? "{}").Skip(index).Limit(count).ToListAsync();
+			return documents;
 		}
 		
 		/// <summary>
@@ -846,18 +493,18 @@ namespace Horde.Build.Controllers
 		{
 			await DotTrace.EnsurePrerequisiteAsync();
 
-			string SnapshotDir = Path.Join(Path.GetTempPath(), "horde-profiler-snapshots");
-			if (!Directory.Exists(SnapshotDir))
+			string snapshotDir = Path.Join(Path.GetTempPath(), "horde-profiler-snapshots");
+			if (!Directory.Exists(snapshotDir))
 			{
-				Directory.CreateDirectory(SnapshotDir);
+				Directory.CreateDirectory(snapshotDir);
 			}
 
-			var Config = new DotTrace.Config();
-			Config.SaveToDir(SnapshotDir);
-			DotTrace.Attach(Config);
+			DotTrace.Config config = new ();
+			config.SaveToDir(snapshotDir);
+			DotTrace.Attach(config);
 			DotTrace.StartCollectingData();
 			
-			return new ContentResult { ContentType = "text/plain", StatusCode = (int)HttpStatusCode.OK, Content = "Profiling session started. Using dir " + SnapshotDir };
+			return new ContentResult { ContentType = "text/plain", StatusCode = (int)HttpStatusCode.OK, Content = "Profiling session started. Using dir " + snapshotDir };
 		}
 		
 		/// <summary>
@@ -881,13 +528,13 @@ namespace Horde.Build.Controllers
 		[Route("/api/v1/debug/profiler/download")]
 		public ActionResult DownloadProfilingData()
 		{
-			string SnapshotZipFile = DotTrace.GetCollectedSnapshotFilesArchive(false);
-			if (!System.IO.File.Exists(SnapshotZipFile))
+			string snapshotZipFile = DotTrace.GetCollectedSnapshotFilesArchive(false);
+			if (!System.IO.File.Exists(snapshotZipFile))
 			{
 				return NotFound("The generated snapshot .zip file was not found");
 			}
 			
-			return PhysicalFile(SnapshotZipFile, "application/zip", Path.GetFileName(SnapshotZipFile));
+			return PhysicalFile(snapshotZipFile, "application/zip", Path.GetFileName(snapshotZipFile));
 		}
 	}
 }
