Horde: Add route for generating log messages

For aiding debugging of missing/oversized log messages not getting propagated.

diff --git a/Engine/Source/Programs/Horde/Horde.Build/Server/DebugController.cs b/Engine/Source/Programs/Horde/Horde.Build/Server/DebugController.cs
--- a/Engine/Source/Programs/Horde/Horde.Build/Server/DebugController.cs
+++ b/Engine/Source/Programs/Horde/Horde.Build/Server/DebugController.cs
@@ -4,8 +4,10 @@
 #define ENABLE_SECURE_DEBUG_CONTROLLER
 
 
+using System;
 using System.Collections.Generic;
 using System.IO;
+using System.Linq;
 using System.Net;
 using System.Text;
 using System.Text.Json;
@@ -22,6 +24,7 @@ using Horde.Build.Utilities;
 using JetBrains.Profiler.SelfApi;
 using Microsoft.AspNetCore.Authorization;
 using Microsoft.AspNetCore.Mvc;
+using Microsoft.Extensions.Logging;
 using Microsoft.Extensions.Options;
 using MongoDB.Driver;
 
@@ -240,6 +243,8 @@ namespace Horde.Build.Server
 	[Authorize]
 	public class SecureDebugController : ControllerBase
 	{
+		private static readonly Random s_random = new ();
+		
 		/// <summary>
 		/// The ACL service singleton
 		/// </summary>
@@ -274,6 +279,11 @@ namespace Horde.Build.Server
 		/// Settings
 		/// </summary>
 		private readonly IOptionsMonitor<ServerSettings> _settings;
+		
+		/// <summary>
+		/// Logger
+		/// </summary>
+		private readonly ILogger<SecureDebugController> _logger;
 
 		/// <summary>
 		/// Constructor
@@ -285,9 +295,10 @@ namespace Horde.Build.Server
 		/// <param name="graphCollection">The graph collection</param>
 		/// <param name="logFileCollection">The log file collection</param>
 		/// <param name="settings">Settings</param>
+		/// <param name="logger">Logger</param>
 		public SecureDebugController(AclService aclService, MongoService mongoService, JobTaskSource jobTaskSource,
 			ITemplateCollection templateCollection, IGraphCollection graphCollection,
-			ILogFileCollection logFileCollection, IOptionsMonitor<ServerSettings> settings)
+			ILogFileCollection logFileCollection, IOptionsMonitor<ServerSettings> settings, ILogger<SecureDebugController> logger)
 		{
 			_aclService = aclService;
 			_mongoService = mongoService;
@@ -296,6 +307,7 @@ namespace Horde.Build.Server
 			_graphCollection = graphCollection;
 			_logFileCollection = logFileCollection;
 			_settings = settings;
+			_logger = logger;
 		}
 
 		/// <summary>
@@ -360,6 +372,59 @@ namespace Horde.Build.Server
 
 			return Ok(JsonSerializer.Serialize(_settings.CurrentValue, options));
 		}
+		
+		/// <summary>
+		/// Generate log message of varying size
+		/// </summary>
+		/// <returns>Information about the log message generated</returns>
+		[HttpGet]
+		[Route("/api/v1/debug/generate-log-msg")]
+		public async Task<ActionResult> GenerateLogMessage(
+			[FromQuery] string? logLevel = null,
+			[FromQuery] int messageLen = 0,
+			[FromQuery] int exceptionMessageLen = 0,
+			[FromQuery] int argCount = 0,
+			[FromQuery] int argLen = 10)
+		{
+			if (!await _aclService.AuthorizeAsync(AclAction.AdminRead, User))
+			{
+				return Forbid();
+			}
+
+			string RandomString(int length)
+			{
+				const string Chars = "ABCDEFGHIJKLMNOPQRSTUVWXYZabcdefghijklmnopqrstuvwxyz0123456789";
+				return new string(Enumerable.Repeat(Chars, length).Select(s => s[s_random.Next(s.Length)]).ToArray());
+			}
+			
+			if (!Enum.TryParse(logLevel, out LogLevel logLevelInternal))
+			{
+				logLevelInternal = LogLevel.Information;
+			}
+
+			Exception? exception = null;
+			string message = "Message generated by /api/v1/debug/generate-log-msg";
+			message += RandomString(messageLen);
+
+			if (exceptionMessageLen > 0)
+			{
+				exception = new Exception("Exception from /api/v1/debug/generate-log-msg " + RandomString(exceptionMessageLen));
+			}
+
+			Dictionary<string, object> args = new ();
+			if (argCount > 0)
+			{
+				for (int i = 0; i < argCount; i++)
+				{
+					args["Arg" + i] = "Arg 1 - " + RandomString(argLen);
+				}
+			}
+
+			using IDisposable logScope = _logger.BeginScope(args);
+			_logger.Log(logLevelInternal, exception, message);
+			
+			return Ok($"Log message generated logLevel={logLevelInternal} messageLen={messageLen} exceptionMessageLen={exceptionMessageLen} argCount={argCount} argLen={argLen}");
+		}
 
 		/// <summary>
 		/// Queries for all graphs
