Horde: Incremental GC for storage. Now tracks bundle import metadata in the database and tracks and removed references to bundles in Redis, before querying DB to determine if any references still exist.

diff --git a/Engine/Source/Programs/Horde/Horde.Build.Tests/GcServiceTests.cs b/Engine/Source/Programs/Horde/Horde.Build.Tests/GcServiceTests.cs
--- a/Engine/Source/Programs/Horde/Horde.Build.Tests/GcServiceTests.cs
+++ b/Engine/Source/Programs/Horde/Horde.Build.Tests/GcServiceTests.cs
@@ -11,7 +11,6 @@ using Horde.Build.Storage;
 using System.Threading;
 using Horde.Build.Server;
 using Horde.Build.Utilities;
-using Microsoft.Extensions.DependencyInjection;
 
 namespace Horde.Build.Tests
 {
@@ -32,12 +31,14 @@ namespace Horde.Build.Tests
 		[TestMethod]
 		public async Task CreateBasicTree()
 		{
+			await StorageService.StartAsync(CancellationToken.None);
+
 			BlobId.UseDeterministicIds();
 			await SetupNamespaceAsync();
 			IStorageClientImpl store = await StorageService.GetClientAsync(new NamespaceId("default"), CancellationToken.None);
 
 			Random random = new Random(0);
-			BlobLocator[] blobs = await CreateTestDataAsync(store, 2000, 4000, 2000, 10, random);
+			BlobLocator[] blobs = await CreateTestDataAsync(store, 30, 50, 30, 5, random);
 
 			HashSet<BlobLocator> roots = new HashSet<BlobLocator>();
 			for (int idx = 0; idx < 10; idx++)
@@ -53,9 +54,6 @@ namespace Horde.Build.Tests
 
 			await Clock.AdvanceAsync(TimeSpan.FromDays(1.0));
 
-			GcService gc = ServiceProvider.GetRequiredService<GcService>();
-			await gc.RunAsync(new NamespaceId("default"), CancellationToken.None);
-
 			BlobLocator[] remaining = await store.Backend.EnumerateAsync().Select(x => new BlobLocator(HostId.Empty, StorageBackend.GetBlobIdFromPath(x))).ToArrayAsync();
 			Assert.AreEqual(nodes.Count, remaining.Length);
 			Assert.IsTrue(remaining.All(x => nodes.Contains(x)));
@@ -99,7 +97,7 @@ namespace Horde.Build.Tests
 				int numParents = 1 + (int)(random.NextDouble() * maxParents);
 				for (; numParents > 0; numParents--)
 				{
-					int parentIdx = (int)(random.NextDouble() * Math.Min(idx, numRoots + numInterior - 1));
+					int parentIdx = Math.Min((int)(random.NextDouble() * Math.Min(idx, numRoots + numInterior)), idx - 1);
 					children[parentIdx].Add(idx);
 				}
 			}
diff --git a/Engine/Source/Programs/Horde/Horde.Build.Tests/TestSetup.cs b/Engine/Source/Programs/Horde/Horde.Build.Tests/TestSetup.cs
--- a/Engine/Source/Programs/Horde/Horde.Build.Tests/TestSetup.cs
+++ b/Engine/Source/Programs/Horde/Horde.Build.Tests/TestSetup.cs
@@ -251,9 +251,7 @@ namespace Horde.Build.Tests
 
 			services.AddSingleton<ILegacyStorageClient, BasicStorageClient>();
 
-			services.AddSingleton<IRefCollection, RefCollection>();
 			services.AddSingleton<StorageService>();
-			services.AddSingleton<GcService>();
 
 			services.AddSingleton<ISingletonDocument<AgentSoftwareChannels>>(new SingletonDocumentStub<AgentSoftwareChannels>());
 			services.AddSingleton<ISingletonDocument<DevicePlatformMapV1>>(new SingletonDocumentStub<DevicePlatformMapV1>());
diff --git a/Engine/Source/Programs/Horde/Horde.Build/Startup.cs b/Engine/Source/Programs/Horde/Horde.Build/Startup.cs
--- a/Engine/Source/Programs/Horde/Horde.Build/Startup.cs
+++ b/Engine/Source/Programs/Horde/Horde.Build/Startup.cs
@@ -346,7 +346,6 @@ namespace Horde.Build
 			services.AddSingleton<INotificationTriggerCollection, NotificationTriggerCollection>();
 			services.AddSingleton<IPoolCollection, PoolCollection>();
 			services.AddSingleton<IProjectCollection, ProjectCollection>();
-			services.AddSingleton<IRefCollection, RefCollection>();
 			services.AddSingleton<ISessionCollection, SessionCollection>();
 			services.AddSingleton<IServiceAccountCollection, ServiceAccountCollection>();
 			services.AddSingleton<ISubscriptionCollection, SubscriptionCollection>();
@@ -786,7 +785,25 @@ namespace Horde.Build
 				}
 		*/
 
-		public sealed class BlobIdBsonSerializer : SerializerBase<BlobLocator>
+		public sealed class HostIdBsonSerializer : SerializerBase<HostId>
+		{
+			/// <inheritdoc/>
+			public override HostId Deserialize(BsonDeserializationContext context, BsonDeserializationArgs args) => new HostId(context.Reader.ReadString());
+
+			/// <inheritdoc/>
+			public override void Serialize(BsonSerializationContext context, BsonSerializationArgs args, HostId value) => context.Writer.WriteString(value.ToString());
+		}
+
+		public sealed class BlobIdBsonSerializer : SerializerBase<BlobId>
+		{
+			/// <inheritdoc/>
+			public override BlobId Deserialize(BsonDeserializationContext context, BsonDeserializationArgs args) => new BlobId(context.Reader.ReadString());
+
+			/// <inheritdoc/>
+			public override void Serialize(BsonSerializationContext context, BsonSerializationArgs args, BlobId value) => context.Writer.WriteString(value.ToString());
+		}
+
+		public sealed class BlobLocatorBsonSerializer : SerializerBase<BlobLocator>
 		{
 			/// <inheritdoc/>
 			public override BlobLocator Deserialize(BsonDeserializationContext context, BsonDeserializationArgs args)
@@ -874,7 +891,9 @@ namespace Horde.Build
 				ConventionRegistry.Register("Horde", conventionPack, type => true);
 
 				// Register the custom serializers
+				BsonSerializer.RegisterSerializer(new HostIdBsonSerializer());
 				BsonSerializer.RegisterSerializer(new BlobIdBsonSerializer());
+				BsonSerializer.RegisterSerializer(new BlobLocatorBsonSerializer());
 				BsonSerializer.RegisterSerializer(new RefIdBsonSerializer());
 				BsonSerializer.RegisterSerializer(new RefNameBsonSerializer());
 				BsonSerializer.RegisterSerializer(new IoHashBsonSerializer());
diff --git a/Engine/Source/Programs/Horde/Horde.Build/Storage/GcService.cs b/Engine/Source/Programs/Horde/Horde.Build/Storage/GcService.cs
deleted file mode 100644
diff --git a/Engine/Source/Programs/Horde/Horde.Build/Storage/IRefCollection.cs b/Engine/Source/Programs/Horde/Horde.Build/Storage/IRefCollection.cs
deleted file mode 100644
diff --git a/Engine/Source/Programs/Horde/Horde.Build/Storage/RefCollection.cs b/Engine/Source/Programs/Horde/Horde.Build/Storage/RefCollection.cs
deleted file mode 100644
diff --git a/Engine/Source/Programs/Horde/Horde.Build/Storage/StorageService.cs b/Engine/Source/Programs/Horde/Horde.Build/Storage/StorageService.cs
--- a/Engine/Source/Programs/Horde/Horde.Build/Storage/StorageService.cs
+++ b/Engine/Source/Programs/Horde/Horde.Build/Storage/StorageService.cs
@@ -1,24 +1,35 @@
 // Copyright Epic Games, Inc. All Rights Reserved.
 
 using System;
+using System.Buffers;
+using System.Buffers.Binary;
 using System.Collections;
 using System.Collections.Generic;
 using System.IO;
 using System.Linq;
 using System.Reflection;
+using System.Runtime.InteropServices;
 using System.Security.Claims;
 using System.Threading;
 using System.Threading.Tasks;
 using EpicGames.Core;
 using EpicGames.Horde.Storage;
+using EpicGames.Redis;
+using EpicGames.Redis.Utility;
 using Horde.Build.Acls;
 using Horde.Build.Server;
 using Horde.Build.Storage.Backends;
 using Horde.Build.Utilities;
+using HordeCommon;
 using Microsoft.Extensions.Caching.Memory;
 using Microsoft.Extensions.Configuration;
 using Microsoft.Extensions.DependencyInjection;
+using Microsoft.Extensions.Hosting;
 using Microsoft.Extensions.Logging;
+using MongoDB.Bson;
+using MongoDB.Bson.Serialization.Attributes;
+using MongoDB.Driver;
+using StackExchange.Redis;
 
 namespace Horde.Build.Storage
 {
@@ -53,11 +64,6 @@ namespace Horde.Build.Storage
 		/// </summary>
 		IStorageBackend Backend { get; }
 
-		/// <summary>
-		/// The ref collection
-		/// </summary>
-		IRefCollection Refs { get; }
-
 		/// <summary>
 		/// Whether the backend supports redirects
 		/// </summary>
@@ -83,26 +89,26 @@ namespace Horde.Build.Storage
 	/// <summary>
 	/// Functionality related to the storage service
 	/// </summary>
-	public sealed class StorageService : IDisposable
+	public sealed class StorageService : IHostedService, IDisposable
 	{
 		sealed class StorageClient : StorageClientBase, IStorageClientImpl, IDisposable
 		{
+			readonly StorageService _outer;
+
 			public NamespaceConfig Config { get; }
 
 			public IStorageBackend Backend { get; }
 
 			NamespaceId NamespaceId => Config.Id;
 
-			public IRefCollection Refs { get; }
-
 			public bool SupportsRedirects { get; }
 
-			public StorageClient(NamespaceConfig config, IStorageBackend backend, IRefCollection refs, IMemoryCache cache, ILogger logger)
+			public StorageClient(StorageService outer, NamespaceConfig config, IStorageBackend backend, IMemoryCache cache, ILogger logger)
 				: base(cache, logger)
 			{
+				_outer = outer;
 				Config = config;
 				Backend = backend;
-				Refs = refs;
 				SupportsRedirects = backend is IStorageBackendWithRedirects;
 			}
 
@@ -160,12 +166,12 @@ namespace Horde.Build.Storage
 			/// <inheritdoc/>
 			public override async Task<BlobLocator> WriteBlobAsync(Stream stream, Utf8String prefix = default, CancellationToken cancellationToken = default)
 			{
-				BlobId id = BlobId.CreateNew(prefix);
+				BlobLocator locator = await _outer.AddBlobAsync(NamespaceId, prefix, cancellationToken);
 
-				string path = GetBlobPath(id);
+				string path = GetBlobPath(locator.BlobId);
 				await Backend.WriteAsync(path, stream, cancellationToken);
 
-				return new BlobLocator(HostId.Empty, id);
+				return locator;
 			}
 
 			/// <inheritdoc/>
@@ -177,8 +183,8 @@ namespace Horde.Build.Storage
 					return null;
 				}
 
-				BlobId id = BlobId.CreateNew(prefix);
-				string path = GetBlobPath(id);
+				BlobLocator locator = await _outer.AddBlobAsync(NamespaceId, prefix, cancellationToken);
+				string path = GetBlobPath(locator.BlobId);
 
 				Uri? url = await redirectBackend.GetWriteRedirectAsync(path, cancellationToken);
 				if (url == null)
@@ -186,7 +192,13 @@ namespace Horde.Build.Storage
 					return null;
 				}
 
-				return (new BlobLocator(HostId.Empty, id), url);
+				return (locator, url);
+			}
+
+			public async Task DeleteBlobAsync(BlobId blobId, CancellationToken cancellationToken = default)
+			{
+				string path = GetBlobPath(blobId);
+				await Backend.DeleteAsync(path, cancellationToken);
 			}
 
 			#endregion
@@ -194,13 +206,13 @@ namespace Horde.Build.Storage
 			#region Refs
 
 			/// <inheritdoc/>
-			public override Task<NodeLocator> TryReadRefTargetAsync(RefName name, DateTime cacheTime = default, CancellationToken cancellationToken = default) => Refs.TryReadRefTargetAsync(NamespaceId, name, cacheTime, cancellationToken);
+			public override Task<NodeLocator> TryReadRefTargetAsync(RefName name, DateTime cacheTime = default, CancellationToken cancellationToken = default) => _outer.TryReadRefTargetAsync(NamespaceId, name, cacheTime, cancellationToken);
 
 			/// <inheritdoc/>
-			public override Task WriteRefTargetAsync(RefName name, NodeLocator target, RefOptions? options = null, CancellationToken cancellationToken = default) => Refs.WriteRefTargetAsync(NamespaceId, name, target, options, cancellationToken);
+			public override Task WriteRefTargetAsync(RefName name, NodeLocator target, RefOptions? options = null, CancellationToken cancellationToken = default) => _outer.WriteRefTargetAsync(NamespaceId, name, target, options, cancellationToken);
 
 			/// <inheritdoc/>
-			public override Task DeleteRefAsync(RefName name, CancellationToken cancellationToken = default) => Refs.DeleteRefAsync(NamespaceId, name, cancellationToken);
+			public override Task DeleteRefAsync(RefName name, CancellationToken cancellationToken = default) => _outer.DeleteRefAsync(NamespaceId, name, cancellationToken);
 
 			#endregion
 		}
@@ -208,10 +220,10 @@ namespace Horde.Build.Storage
 		class NamespaceInfo : IDisposable
 		{
 			public NamespaceConfig Config { get; }
-			public IStorageClientImpl Client { get; }
+			public StorageClient Client { get; }
 			public Acl? Acl { get; }
 
-			public NamespaceInfo(NamespaceConfig config, IStorageClientImpl client)
+			public NamespaceInfo(NamespaceConfig config, StorageClient client)
 			{
 				Config = config;
 				Client = client;
@@ -240,13 +252,128 @@ namespace Horde.Build.Storage
 			}
 		}
 
+		class BlobInfo
+		{
+			public ObjectId Id { get; set; }
+
+			[BsonElement("ns")]
+			public NamespaceId NamespaceId { get; set; }
+
+			[BsonElement("host")]
+			public HostId HostId { get; set; }
+
+			[BsonElement("blob")]
+			public BlobId BlobId { get; set; }
+
+			[BsonElement("imp")]
+			public List<ObjectId>? Imports { get; set; }
+
+			[BsonIgnore]
+			public BlobLocator Locator => new BlobLocator(HostId, BlobId);
+
+			public BlobInfo()
+			{
+			}
+
+			public BlobInfo(ObjectId id, NamespaceId namespaceId, HostId hostId, BlobId blobId)
+			{
+				Id = id;
+				NamespaceId = namespaceId;
+				HostId = hostId;
+				BlobId = blobId;
+			}
+		}
+
+		class RefInfo
+		{
+			[BsonIgnoreIfDefault]
+			public ObjectId Id { get; set; }
+
+			[BsonElement("ns")]
+			public NamespaceId NamespaceId { get; set; }
+
+			[BsonElement("name")]
+			public RefName Name { get; set; }
+
+			[BsonElement("blob")]
+			public BlobLocator BlobLocator { get; set; }
+
+			[BsonElement("binf")]
+			public ObjectId BlobInfoId { get; set; }
+
+			[BsonElement("idx")]
+			public int ExportIdx { get; set; }
+
+			[BsonElement("xa"), BsonIgnoreIfDefault]
+			public DateTime? ExpiresAtUtc { get; set; }
+
+			[BsonElement("xt"), BsonIgnoreIfDefault]
+			public TimeSpan? Lifetime { get; set; }
+
+			[BsonIgnore]
+			public NodeLocator Target => new NodeLocator(BlobLocator, ExportIdx);
+
+			public RefInfo()
+			{
+				Name = RefName.Empty;
+				BlobLocator = BlobLocator.Empty;
+			}
+
+			public RefInfo(NamespaceId namespaceId, RefName name, NodeLocator target, ObjectId blobInfoId)
+			{
+				NamespaceId = namespaceId;
+				Name = name;
+				BlobLocator = target.Blob;
+				BlobInfoId = blobInfoId;
+				ExportIdx = target.ExportIdx;
+			}
+
+			public bool HasExpired(DateTime utcNow) => ExpiresAtUtc.HasValue && utcNow >= ExpiresAtUtc.Value;
+
+			public bool RequiresTouch(DateTime utcNow) => ExpiresAtUtc.HasValue && Lifetime.HasValue && utcNow >= ExpiresAtUtc.Value - new TimeSpan(Lifetime.Value.Ticks / 4);
+		}
+
+		[SingletonDocument("gc-state")]
+		class GcState : SingletonBase
+		{
+			public ObjectId LastImportBlobInfoId { get; set; }
+			public List<GcNamespaceState> Namespaces { get; set; } = new List<GcNamespaceState>();
+
+			public GcNamespaceState FindOrAddNamespace(NamespaceId namespaceId)
+			{
+				GcNamespaceState? namespaceState = Namespaces.FirstOrDefault(x => x.Id == namespaceId);
+				if (namespaceState == null)
+				{
+					namespaceState = new GcNamespaceState { Id = namespaceId, LastTime = DateTime.UtcNow };
+					Namespaces.Add(namespaceState);
+				}
+				return namespaceState;
+			}
+		}
+
+		class GcNamespaceState
+		{
+			public NamespaceId Id { get; set; }
+			public DateTime LastTime { get; set; }
+		}
+
 		readonly GlobalsService _globalsService;
+		readonly RedisService _redisService;
 		readonly AclService _aclService;
-		readonly IRefCollection _refCollection;
+		readonly IClock _clock;
 		readonly IMemoryCache _cache;
 		readonly IServiceProvider _serviceProvider;
 		readonly ILogger _logger;
 
+		readonly IMongoCollection<BlobInfo> _blobCollection;
+		readonly IMongoCollection<RefInfo> _refCollection;
+
+		readonly ITicker _blobTicker;
+		readonly ITicker _refTicker;
+
+		readonly SingletonDocument<GcState> _gcState;
+		readonly ITicker _gcTicker;
+
 		State? _lastState;
 		string? _lastConfigRevision;
 
@@ -257,15 +384,32 @@ namespace Horde.Build.Storage
 		/// <summary>
 		/// Constructor
 		/// </summary>
-		public StorageService(GlobalsService globalsService, AclService aclService, IRefCollection refCollection, IMemoryCache cache, IServiceProvider serviceProvider, ILogger<StorageService> logger)
+		public StorageService(GlobalsService globalsService, MongoService mongoService, RedisService redisService, AclService aclService, IClock clock, IMemoryCache cache, IServiceProvider serviceProvider, ILogger<StorageService> logger)
 		{
 			_globalsService = globalsService;
+			_redisService = redisService;
 			_aclService = aclService;
-			_refCollection = refCollection;
+			_clock = clock;
 			_cache = cache;
 			_serviceProvider = serviceProvider;
 			_cachedState = new AsyncCachedValue<State>(() => GetNextState(), TimeSpan.FromMinutes(1.0));
 			_logger = logger;
+
+			List<MongoIndex<BlobInfo>> blobIndexes = new List<MongoIndex<BlobInfo>>();
+			blobIndexes.Add(keys => keys.Ascending(x => x.Imports));
+			blobIndexes.Add(keys => keys.Ascending(x => x.NamespaceId).Ascending(x => x.BlobId), unique: true);
+			_blobCollection = mongoService.GetCollection<BlobInfo>("Storage.Blobs", blobIndexes);
+
+			List<MongoIndex<RefInfo>> refIndexes = new List<MongoIndex<RefInfo>>();
+			refIndexes.Add(keys => keys.Ascending(x => x.NamespaceId).Ascending(x => x.Name), unique: true);
+			refIndexes.Add(keys => keys.Ascending(x => x.BlobInfoId));
+			_refCollection = mongoService.GetCollection<RefInfo>("Storage.Refs", refIndexes);
+
+			_blobTicker = clock.AddSharedTicker("Storage:Blobs", TimeSpan.FromMinutes(5.0), TickBlobsAsync, _logger);
+			_refTicker = clock.AddSharedTicker("Storage:Refs", TimeSpan.FromMinutes(5.0), TickRefsAsync, _logger);
+
+			_gcState = new SingletonDocument<GcState>(mongoService);
+			_gcTicker = clock.AddTicker("Storage:GC", TimeSpan.FromMinutes(5.0), TickGcAsync, logger);
 		}
 
 		/// <inheritdoc/>
@@ -276,12 +420,39 @@ namespace Horde.Build.Storage
 				_lastState.Dispose();
 				_lastState = null;
 			}
+
+			_blobTicker.Dispose();
+			_refTicker.Dispose();
+			_gcTicker.Dispose();
 		}
 
-		async ValueTask<NamespaceInfo> GetNamespaceInfoAsync(NamespaceId namespaceId, CancellationToken cancellationToken)
+		/// <inheritdoc/>
+		public async Task StartAsync(CancellationToken cancellationToken)
+		{
+			await _blobTicker.StartAsync();
+			await _refTicker.StartAsync();
+			await _gcTicker.StartAsync();
+		}
+
+		/// <inheritdoc/>
+		public async Task StopAsync(CancellationToken cancellationToken)
+		{
+			await _gcTicker.StopAsync();
+			await _refTicker.StopAsync();
+			await _blobTicker.StopAsync();
+		}
+
+		async ValueTask<NamespaceInfo?> TryGetNamespaceInfoAsync(NamespaceId namespaceId, CancellationToken cancellationToken)
 		{
 			State state = await _cachedState.GetAsync(cancellationToken);
-			if (!state.Namespaces.TryGetValue(namespaceId, out NamespaceInfo? namespaceInfo))
+			state.Namespaces.TryGetValue(namespaceId, out NamespaceInfo? namespaceInfo);
+			return namespaceInfo;
+		}
+
+		async ValueTask<NamespaceInfo> GetNamespaceInfoAsync(NamespaceId namespaceId, CancellationToken cancellationToken)
+		{
+			NamespaceInfo? namespaceInfo = await TryGetNamespaceInfoAsync(namespaceId, cancellationToken);
+			if (namespaceInfo == null)
 			{
 				throw new StorageException($"No namespace '{namespaceId}' is configured.");
 			}
@@ -335,6 +506,8 @@ namespace Horde.Build.Storage
 			}
 		}
 
+		#region Config
+
 		async Task<State> GetNextState()
 		{
 			IGlobals globals = await _globalsService.GetAsync();
@@ -372,7 +545,7 @@ namespace Horde.Build.Storage
 						IStorageBackend backend = CreateStorageBackend(backendConfig);
 
 #pragma warning disable CA2000 // Dispose objects before losing scope
-						StorageClient client = new StorageClient(namespaceConfig, backend, _refCollection, _cache, _logger);
+						StorageClient client = new StorageClient(this, namespaceConfig, backend, _cache, _logger);
 						nextState.Namespaces.Add(namespaceConfig.Id, new NamespaceInfo(namespaceConfig, client));
 #pragma warning restore CA2000 // Dispose objects before losing scope
 					}
@@ -497,5 +670,440 @@ namespace Horde.Build.Storage
 
 			return source.Equals(target);
 		}
+
+		#endregion
+
+		#region Blobs
+
+		/// <inheritdoc/>
+		async Task<BlobLocator> AddBlobAsync(NamespaceId namespaceId, Utf8String prefix = default, CancellationToken cancellationToken = default)
+		{
+			HostId hostId = HostId.Empty;
+
+			ObjectId id = ObjectId.GenerateNewId(_clock.UtcNow);
+			BlobId blobId = (prefix.Length > 0) ? new BlobId($"{prefix}{id}", BlobId.Validate.None) : new BlobId(id.ToString());
+
+			BlobInfo blobInfo = new BlobInfo(id, namespaceId, hostId, blobId);
+			await _blobCollection.InsertOneAsync(blobInfo, new InsertOneOptions { }, cancellationToken);
+
+			return new BlobLocator(hostId, blobId);
+		}
+
+		/// <inheritdoc/>
+		async Task<bool> IsBlobReferenced(ObjectId blobInfoId, CancellationToken cancellationToken = default)
+		{
+			FilterDefinition<BlobInfo> blobFilter = Builders<BlobInfo>.Filter.AnyEq(x => x.Imports, blobInfoId);
+			if (await _blobCollection.Find(blobFilter).AnyAsync(cancellationToken))
+			{
+				return true;
+			}
+
+			FilterDefinition<RefInfo> refFilter = Builders<RefInfo>.Filter.Eq(x => x.BlobInfoId, blobInfoId);
+			if (await _refCollection.Find(refFilter).AnyAsync(cancellationToken))
+			{
+				return true;
+			}
+
+			return false;
+		}
+
+		/// <summary>
+		/// Finds blobs at least 30 minutes old and computes import metadata for them. Done with a delay to allow write redirects.
+		/// </summary>
+		/// <param name="cancellationToken">Cancellation token for the operation</param>
+		async ValueTask TickBlobsAsync(CancellationToken cancellationToken)
+		{
+			GcState state = await _gcState.GetAsync();
+			DateTime utcNow = _clock.UtcNow;
+
+			// Cached storage clients for each namespace
+			Dictionary<NamespaceId, IStorageClient?> namespaceIdToClient = new Dictionary<NamespaceId, IStorageClient?>();
+
+			// Compute missing import info, by searching for blobs with an ObjectId timestamp after the last import compute cycle
+			ObjectId latestInfoId = ObjectId.GenerateNewId(utcNow - TimeSpan.FromMinutes(30.0));
+			using (IAsyncCursor<BlobInfo> cursor = await _blobCollection.Find(x => x.Id >= state.LastImportBlobInfoId && x.Id < latestInfoId).ToCursorAsync(cancellationToken))
+			{
+				while (await cursor.MoveNextAsync(cancellationToken))
+				{
+					// Find imports, and add a check record for each new blob
+					foreach (BlobInfo blobInfo in cursor.Current)
+					{
+						IStorageClient? client;
+						if (!namespaceIdToClient.TryGetValue(blobInfo.NamespaceId, out client))
+						{
+							NamespaceInfo? namespaceInfo = await TryGetNamespaceInfoAsync(blobInfo.NamespaceId, cancellationToken);
+							client = namespaceInfo?.Client;
+							namespaceIdToClient.Add(blobInfo.NamespaceId, client);
+						}
+
+						if (client != null)
+						{
+							BundleHeader? header = await ReadHeaderAsync(client, blobInfo.Locator, cancellationToken);
+							if (header != null)
+							{
+								List<ObjectId> importInfoIds = new List<ObjectId>();
+								foreach (BundleImport import in header.Imports)
+								{
+									FilterDefinition<BlobInfo> filter = Builders<BlobInfo>.Filter.Expr(x => x.NamespaceId == blobInfo.NamespaceId && x.BlobId == import.Locator.BlobId);
+									UpdateDefinition<BlobInfo> update = Builders<BlobInfo>.Update.SetOnInsert(x => x.HostId, import.Locator.HostId);
+									BlobInfo blobInfoDoc = await _blobCollection.FindOneAndUpdateAsync(filter, update, new FindOneAndUpdateOptions<BlobInfo> { IsUpsert = true, ReturnDocument = ReturnDocument.After }, cancellationToken);
+									importInfoIds.Add(blobInfoDoc.Id);
+								}
+								await _blobCollection.UpdateOneAsync(x => x.Id == blobInfo.Id, Builders<BlobInfo>.Update.Set(x => x.Imports, importInfoIds), null, cancellationToken);
+							}
+							AddGcCheckRecord(blobInfo.NamespaceId, blobInfo.Id);
+						}
+					}
+
+					// Update the last imported blob id
+					await _gcState.UpdateAsync(state => state.LastImportBlobInfoId = latestInfoId);
+				}
+			}
+		}
+
+		#endregion
+
+		#region Refs
+
+		record RefCacheKey(NamespaceId NamespaceId, RefName Name);
+		record RefCacheValue(RefInfo? Value, DateTime Time);
+
+		/// <summary>
+		/// Adds a ref value to the cache
+		/// </summary>
+		/// <param name="namespaceId">Namespace containing the ref</param>
+		/// <param name="name">Name of the ref</param>
+		/// <param name="value">New target for the ref</param>
+		/// <returns>The cached value</returns>
+		RefCacheValue AddRefToCache(NamespaceId namespaceId, RefName name, RefInfo? value)
+		{
+			RefCacheValue cacheValue = new RefCacheValue(value, DateTime.UtcNow);
+			using (ICacheEntry newEntry = _cache.CreateEntry(new RefCacheKey(namespaceId, name)))
+			{
+				newEntry.Value = cacheValue;
+				newEntry.SetSize(name.Text.Length);
+				newEntry.SetSlidingExpiration(TimeSpan.FromMinutes(5.0));
+			}
+			return cacheValue;
+		}
+
+		/// <summary>
+		/// Expires any refs that are no longer valid
+		/// </summary>
+		/// <param name="cancellationToken">Cancellation token for the operation</param>
+		async ValueTask TickRefsAsync(CancellationToken cancellationToken)
+		{
+			DateTime utcNow = DateTime.UtcNow;
+			using (IAsyncCursor<RefInfo> cursor = await _refCollection.Find(x => x.ExpiresAtUtc < utcNow).ToCursorAsync(cancellationToken))
+			{
+				List<DeleteOneModel<RefInfo>> requests = new List<DeleteOneModel<RefInfo>>();
+				while (await cursor.MoveNextAsync(cancellationToken))
+				{
+					requests.Clear();
+
+					foreach (RefInfo refInfo in cursor.Current)
+					{
+						_logger.LogInformation("Expired ref {NamespaceId}:{RefName}", refInfo.NamespaceId, refInfo.Name);
+						FilterDefinition<RefInfo> filter = Builders<RefInfo>.Filter.Expr(x => x.Id == refInfo.Id && x.ExpiresAtUtc == refInfo.ExpiresAtUtc);
+						requests.Add(new DeleteOneModel<RefInfo>(filter));
+						AddGcCheckRecord(refInfo.NamespaceId, refInfo.BlobInfoId);
+						AddRefToCache(refInfo.NamespaceId, refInfo.Name, default);
+					}
+
+					if (requests.Count > 0)
+					{
+						await _refCollection.BulkWriteAsync(requests, cancellationToken: cancellationToken);
+					}
+				}
+			}
+		}
+
+		/// <inheritdoc/>
+		async Task DeleteRefAsync(NamespaceId namespaceId, RefName name, CancellationToken cancellationToken = default)
+		{
+			FilterDefinition<RefInfo> filter = Builders<RefInfo>.Filter.Expr(x => x.NamespaceId == namespaceId && x.Name == name);
+			await DeleteRefInternalAsync(namespaceId, name, filter, cancellationToken);
+		}
+
+		/// <summary>
+		/// Deletes a ref document that has reached its expiry time
+		/// </summary>
+		async Task DeleteExpiredRefAsync(RefInfo refDocument, CancellationToken cancellationToken = default)
+		{
+			FilterDefinition<RefInfo> filter = Builders<RefInfo>.Filter.Expr(x => x.Id == refDocument.Id && x.ExpiresAtUtc == refDocument.ExpiresAtUtc);
+			await DeleteRefInternalAsync(refDocument.NamespaceId, refDocument.Name, filter, cancellationToken);
+		}
+
+		async Task DeleteRefInternalAsync(NamespaceId namespaceId, RefName name, FilterDefinition<RefInfo> filter, CancellationToken cancellationToken = default)
+		{
+			RefInfo? oldRefInfo = await _refCollection.FindOneAndDeleteAsync<RefInfo>(filter, cancellationToken: cancellationToken);
+			if (oldRefInfo != null)
+			{
+				_logger.LogInformation("Deleted ref {NamespaceId}:{RefName}", namespaceId, name);
+				AddGcCheckRecord(namespaceId, oldRefInfo.BlobInfoId);
+			}
+			AddRefToCache(namespaceId, name, default);
+		}
+
+		/// <inheritdoc/>
+		async Task<NodeLocator> TryReadRefTargetAsync(NamespaceId namespaceId, RefName name, DateTime cacheTime = default, CancellationToken cancellationToken = default)
+		{
+			RefCacheValue entry;
+			if (!_cache.TryGetValue(name, out entry) || entry.Time < cacheTime)
+			{
+				RefInfo? refDocument = await _refCollection.Find(x => x.NamespaceId == namespaceId && x.Name == name).FirstOrDefaultAsync(cancellationToken);
+				entry = AddRefToCache(namespaceId, name, refDocument);
+			}
+
+			if (entry.Value == null)
+			{
+				return default;
+			}
+
+			if (entry.Value.ExpiresAtUtc != null)
+			{
+				DateTime utcNow = _clock.UtcNow;
+				if (entry.Value.HasExpired(utcNow))
+				{
+					await DeleteExpiredRefAsync(entry.Value, cancellationToken);
+					return default;
+				}
+				if (entry.Value.RequiresTouch(utcNow))
+				{
+					await _refCollection.UpdateOneAsync(x => x.Id == entry.Value.Id, Builders<RefInfo>.Update.Set(x => x.ExpiresAtUtc, utcNow + entry.Value.Lifetime!.Value), cancellationToken: cancellationToken);
+				}
+			}
+
+			return new NodeLocator(entry.Value.BlobLocator, entry.Value.ExportIdx);
+		}
+
+		/// <inheritdoc/>
+		async Task WriteRefTargetAsync(NamespaceId namespaceId, RefName name, NodeLocator target, RefOptions? options = null, CancellationToken cancellationToken = default)
+		{
+			BlobInfo? newBlobInfo = await _blobCollection.Find(x => x.NamespaceId == namespaceId && x.BlobId == target.Blob.BlobId).FirstOrDefaultAsync(cancellationToken);
+			if (newBlobInfo == null)
+			{
+				throw new Exception($"Invalid/unknown blob identifier '{target.Blob.BlobId}' in namespace {namespaceId}");
+			}
+
+			RefInfo newRefInfo = new RefInfo(namespaceId, name, target, newBlobInfo.Id);
+
+			if (options != null && options.Lifetime.HasValue)
+			{
+				newRefInfo.ExpiresAtUtc = _clock.UtcNow + options.Lifetime.Value;
+				if (options.Extend ?? true)
+				{
+					newRefInfo.Lifetime = options.Lifetime;
+				}
+			}
+
+			RefInfo? oldRefInfo = await _refCollection.FindOneAndReplaceAsync<RefInfo>(x => x.NamespaceId == namespaceId && x.Name == name, newRefInfo, new FindOneAndReplaceOptions<RefInfo> { IsUpsert = true }, cancellationToken);
+			if (oldRefInfo != null && oldRefInfo.BlobInfoId != newRefInfo.BlobInfoId)
+			{
+				AddGcCheckRecord(namespaceId, oldRefInfo.BlobInfoId);
+			}
+
+			if (oldRefInfo == null)
+			{
+				_logger.LogInformation("Inserted ref {NamespaceId}:{RefName} to {Target}", namespaceId, name, target);
+			}
+			else if (oldRefInfo.Target != newRefInfo.Target)
+			{
+				_logger.LogInformation("Updated ref {NamespaceId}:{RefName} to {Target} (was {OldTarget})", namespaceId, name, newRefInfo.Target, oldRefInfo.Target);
+			}
+			else
+			{
+				_logger.LogInformation("Updated ref {NamespaceId}:{RefName} to {Target} (no-op)", namespaceId, name, target);
+			}
+
+			AddRefToCache(namespaceId, name, newRefInfo);
+		}
+
+		#endregion
+
+		#region GC
+
+		uint GetGcTimestamp() => GetGcTimestamp(_clock.UtcNow);
+		static uint GetGcTimestamp(DateTime utcTime) => (uint)((utcTime - DateTime.UnixEpoch).Ticks / TimeSpan.TicksPerMinute);
+
+		static RedisKey GetRedisKey(string suffix) => $"storage:{suffix}";
+		static RedisKey GetRedisKey(NamespaceId namespaceId, string suffix) => GetRedisKey($"{namespaceId}:{suffix}");
+
+		RedisSortedSet<RedisValue> GetGcCheckSet(NamespaceId namespaceId) => new RedisSortedSet<RedisValue>(_redisService.ConnectionPool, GetRedisKey(namespaceId, "check"));
+
+		/// <summary>
+		/// Find the next namespace to run GC on
+		/// </summary>
+		/// <param name="cancellationToken"></param>
+		/// <returns></returns>
+		async ValueTask TickGcAsync(CancellationToken cancellationToken)
+		{
+			HashSet<NamespaceId> ranNamespaceIds = new HashSet<NamespaceId>();
+
+			DateTime utcNow = _clock.UtcNow;
+			for (; ; )
+			{
+				// Synchronize the list of configured namespaces with the GC state object
+				List<NamespaceConfig> namespaces = await GetNamespacesAsync(cancellationToken);
+
+				GcState state = await _gcState.GetAsync();
+				if (!Enumerable.SequenceEqual(namespaces.Select(x => x.Id).OrderBy(x => x), state.Namespaces.Select(x => x.Id).OrderBy(x => x)))
+				{
+					state = await _gcState.UpdateAsync(s => SyncNamespaceList(s, namespaces));
+				}
+
+				// Find all the namespaces that need to have GC run on them
+				List<(DateTime, GcNamespaceState)> pending = new List<(DateTime, GcNamespaceState)>();
+				foreach (GcNamespaceState namespaceState in state.Namespaces)
+				{
+					if (!ranNamespaceIds.Contains(namespaceState.Id))
+					{
+						NamespaceConfig? config = namespaces.FirstOrDefault(x => x.Id == namespaceState.Id);
+						if (config != null)
+						{
+							DateTime time = namespaceState.LastTime + TimeSpan.FromHours(config.GcFrequencyHrs);
+							if (time < utcNow)
+							{
+								pending.Add((time, namespaceState));
+							}
+						}
+					}
+				}
+				pending.SortBy(x => x.Item1);
+
+				// If there's nothing left to GC, bail out
+				if (pending.Count == 0)
+				{
+					break;
+				}
+
+				// Update the first one we can acquire a lock for
+				foreach ((_, GcNamespaceState namespaceState) in pending)
+				{
+					NamespaceId namespaceId = namespaceState.Id;
+					if (ranNamespaceIds.Add(namespaceId))
+					{
+						RedisKey key = GetRedisKey(namespaceId, "lock");
+						using (RedisLock namespaceLock = new RedisLock(_redisService.GetDatabase(), key))
+						{
+							if (await namespaceLock.AcquireAsync(TimeSpan.FromMinutes(20.0)))
+							{
+								try
+								{
+									await TickGcForNamespaceAsync(namespaceId, state.LastImportBlobInfoId, utcNow, cancellationToken);
+								}
+								catch (Exception ex)
+								{
+									_logger.LogError(ex, "Exception while running garbage collection: {Message}", ex.Message);
+								}
+								break;
+							}
+						}
+					}
+				}
+			}
+		}
+
+		async Task TickGcForNamespaceAsync(NamespaceId namespaceId, ObjectId lastImportBlobInfoId, DateTime utcNow, CancellationToken cancellationToken)
+		{
+			NamespaceInfo namespaceInfo = await GetNamespaceInfoAsync(namespaceId, cancellationToken);
+
+			double score = GetGcTimestamp(utcNow);
+
+			RedisSortedSet<RedisValue> checkSet = GetGcCheckSet(namespaceId);
+			for (; ; )
+			{
+				RedisValue[] values = await checkSet.RangeByRankAsync(0, 0);
+				if (values.Length == 0)
+				{
+					break;
+				}
+
+				ObjectId blobInfoId = new ObjectId(((byte[]?)values[0])!);
+				if (blobInfoId < lastImportBlobInfoId && !await IsBlobReferenced(blobInfoId, cancellationToken))
+				{
+					BlobInfo? info = await _blobCollection.FindOneAndDeleteAsync(x => x.Id == blobInfoId, cancellationToken: cancellationToken);
+					if (info != null)
+					{
+						if (info.Imports != null)
+						{
+							SortedSetEntry<RedisValue>[] entries = info.Imports.Select(x => new SortedSetEntry<RedisValue>(x.ToByteArray(), score)).ToArray();
+							_ = checkSet.AddAsync(entries, flags: CommandFlags.FireAndForget);
+							score = Math.BitIncrement(score);
+						}
+						await namespaceInfo.Client.DeleteBlobAsync(info.BlobId, cancellationToken);
+					}
+				}
+				_ = checkSet.RemoveAsync(values[0], CommandFlags.FireAndForget);
+			}
+
+			await _gcState.UpdateAsync(state => state.FindOrAddNamespace(namespaceId).LastTime = utcNow);
+		}
+
+		async Task<BundleHeader?> ReadHeaderAsync(IStorageClient store, BlobLocator locator, CancellationToken cancellationToken)
+		{
+			int fetchSize = 64 * 1024;
+			for (; ; )
+			{
+				using (IMemoryOwner<byte> owner = MemoryPool<byte>.Shared.Rent(fetchSize))
+				{
+					using (Stream? stream = await store.ReadBlobRangeAsync(locator, 0, fetchSize, cancellationToken))
+					{
+						if (stream == null)
+						{
+							_logger.LogError("Unable to read blob {Blob}", locator);
+							return null;
+						}
+
+						// Read the start of the blob
+						Memory<byte> memory = owner.Memory.Slice(0, fetchSize);
+
+						int length = await stream.ReadGreedyAsync(memory, cancellationToken);
+						if (length < BundleHeader.PreludeLength)
+						{
+							_logger.LogError("Blob {Blob} does not have a valid prelude", locator);
+							return null;
+						}
+
+						memory = memory.Slice(0, length);
+
+						// Make sure it's large enough to hold the header
+						int headerSize = BundleHeader.ReadPrelude(memory);
+						if (headerSize <= fetchSize)
+						{
+							return new BundleHeader(new MemoryReader(memory));
+						}
+
+						// Increase the fetch size and retry
+						fetchSize = headerSize;
+					}
+				}
+			}
+		}
+
+		static void SyncNamespaceList(GcState state, List<NamespaceConfig> namespaces)
+		{
+			HashSet<NamespaceId> validNamespaceIds = new HashSet<NamespaceId>(namespaces.Select(x => x.Id));
+			state.Namespaces.RemoveAll(x => !validNamespaceIds.Contains(x.Id));
+
+			HashSet<NamespaceId> currentNamespaceIds = new HashSet<NamespaceId>(state.Namespaces.Select(x => x.Id));
+			foreach (NamespaceConfig config in namespaces)
+			{
+				if (!currentNamespaceIds.Contains(config.Id))
+				{
+					state.Namespaces.Add(new GcNamespaceState { Id = config.Id });
+				}
+			}
+
+			state.Namespaces.SortBy(x => x.Id);
+		}
+
+		void AddGcCheckRecord(NamespaceId namespaceId, ObjectId blobInfoId)
+		{
+			double score = GetGcTimestamp();
+			_ = GetGcCheckSet(namespaceId).AddAsync(blobInfoId.ToByteArray(), score, flags: CommandFlags.FireAndForget);
+		}
+
+		#endregion
 	}
 }
diff --git a/Engine/Source/Programs/Shared/EpicGames.Horde/Storage/IStorageClient.cs b/Engine/Source/Programs/Shared/EpicGames.Horde/Storage/IStorageClient.cs
--- a/Engine/Source/Programs/Shared/EpicGames.Horde/Storage/IStorageClient.cs
+++ b/Engine/Source/Programs/Shared/EpicGames.Horde/Storage/IStorageClient.cs
@@ -3,6 +3,7 @@
 using System;
 using System.Collections.Generic;
 using System.Diagnostics;
+using System.Diagnostics.CodeAnalysis;
 using System.IO;
 using System.Reflection;
 using System.Reflection.Metadata;
@@ -78,7 +79,7 @@ namespace EpicGames.Horde.Storage
 	/// <summary>
 	/// Locates a node in storage
 	/// </summary>
-	public struct NodeLocator
+	public struct NodeLocator : IEquatable<NodeLocator>
 	{
 		/// <summary>
 		/// Location of the blob containing this node
@@ -104,8 +105,23 @@ namespace EpicGames.Horde.Storage
 		/// </summary>
 		public bool IsValid() => Blob.IsValid();
 
+		/// <inheritdoc/>
+		public override bool Equals([NotNullWhen(true)] object? obj) => obj is NodeLocator locator && Equals(locator);
+
+		/// <inheritdoc/>
+		public bool Equals(NodeLocator other) => Blob == other.Blob && ExportIdx == other.ExportIdx;
+
+		/// <inheritdoc/>
+		public override int GetHashCode() => HashCode.Combine(Blob, ExportIdx);
+
 		/// <inheritdoc/>
 		public override string ToString() => $"{Blob}#{ExportIdx}";
+
+		/// <inheritdoc/>
+		public static bool operator ==(NodeLocator left, NodeLocator right) => left.Equals(right);
+
+		/// <inheritdoc/>
+		public static bool operator !=(NodeLocator left, NodeLocator right) => !left.Equals(right);
 	}
 
 	/// <summary>
