Horde: Missing files.

diff --git a/Engine/Source/Programs/Horde/Horde.Build.Tests/BlobStoreTests.cs b/Engine/Source/Programs/Horde/Horde.Build.Tests/BlobStoreTests.cs
--- a/Engine/Source/Programs/Horde/Horde.Build.Tests/BlobStoreTests.cs
+++ b/Engine/Source/Programs/Horde/Horde.Build.Tests/BlobStoreTests.cs
@@ -140,7 +140,9 @@ namespace Horde.Build.Tests
 		public async Task RefExpiryTest()
 		{
 			IStorageClient store = await CreateStorageClientAsync();
-			NodeLocator target = new NodeLocator(new BlobLocator(HostId.Empty, new BlobId("test")), 123);
+
+			BlobLocator locator1 = await store.WriteBundleAsync(CreateTestBundle(new byte[] { 1, 2, 3 }, Array.Empty<BlobLocator>()));
+			NodeLocator target = new NodeLocator(locator1, 0);
 
 			await store.WriteRefTargetAsync("test-ref-1", target);
 			await store.WriteRefTargetAsync("test-ref-2", target, new RefOptions { Lifetime = TimeSpan.FromMinutes(30.0), Extend = true });
diff --git a/Engine/Source/Programs/Shared/EpicGames.Horde/Storage/HostId.cs b/Engine/Source/Programs/Shared/EpicGames.Horde/Storage/HostId.cs
--- a/Engine/Source/Programs/Shared/EpicGames.Horde/Storage/HostId.cs
+++ b/Engine/Source/Programs/Shared/EpicGames.Horde/Storage/HostId.cs
@@ -87,7 +87,7 @@ namespace EpicGames.Horde.Storage
 		{
 			if (text.Length == 0)
 			{
-				throw new ArgumentException("Host identifiers cannot be empty", name);
+				return;
 			}
 			if (text[^1] == ':')
 			{
diff --git a/Engine/Source/Programs/Shared/EpicGames.Redis/RedisSerializer.cs b/Engine/Source/Programs/Shared/EpicGames.Redis/RedisSerializer.cs
--- a/Engine/Source/Programs/Shared/EpicGames.Redis/RedisSerializer.cs
+++ b/Engine/Source/Programs/Shared/EpicGames.Redis/RedisSerializer.cs
@@ -108,6 +108,7 @@ namespace EpicGames.Redis
 		{
 			KeyValuePair<Type, object>[] converters =
 			{
+				CreateNativeConverter<RedisValue>(x => x, x => x),
 				CreateNativeConverter(x => (bool)x, x => x),
 				CreateNativeConverter(x => (int)x, x => x),
 				CreateNativeConverter(x => (int?)x, x => x),
