Horde: Adding some logging to see what is going on with test data
#rnx

diff --git a/Engine/Source/Programs/Horde/Horde.Build/Jobs/TestData/TestDataCollection.cs b/Engine/Source/Programs/Horde/Horde.Build/Jobs/TestData/TestDataCollection.cs
--- a/Engine/Source/Programs/Horde/Horde.Build/Jobs/TestData/TestDataCollection.cs
+++ b/Engine/Source/Programs/Horde/Horde.Build/Jobs/TestData/TestDataCollection.cs
@@ -335,6 +335,8 @@ namespace Horde.Build.Jobs.TestData
 
 			if (_settings.CurrentValue.FeatureFlags.EnableTestDataV2)
 			{
+				_logger.LogInformation("EnableTestDataV2 is true, creating new test data collections");
+
 				List<MongoIndex<TestMetaDocument>> metaIndexes = new List<MongoIndex<TestMetaDocument>>();
 				metaIndexes.Add(keys => keys.Ascending(x => x.ProjectName).Ascending(x => x.Platforms).Ascending(x => x.Configurations));
 				_testMeta = mongoService.GetCollection<TestMetaDocument>("TestData.Meta", metaIndexes);
@@ -432,8 +434,13 @@ namespace Horde.Build.Jobs.TestData
 			{
 				if (_settings.CurrentValue.FeatureFlags.EnableTestDataV2)
 				{
+					_logger.LogInformation("Attempting to add test report data");
 					await AddTestReportData(job, step, value, newTestData.Id);
 				}									
+				else
+				{
+					_logger.LogInformation("Test report data skipped due to disabled feature flag");
+				}
 			}
 			catch (Exception ex)
 			{
@@ -670,6 +677,7 @@ namespace Horde.Build.Jobs.TestData
 			int version = rootDoc.GetValue("Version", new BsonInt32(0)).AsInt32;
 			if (version < 1)
 			{
+				_logger.LogWarning("Test data does not have version and needs to be updated in stream for job {JobId} step {StepId}", job.Id, step.Id);
 				return;
 			}
 
