Horde: Add horde report version to somewhere accessible by grpc endpoint
#rnx

diff --git a/Engine/Source/Programs/Horde/Horde.Build.Tests/TestDataTests.cs b/Engine/Source/Programs/Horde/Horde.Build.Tests/TestDataTests.cs
--- a/Engine/Source/Programs/Horde/Horde.Build.Tests/TestDataTests.cs
+++ b/Engine/Source/Programs/Horde/Horde.Build.Tests/TestDataTests.cs
@@ -256,8 +256,8 @@ namespace Horde.Build.Tests
 			@"  ""Items"": [",
 			@"    {",
 			@"      ""Key"": ""Simple Report::UE.BootTest EngineTest Editor Win64"",",
-			@"      ""Data"": {",
 			@"		""Version"" : 1,",
+			@"      ""Data"": {",			
 			@"        ""Type"": ""Simple Report"",",
 			@"        ""TestName"": ""EditorBootTest"",",
 			@"        ""Description"": ""Win64 Development EditorGame"",",
@@ -292,17 +292,16 @@ namespace Horde.Build.Tests
 			@"      }",
 			@"    }",
 			@"  ],",
-			@"  ""Version"": 1",
 			@"}"
 		};
 
 		private readonly string[] _testSessionDataLines =
 		{
-			@"{",
-			@"	""Version"": 1,",
+			@"{",			
 			@"    ""Items"": [	",
 			@"        {",
 			@"            ""Key"": ""Automated Test Session"",",
+			@"		      ""Version"" : 1,",
 			@"            ""Data"": {",
 			@"                ""Type"": ""Automated Test Session"",",
 			@"                ""Name"": ""UE.Automation(Group:UI) EngineTest"",",
@@ -399,6 +398,7 @@ namespace Horde.Build.Tests
 			@"        },",
 			@"        {",
 			@"            ""Key"": ""Unreal Automated Tests::UE.TargetAutomation(RunTest=UI) Win64"",",
+			@"		      ""Version"" : 1,",
 			@"            ""Data"": {",
 			@"                ""Type"": ""Unreal Automated Tests"",",
 			@"                ""Devices"": [",
@@ -471,6 +471,7 @@ namespace Horde.Build.Tests
 			@"        },",
 			@"        {",
 			@"            ""Key"": ""Automated Test Session Result Details::b420bdde-c030-4add-81d2-3a8404ab3e45"",",
+			@"	          ""Version"" : 1,",
 			@"            ""Data"": {",
 			@"                ""cbdb55ea"": {",
 			@"                    ""Events"": [",
diff --git a/Engine/Source/Programs/Horde/Horde.Build/Jobs/TestData/TestDataCollection.cs b/Engine/Source/Programs/Horde/Horde.Build/Jobs/TestData/TestDataCollection.cs
--- a/Engine/Source/Programs/Horde/Horde.Build/Jobs/TestData/TestDataCollection.cs
+++ b/Engine/Source/Programs/Horde/Horde.Build/Jobs/TestData/TestDataCollection.cs
@@ -674,13 +674,6 @@ namespace Horde.Build.Jobs.TestData
 
 		async Task AddTestReportData(IJob job, IJobStep step, BsonDocument rootDoc, ObjectId testDataId)
 		{
-			int version = rootDoc.GetValue("Version", new BsonInt32(0)).AsInt32;
-			if (version < 1)
-			{
-				_logger.LogWarning("Test data does not have version and needs to be updated in stream for job {JobId} step {StepId}", job.Id, step.Id);
-				return;
-			}
-
 			List<AutomatedTestSessionData> sessions = new List<AutomatedTestSessionData>();
 			List<UnrealAutomatedTestData> tests = new List<UnrealAutomatedTestData>();
 
@@ -694,6 +687,13 @@ namespace Horde.Build.Jobs.TestData
 
 				BsonDocument testData = item.AsBsonDocument;
 
+				int version = testData.GetValue("Version", new BsonInt32(0)).AsInt32;
+				if (version < 1)
+				{
+					_logger.LogWarning("Test data does not have version and needs to be updated in stream for job {JobId} step {StepId}", job.Id, step.Id);
+					return;
+				}
+
 				BsonValue? value;
 				if (!testData.TryGetPropertyValue("Data.Type", BsonType.String, out value))
 				{
