Horde: Updating test data ingestion to support grouping detailed test report documents to refs
#rnx

diff --git a/Engine/Source/Programs/Horde/Horde.Build.Tests/TestDataTests.cs b/Engine/Source/Programs/Horde/Horde.Build.Tests/TestDataTests.cs
--- a/Engine/Source/Programs/Horde/Horde.Build.Tests/TestDataTests.cs
+++ b/Engine/Source/Programs/Horde/Horde.Build.Tests/TestDataTests.cs
@@ -3,23 +3,19 @@
 using System;
 using System.Collections.Generic;
 using System.Threading.Tasks;
-using EpicGames.Core;
 using Horde.Build.Jobs.Graphs;
 using Horde.Build.Projects;
 using Horde.Build.Jobs;
 using Horde.Build.Users;
 using Horde.Build.Streams;
 using Horde.Build.Utilities;
-using HordeCommon;
 using Microsoft.VisualStudio.TestTools.UnitTesting;
 using MongoDB.Bson;
 using Moq;
-using Horde.Build.Jobs.Templates;
 using Horde.Build.Jobs.TestData;
 using Microsoft.AspNetCore.Http;
 using Microsoft.AspNetCore.Mvc;
 using System.Security.Claims;
-using EpicGames.BuildGraph.Expressions;
 using Horde.Build.Logs;
 
 namespace Horde.Build.Tests
@@ -176,9 +172,16 @@ namespace Horde.Build.Tests
 
 				BsonDocument testData = BsonDocument.Parse(String.Join('\n', _simpleTestDataLines));
 
-				await TestDataCollection.AddAsync(job, step, "Simple Report Key", testData);
-				await TestDataCollection.AddAsync(job2, step2, "Simple Report Key", testData);
+				List<(string key, BsonDocument)> data = new List<(string key, BsonDocument)>();
+				BsonArray items = testData.GetValue("Items").AsBsonArray;
+				foreach (BsonValue item in items.ToList())
+				{
+					BsonDocument value = item.AsBsonDocument.GetValue("Data").AsBsonDocument;
+					data.Add(("Simple Report Key", value));
+				}
 
+				await TestDataCollection.AddAsync(job, step, data.ToArray());
+				await TestDataCollection.AddAsync(job2, step2, data.ToArray());
 
 				ActionResult<List<GetTestStreamResponse>> streamResult = await TestDataController.GetTestStreamsAsync(streamIds);
 				Assert.IsNotNull(streamResult);
@@ -201,6 +204,9 @@ namespace Horde.Build.Tests
 				List<GetTestDataRefResponse> refs = refResult.Value;
 
 				Assert.AreEqual(2, refs.Count);
+
+				Assert.AreEqual(1, refs[0].TestDataIds.Count);
+				Assert.AreEqual(1, refs[1].TestDataIds.Count);
 			}
 		}
 
@@ -218,8 +224,16 @@ namespace Horde.Build.Tests
 
 				BsonDocument testData = BsonDocument.Parse(String.Join('\n', _testSessionDataLines));
 
-				await TestDataCollection.AddAsync(job, step, "Session Report Key", testData);
-				await TestDataCollection.AddAsync(job2, step2, "Session Report Key", testData);
+				List<(string key, BsonDocument)> data = new List<(string key, BsonDocument)>();
+				BsonArray items = testData.GetValue("Items").AsBsonArray;
+				foreach (BsonValue item in items.ToList())
+				{
+					BsonDocument value = item.AsBsonDocument.GetValue("Data").AsBsonDocument;
+					data.Add(("Session Report Key", value));
+				}
+
+				await TestDataCollection.AddAsync(job, step, data.ToArray());
+				await TestDataCollection.AddAsync(job2, step2, data.ToArray());
 
 				ActionResult<List<GetTestStreamResponse>> streamResult = await TestDataController.GetTestStreamsAsync(streamIds);
 				Assert.IsNotNull(streamResult);
@@ -246,8 +260,12 @@ namespace Horde.Build.Tests
 				List<GetTestDataRefResponse> refs = refResult.Value;
 
 				Assert.AreEqual(2, refs.Count);
+
+				Assert.AreEqual(3, refs[0].TestDataIds.Count);
+				Assert.AreEqual(3, refs[1].TestDataIds.Count);
+
 			}
-			
+
 		}
 
 		private readonly string[] _simpleTestDataLines =
@@ -255,9 +273,9 @@ namespace Horde.Build.Tests
 			@"{",
 			@"  ""Items"": [",
 			@"    {",
-			@"      ""Key"": ""Simple Report::UE.BootTest EngineTest Editor Win64"",",
-			@"		""Version"" : 1,",
-			@"      ""Data"": {",			
+			@"      ""Key"": ""Simple Report::UE.BootTest EngineTest Editor Win64"",",			
+			@"      ""Data"": {",
+			@"  	  ""Version"" : 1,",
 			@"        ""Type"": ""Simple Report"",",
 			@"        ""TestName"": ""EditorBootTest"",",
 			@"        ""Description"": ""Win64 Development EditorGame"",",
@@ -300,9 +318,9 @@ namespace Horde.Build.Tests
 			@"{",			
 			@"    ""Items"": [	",
 			@"        {",
-			@"            ""Key"": ""Automated Test Session"",",
-			@"		      ""Version"" : 1,",
+			@"            ""Key"": ""Automated Test Session"",",			
 			@"            ""Data"": {",
+			@"  	          ""Version"" : 1,",
 			@"                ""Type"": ""Automated Test Session"",",
 			@"                ""Name"": ""UE.Automation(Group:UI) EngineTest"",",
 			@"                ""PreFlightChange"": """",",
@@ -398,9 +416,9 @@ namespace Horde.Build.Tests
 			@"        },",
 			@"        {",
 			@"            ""Key"": ""Unreal Automated Tests::UE.TargetAutomation(RunTest=UI) Win64"",",
-			@"		      ""Version"" : 1,",
 			@"            ""Data"": {",
 			@"                ""Type"": ""Unreal Automated Tests"",",
+			@"   		      ""Version"" : 1,",
 			@"                ""Devices"": [",
 			@"                    {",
 			@"                        ""DeviceName"": ""RDU-WIN64-12"",",
@@ -470,9 +488,9 @@ namespace Horde.Build.Tests
 			@"            }",
 			@"        },",
 			@"        {",
-			@"            ""Key"": ""Automated Test Session Result Details::b420bdde-c030-4add-81d2-3a8404ab3e45"",",
-			@"	          ""Version"" : 1,",
+			@"            ""Key"": ""Automated Test Session Result Details::b420bdde-c030-4add-81d2-3a8404ab3e45"",",			
 			@"            ""Data"": {",
+			@"	              ""Version"" : 1,",
 			@"                ""cbdb55ea"": {",
 			@"                    ""Events"": [",
 			@"                        {",
diff --git a/Engine/Source/Programs/Horde/Horde.Build/Jobs/TestData/ITestData.cs b/Engine/Source/Programs/Horde/Horde.Build/Jobs/TestData/ITestData.cs
--- a/Engine/Source/Programs/Horde/Horde.Build/Jobs/TestData/ITestData.cs
+++ b/Engine/Source/Programs/Horde/Horde.Build/Jobs/TestData/ITestData.cs
@@ -160,7 +160,7 @@ namespace Horde.Build.Jobs.TestData
 		/// <summary>
 		/// The full details test data for this ref
 		/// </summary>
-		ObjectId TestDataId { get; }
+		IReadOnlyList<ObjectId> TestDataIds { get; }
 
 		/// <summary>
 		/// How long the test ran
diff --git a/Engine/Source/Programs/Horde/Horde.Build/Jobs/TestData/ITestDataCollection.cs b/Engine/Source/Programs/Horde/Horde.Build/Jobs/TestData/ITestDataCollection.cs
--- a/Engine/Source/Programs/Horde/Horde.Build/Jobs/TestData/ITestDataCollection.cs
+++ b/Engine/Source/Programs/Horde/Horde.Build/Jobs/TestData/ITestDataCollection.cs
@@ -46,10 +46,9 @@ namespace Horde.Build.Jobs.TestData
 		/// </summary>
 		/// <param name="job">The job containing the step</param>
 		/// <param name="step">The step producing the data</param>
-		/// <param name="key">Key identifying the test</param>
-		/// <param name="value">The data to store</param>
+		/// <param name="data">The test data generated by the job step</param>
 		/// <returns>The new stream document</returns>
-		Task<ITestData> AddAsync(IJob job, IJobStep step, string key, BsonDocument value);
+		Task<List<ITestData>> AddAsync(IJob job, IJobStep step, (string key, BsonDocument value)[] data);
 
 		/// <summary>
 		/// Gets a stream by ID
diff --git a/Engine/Source/Programs/Horde/Horde.Build/Jobs/TestData/TestDataCollection.cs b/Engine/Source/Programs/Horde/Horde.Build/Jobs/TestData/TestDataCollection.cs
--- a/Engine/Source/Programs/Horde/Horde.Build/Jobs/TestData/TestDataCollection.cs
+++ b/Engine/Source/Programs/Horde/Horde.Build/Jobs/TestData/TestDataCollection.cs
@@ -221,7 +221,8 @@ namespace Horde.Build.Jobs.TestData
 			public int BuildChangeList { get; set; }
 
 			[BsonRequired, BsonElement("tdid")]
-			public ObjectId TestDataId { get; set; }
+			public List<ObjectId> TestDataIds { get; set; } = new List<ObjectId>();
+			IReadOnlyList<ObjectId> ITestDataRef.TestDataIds => TestDataIds;
 
 			[BsonRequired, BsonElement("d")]
 			public TimeSpan Duration { get; set; }
@@ -237,7 +238,6 @@ namespace Horde.Build.Jobs.TestData
 
 			[BsonIgnoreIfNull, BsonElement("sts")]
 			public List<SuiteTestData>? SuiteTests { get; set; }
-
 			IReadOnlyList<ISuiteTestData>? ITestDataRef.SuiteTests => SuiteTests;
 
 			private TestDataRefDocument()
@@ -424,18 +424,25 @@ namespace Horde.Build.Jobs.TestData
 		}
 
 		/// <inheritdoc/>
-		public async Task<ITestData> AddAsync(IJob job, IJobStep step, string key, BsonDocument value)
+		public async Task<List<ITestData>> AddAsync(IJob job, IJobStep step, (string key, BsonDocument value)[] data)
 		{
 			// detailed test data
-			TestDataDocument newTestData = new TestDataDocument(job, step, key, value);
-			await _testDataDocuments.InsertOneAsync(newTestData);
+			List<TestDataDocument> documents = new List<TestDataDocument>();
+			for (int i = 0; i < data.Length; i++)
+			{
+				(string key, BsonDocument document) = data[i];
+				documents.Add(new TestDataDocument(job, step, key, document));
+				
+			}
+
+			await _testDataDocuments.InsertManyAsync(documents);
 
 			try
 			{
 				if (_settings.CurrentValue.FeatureFlags.EnableTestDataV2)
 				{
 					_logger.LogInformation("Attempting to add test report data");
-					await AddTestReportData(job, step, value, newTestData.Id);
+					await AddTestReportData(job, step, documents);
 				}									
 				else
 				{
@@ -447,7 +454,7 @@ namespace Horde.Build.Jobs.TestData
 				_logger.LogWarning(ex, "Exception while adding test data  report, jobId: {JobId} stepId: {StepId}", job.Id, step.Id);
 			}
 
-			return newTestData;
+			return documents.ConvertAll<ITestData>(x => x); ;
 		}
 
 		/// <inheritdoc/>
@@ -672,20 +679,14 @@ namespace Horde.Build.Jobs.TestData
 			return suite;
 		}
 
-		async Task AddTestReportData(IJob job, IJobStep step, BsonDocument rootDoc, ObjectId testDataId)
+		async Task AddTestReportData(IJob job, IJobStep step, List<TestDataDocument> documents)
 		{
 			List<AutomatedTestSessionData> sessions = new List<AutomatedTestSessionData>();
 			List<UnrealAutomatedTestData> tests = new List<UnrealAutomatedTestData>();
 
-			BsonArray items = rootDoc.GetValue("Items").AsBsonArray;
-			foreach (BsonValue item in items.ToList())
+			foreach (TestDataDocument item in documents)
 			{
-				if (!item.IsBsonDocument)
-				{
-					continue;
-				}
-
-				BsonDocument testData = item.AsBsonDocument;
+				BsonDocument testData = item.Data;
 
 				int version = testData.GetValue("Version", new BsonInt32(0)).AsInt32;
 				if (version < 1)
@@ -695,34 +696,33 @@ namespace Horde.Build.Jobs.TestData
 				}
 
 				BsonValue? value;
-				if (!testData.TryGetPropertyValue("Data.Type", BsonType.String, out value))
+				if (!testData.TryGetPropertyValue("Type", BsonType.String, out value))
 				{
 					continue;
 				}
 
-				string type = value.AsString;
-				BsonDocument data = testData.GetValue("Data").AsBsonDocument;
+				string type = value.AsString;				
 
 				if (type == "Simple Report")
 				{
-					await AddSimpleReportData(job, step, testDataId, data);
+					await AddSimpleReportData(job, step, item.Id, testData);
 					continue;
 				}
 
 				if (type == "Automated Test Session")
 				{
-					sessions.Add(BsonSerializer.Deserialize<AutomatedTestSessionData>(data));
+					sessions.Add(BsonSerializer.Deserialize<AutomatedTestSessionData>(testData));
 				}
 
 				if (type == "Unreal Automated Tests")
 				{
-					tests.Add(BsonSerializer.Deserialize<UnrealAutomatedTestData>(data));
+					tests.Add(BsonSerializer.Deserialize<UnrealAutomatedTestData>(testData));
 				}
 			}
 
 			if (sessions.Count > 0)
 			{
-				await AddTestSessionReportData(job, step, testDataId, sessions, tests);
+				await AddTestSessionReportData(job, step, documents, sessions, tests);
 			}
 		}
 
@@ -830,7 +830,7 @@ namespace Horde.Build.Jobs.TestData
 				testRef.BuildChangeList = testData.BuildChangeList;
 				testRef.Duration = TimeSpan.FromSeconds(testData.TotalDurationSeconds);
 				testRef.TestId = test.Id;
-				testRef.TestDataId = testDataId;
+				testRef.TestDataIds.Add(testDataId);
 				testRef.Outcome = TestOutcome.Unspecified;
 
 				switch (testData.TestResult)
@@ -854,7 +854,7 @@ namespace Horde.Build.Jobs.TestData
 			}
 		}
 
-		async Task AddTestSessionReportData(IJob job, IJobStep step, ObjectId testDataId, List<AutomatedTestSessionData> sessions, List<UnrealAutomatedTestData> tests)
+		async Task AddTestSessionReportData(IJob job, IJobStep step, List<TestDataDocument> documents, List<AutomatedTestSessionData> sessions, List<UnrealAutomatedTestData> tests)
 		{
 			HashSet<string> suites = new HashSet<string>();
 
@@ -960,7 +960,7 @@ namespace Horde.Build.Jobs.TestData
 						testRef.Metadata = metaId.Value;
 						testRef.BuildChangeList = job.Change;
 						testRef.Duration = TimeSpan.FromSeconds(session.TestSessionInfo!.TimeElapseSec);
-						testRef.TestDataId = testDataId;
+						testRef.TestDataIds = documents.Select(x => x.Id).ToList();
 						testRef.Outcome = TestOutcome.Unspecified;
 						testRef.SuiteId = testSuite.Id;
 						testRef.SuiteTests = new List<SuiteTestData>();
diff --git a/Engine/Source/Programs/Horde/Horde.Build/Jobs/TestData/TestDataController.cs b/Engine/Source/Programs/Horde/Horde.Build/Jobs/TestData/TestDataController.cs
--- a/Engine/Source/Programs/Horde/Horde.Build/Jobs/TestData/TestDataController.cs
+++ b/Engine/Source/Programs/Horde/Horde.Build/Jobs/TestData/TestDataController.cs
@@ -175,9 +175,9 @@ namespace Horde.Build.Jobs.TestData
 			{
 				return NotFound();
 			}
-
-			ITestData testData = await _testDataCollection.AddAsync(job, jobStep, request.Key, new BsonDocument(request.Data));
-			return new CreateTestDataResponse(testData.Id.ToString());
+			
+			List<ITestData> testData = await _testDataCollection.AddAsync(job, jobStep, new (string key, BsonDocument value)[] { (request.Key, new BsonDocument(request.Data))});
+			return new CreateTestDataResponse(testData[0].Id.ToString());
 		}
 
 		/// <summary>
diff --git a/Engine/Source/Programs/Horde/Horde.Build/Jobs/TestData/TestDataMessages.cs b/Engine/Source/Programs/Horde/Horde.Build/Jobs/TestData/TestDataMessages.cs
--- a/Engine/Source/Programs/Horde/Horde.Build/Jobs/TestData/TestDataMessages.cs
+++ b/Engine/Source/Programs/Horde/Horde.Build/Jobs/TestData/TestDataMessages.cs
@@ -327,7 +327,7 @@ namespace Horde.Build.Jobs.TestData
 		/// <summary>
 		/// The full details test data for this ref
 		/// </summary>
-		public string TestDataId { get; set; }
+		public List<string> TestDataIds { get; set; }
 
 		/// <summary>
 		/// How long the test ran
@@ -372,7 +372,7 @@ namespace Horde.Build.Jobs.TestData
 		public GetTestDataRefResponse(ITestDataRef testData)
 		{
 			StreamId = testData.StreamId.ToString();
-			TestDataId = testData.TestDataId.ToString();
+			TestDataIds = testData.TestDataIds.Select(x => x.ToString()).ToList();
 			Duration = testData.Duration;
 			BuildChangeList = testData.BuildChangeList;
 			MetaId = testData.Metadata.ToString();
diff --git a/Engine/Source/Programs/Horde/Horde.Build/Server/RpcService.cs b/Engine/Source/Programs/Horde/Horde.Build/Server/RpcService.cs
--- a/Engine/Source/Programs/Horde/Horde.Build/Server/RpcService.cs
+++ b/Engine/Source/Programs/Horde/Horde.Build/Server/RpcService.cs
@@ -1053,33 +1053,53 @@ namespace Horde.Build.Server
 			IJob? job = null;
 			IJobStep? jobStep = null;
 
+			List<(string key, BsonDocument document)> data = new List<(string key, BsonDocument document)>();
+
 			while (await reader.MoveNext())
 			{
 				UploadTestDataRequest request = reader.Current;
 
 				JobId jobId = new JobId(request.JobId);
-				if (job == null || jobId != job.Id)
-				{
+				if (job == null)
+				{					
 					job = await _jobService.GetJobAsync(jobId);
 					if (job == null)
 					{
 						throw new StructuredRpcException(StatusCode.NotFound, "Unable to find job {JobId}", jobId);
 					}
-					jobStep = null;
 				}
+				else if (jobId != job.Id)
+				{
+					throw new StructuredRpcException(StatusCode.InvalidArgument, "Job {JobId} does not match previous Job {JobId} in request", jobId, job.Id);
+				}
+
 
 				SubResourceId jobStepId = request.JobStepId.ToSubResourceId();
-				if (jobStep == null || jobStepId != jobStep.Id)
+
+				if (jobStep == null)
 				{
 					if (!job.TryGetStep(jobStepId, out jobStep))
 					{
 						throw new StructuredRpcException(StatusCode.NotFound, "Unable to find step {JobStepId} on job {JobId}", jobStepId, jobId);
 					}
 				}
-
+				else if (jobStep.Id != jobStepId)
+				{
+					throw new StructuredRpcException(StatusCode.InvalidArgument, "Job step {JobStepId} does not match previous Job step {JobStepId} in request", jobStepId, jobStep.Id);
+				}
+				
 				string text = Encoding.UTF8.GetString(request.Value.ToArray());
 				BsonDocument document = BsonSerializer.Deserialize<BsonDocument>(text);
-				await _testData.AddAsync(job, jobStep, request.Key, document);
+				data.Add((request.Key, document));
+			}
+
+			if (job != null && jobStep != null)
+			{
+				await _testData.AddAsync(job, jobStep, data.ToArray());
+			}			
+			else
+			{
+				throw new StructuredRpcException(StatusCode.NotFound, "Unable to get job or step for test data upload" );
 			}
 
 			return new UploadTestDataResponse();
