Horde: Add copy to clipboard feature on token page

#rnx
#skipci

diff --git a/Engine/Source/Programs/Horde/HordeDashboard/src/components/AdminToken.tsx b/Engine/Source/Programs/Horde/HordeDashboard/src/components/AdminToken.tsx
--- a/Engine/Source/Programs/Horde/HordeDashboard/src/components/AdminToken.tsx
+++ b/Engine/Source/Programs/Horde/HordeDashboard/src/components/AdminToken.tsx
@@ -3,6 +3,7 @@
 import React, { useState } from 'react';
 import { Text, Stack, Spinner, SpinnerSize } from '@fluentui/react';
 import backend from '../backend';
+import {copyToClipboard} from "../base/utilities/clipboard";
 
 
 export const AdminToken: React.FC = () => {
@@ -15,16 +16,17 @@ export const AdminToken: React.FC = () => {
 		backend.getAdminToken().then(token => setToken(token)).catch(reason => setTokenError(reason));
 		return <Stack horizontal verticalAlign="center" tokens={{childrenGap:24}} styles={{root:{padding:24}}}>
 			<Text variant="mediumPlus">Getting Token...</Text>
-			<Spinner size={SpinnerSize.large}/>;
+			<Spinner size={SpinnerSize.large} />
 		</Stack>;
-		
 	}
 
 	return (
 		<Stack horizontalAlign="center" styles={{ root: { width: "100%", padding: 12 } }}>
 			<Stack styles={{ root: { width: 1440 } }} tokens={{ childrenGap: 18 }}>
-				<div >{token ? "Horde Access Token:" : "Error Getting Token:"}</div>
-				<div style={{ wordBreak: "break-word" }}> {token ?? tokenError} </div>
+				<h1>Horde Access Token</h1>
+				{token && <a href="#" onClick={() => copyToClipboard(token)}>Copy to clipboard</a>}
+				{!token && <div>Error getting token</div>}
+				<textarea readOnly={true} style={{ wordBreak: "break-word", height: 400 }} value={token ? token : tokenError}></textarea>
 			</Stack>
 		</Stack>
 	);
