Horde: Respect unset values for minAgents and numReserveAgents for pool updates in dashboard

#skipci

diff --git a/Engine/Source/Programs/Horde/HordeDashboard/src/components/PoolView.tsx b/Engine/Source/Programs/Horde/HordeDashboard/src/components/PoolView.tsx
--- a/Engine/Source/Programs/Horde/HordeDashboard/src/components/PoolView.tsx
+++ b/Engine/Source/Programs/Horde/HordeDashboard/src/components/PoolView.tsx
@@ -23,6 +23,9 @@ type PendingBatch = {
    batch: GetBatchResponse;
 }
 
+// UI visible text for values that are undefined
+const UNSET_VALUE: string = "-";
+
 class PoolHandler extends PollBase {
 
    constructor(pollTime = 15000) {
@@ -681,7 +684,7 @@ const AutoScalerPanel: React.FC = () => {
          color = pool.properties["Color"];
       }
 
-      setState({ name: pool.name, color: color, autoscale: pool.enableAutoscaling ?? false, minAgents: pool.minAgents ?? 0, reserveAgents: pool.numReserveAgents ?? 0, strategy: pool.sizeStrategy ?? PoolSizeStrategy.LeaseUtilization, conformInterval: pool.conformInterval ?? 24 });
+	  setState({ name: pool.name, color: color, autoscale: pool.enableAutoscaling ?? false, minAgents: pool.minAgents, reserveAgents: pool.numReserveAgents, strategy: pool.sizeStrategy ?? PoolSizeStrategy.LeaseUtilization, conformInterval: pool.conformInterval ?? 24 });
       return null;
    }
 
@@ -714,14 +717,14 @@ const AutoScalerPanel: React.FC = () => {
                const update: UpdatePoolRequest = {
                   name: poolName,
                   enableAutoscaling: state.autoscale ?? false,
-                  minAgents: state.minAgents ?? 0,
-                  numReserveAgents: state.reserveAgents ?? 0,
+                  minAgents: state.minAgents,
+                  numReserveAgents: state.reserveAgents,
                   sizeStrategy: state.strategy!,
                   conformInterval: typeof (state.conformInterval) === "number" ? state.conformInterval : 24,
                   properties: { Color: state.color! }
                }
 
-               await backend.updatePool(pool.id, update);
+			   await backend.updatePool(pool.id, update);
 
                handler.pool = await backend.getPool(pool.id);
 
@@ -751,8 +754,8 @@ const AutoScalerPanel: React.FC = () => {
                <Text>{`Pool Name: ${poolName}`}</Text>
                <Text>{`Conform Interval: ${state.conformInterval ?? "24"}`}</Text>
                <Text>{`Autoscale: ${state.autoscale ? "On" : "Off"}`}</Text>
-               <Text>{`Minimum Agents: ${state.minAgents ?? 0}`}</Text>
-               <Text>{`Reserve Agents: ${state.reserveAgents ?? 0}`}</Text>
+               <Text>{`Minimum Agents: ${state.minAgents ?? UNSET_VALUE}`}</Text>
+               <Text>{`Reserve Agents: ${state.reserveAgents ?? UNSET_VALUE}`}</Text>
                <Text>{`Pool Strategy: ${state.strategy}`}</Text>
 
             </Stack>
@@ -818,11 +821,11 @@ const AutoScalerPanel: React.FC = () => {
                   <ComboBox styles={{ root: { width: 180 } }} disabled={disabled} label="Strategy" selectedKey={state.strategy} options={stratItems} onChange={(ev, option, index, value) => {
                      setState({ ...state, modified: true, strategy: option!.data })
                   }} />
-                  <SpinButton styles={{ root: { width: 128 } }} disabled={disabled} label="Minimum Agents" value={state.minAgents?.toString() ?? "0"} labelPosition={Position.top} onChange={(ev, value) => {
-                     setState({ ...state, modified: true, minAgents: parseInt(value ?? "0") })
+                  <SpinButton styles={{ root: { width: 128 } }} disabled={disabled} label="Minimum Agents" value={state.minAgents?.toString() ?? UNSET_VALUE} labelPosition={Position.top} onChange={(ev, value) => {
+                     setState({ ...state, modified: true, minAgents: value == UNSET_VALUE ? undefined : parseInt(value ?? "0") })
                   }} />
-                  <SpinButton styles={{ root: { width: 128 } }} disabled={disabled} label="Reserve Agents" value={state.reserveAgents?.toString() ?? "0"} labelPosition={Position.top} onChange={(ev, value) => {
-                     setState({ ...state, modified: true, reserveAgents: parseInt(value ?? "0") })
+                  <SpinButton styles={{ root: { width: 128 } }} disabled={disabled} label="Reserve Agents" value={state.reserveAgents?.toString() ?? UNSET_VALUE} labelPosition={Position.top} onChange={(ev, value) => {
+                     setState({ ...state, modified: true, reserveAgents: value == UNSET_VALUE ? undefined : parseInt(value ?? "0") })
                   }} />
                </Stack>
                <Stack>
