Horde: Remove feature flag for new test data collections
#rnx

diff --git a/Engine/Source/Programs/Horde/Horde.Build.Tests/TestDataTests.cs b/Engine/Source/Programs/Horde/Horde.Build.Tests/TestDataTests.cs
--- a/Engine/Source/Programs/Horde/Horde.Build.Tests/TestDataTests.cs
+++ b/Engine/Source/Programs/Horde/Horde.Build.Tests/TestDataTests.cs
@@ -156,115 +156,106 @@ namespace Horde.Build.Tests
 			return job.Object;
 		}
 
-
-
 		[TestMethod]
 		public async Task SimpleReportTest()
 		{
-			if (ServerSettings.FeatureFlags.EnableTestDataV2)
-			{
-				string[] streamIds = new string[] { _mainStreamId.ToString(), _releaseStreamId.ToString() };
+			string[] streamIds = new string[] { _mainStreamId.ToString(), _releaseStreamId.ToString() };
 
-				IJob job = CreateJob(_mainStreamId, 105, "Test Build", _graph);
-				IJobStep step = job.Batches[0].Steps[0];
-				IJob job2 = CreateJob(_releaseStreamId, 105, "Test Build", _graph);
-				IJobStep step2 = job2.Batches[0].Steps[0];
+			IJob job = CreateJob(_mainStreamId, 105, "Test Build", _graph);
+			IJobStep step = job.Batches[0].Steps[0];
+			IJob job2 = CreateJob(_releaseStreamId, 105, "Test Build", _graph);
+			IJobStep step2 = job2.Batches[0].Steps[0];
 
-				BsonDocument testData = BsonDocument.Parse(String.Join('\n', _simpleTestDataLines));
+			BsonDocument testData = BsonDocument.Parse(String.Join('\n', _simpleTestDataLines));
 
-				List<(string key, BsonDocument)> data = new List<(string key, BsonDocument)>();
-				BsonArray items = testData.GetValue("Items").AsBsonArray;
-				foreach (BsonValue item in items.ToList())
-				{
-					BsonDocument value = item.AsBsonDocument.GetValue("Data").AsBsonDocument;
-					data.Add(("Simple Report Key", value));
-				}
+			List<(string key, BsonDocument)> data = new List<(string key, BsonDocument)>();
+			BsonArray items = testData.GetValue("Items").AsBsonArray;
+			foreach (BsonValue item in items.ToList())
+			{
+				BsonDocument value = item.AsBsonDocument.GetValue("Data").AsBsonDocument;
+				data.Add(("Simple Report Key", value));
+			}
 
-				await TestDataCollection.AddAsync(job, step, data.ToArray());
-				await TestDataCollection.AddAsync(job2, step2, data.ToArray());
+			await TestDataCollection.AddAsync(job, step, data.ToArray());
+			await TestDataCollection.AddAsync(job2, step2, data.ToArray());
 
-				ActionResult<List<GetTestStreamResponse>> streamResult = await TestDataController.GetTestStreamsAsync(streamIds);
-				Assert.IsNotNull(streamResult);
-				Assert.IsNotNull(streamResult.Value);
-				List<GetTestStreamResponse> streams = streamResult.Value;
+			ActionResult<List<GetTestStreamResponse>> streamResult = await TestDataController.GetTestStreamsAsync(streamIds);
+			Assert.IsNotNull(streamResult);
+			Assert.IsNotNull(streamResult.Value);
+			List<GetTestStreamResponse> streams = streamResult.Value;
 
-				Assert.AreEqual(2, streams.Count);
-				Assert.AreEqual(1, streams[0].Tests.Count);
-				Assert.AreEqual(1, streams[0].TestMetadata.Count);
-				Assert.AreEqual(0, streams[0].TestSuites.Count);
-				Assert.AreEqual(1, streams[1].Tests.Count);
-				Assert.AreEqual(1, streams[1].TestMetadata.Count);
-				Assert.AreEqual(0, streams[1].TestSuites.Count);
+			Assert.AreEqual(2, streams.Count);
+			Assert.AreEqual(1, streams[0].Tests.Count);
+			Assert.AreEqual(1, streams[0].TestMetadata.Count);
+			Assert.AreEqual(0, streams[0].TestSuites.Count);
+			Assert.AreEqual(1, streams[1].Tests.Count);
+			Assert.AreEqual(1, streams[1].TestMetadata.Count);
+			Assert.AreEqual(0, streams[1].TestSuites.Count);
 
-				Assert.AreEqual(streams[0].TestMetadata[0].Id, streams[1].TestMetadata[0].Id);
+			Assert.AreEqual(streams[0].TestMetadata[0].Id, streams[1].TestMetadata[0].Id);
 
-				ActionResult<List<GetTestDataRefResponse>> refResult = await TestDataController.GetTestDataRefAsync(streamIds);
-				Assert.IsNotNull(refResult);
-				Assert.IsNotNull(refResult.Value);
-				List<GetTestDataRefResponse> refs = refResult.Value;
+			ActionResult<List<GetTestDataRefResponse>> refResult = await TestDataController.GetTestDataRefAsync(streamIds);
+			Assert.IsNotNull(refResult);
+			Assert.IsNotNull(refResult.Value);
+			List<GetTestDataRefResponse> refs = refResult.Value;
 
-				Assert.AreEqual(2, refs.Count);
+			Assert.AreEqual(2, refs.Count);
 
-				Assert.AreEqual(1, refs[0].TestDataIds.Count);
-				Assert.AreEqual(1, refs[1].TestDataIds.Count);
-			}
+			Assert.AreEqual(1, refs[0].TestDataIds.Count);
+			Assert.AreEqual(1, refs[1].TestDataIds.Count);			
 		}
 
 		[TestMethod]
 		public async Task SessionReportTest()
 		{
-			if (ServerSettings.FeatureFlags.EnableTestDataV2)
-			{
-				string[] streamIds = new string[] { _mainStreamId.ToString(), _releaseStreamId.ToString() };
-
-				IJob job = CreateJob(_mainStreamId, 105, "Test Build", _graph);
-				IJobStep step = job.Batches[0].Steps[0];
-				IJob job2 = CreateJob(_releaseStreamId, 105, "Test Build", _graph);
-				IJobStep step2 = job2.Batches[0].Steps[0];
+			string[] streamIds = new string[] { _mainStreamId.ToString(), _releaseStreamId.ToString() };
 
-				BsonDocument testData = BsonDocument.Parse(String.Join('\n', _testSessionDataLines));
+			IJob job = CreateJob(_mainStreamId, 105, "Test Build", _graph);
+			IJobStep step = job.Batches[0].Steps[0];
+			IJob job2 = CreateJob(_releaseStreamId, 105, "Test Build", _graph);
+			IJobStep step2 = job2.Batches[0].Steps[0];
 
-				List<(string key, BsonDocument)> data = new List<(string key, BsonDocument)>();
-				BsonArray items = testData.GetValue("Items").AsBsonArray;
-				foreach (BsonValue item in items.ToList())
-				{
-					BsonDocument value = item.AsBsonDocument.GetValue("Data").AsBsonDocument;
-					data.Add(("Session Report Key", value));
-				}
+			BsonDocument testData = BsonDocument.Parse(String.Join('\n', _testSessionDataLines));
 
-				await TestDataCollection.AddAsync(job, step, data.ToArray());
-				await TestDataCollection.AddAsync(job2, step2, data.ToArray());
+			List<(string key, BsonDocument)> data = new List<(string key, BsonDocument)>();
+			BsonArray items = testData.GetValue("Items").AsBsonArray;
+			foreach (BsonValue item in items.ToList())
+			{
+				BsonDocument value = item.AsBsonDocument.GetValue("Data").AsBsonDocument;
+				data.Add(("Session Report Key", value));
+			}
 
-				ActionResult<List<GetTestStreamResponse>> streamResult = await TestDataController.GetTestStreamsAsync(streamIds);
-				Assert.IsNotNull(streamResult);
-				Assert.IsNotNull(streamResult.Value);
-				List<GetTestStreamResponse> streams = streamResult.Value;
+			await TestDataCollection.AddAsync(job, step, data.ToArray());
+			await TestDataCollection.AddAsync(job2, step2, data.ToArray());
 
-				Assert.AreEqual(2, streams.Count);
+			ActionResult<List<GetTestStreamResponse>> streamResult = await TestDataController.GetTestStreamsAsync(streamIds);
+			Assert.IsNotNull(streamResult);
+			Assert.IsNotNull(streamResult.Value);
+			List<GetTestStreamResponse> streams = streamResult.Value;
 
-				Assert.AreEqual(4, streams[0].Tests.Count);				
-				Assert.AreEqual(1, streams[0].TestMetadata.Count);
-				Assert.AreEqual(1, streams[0].TestSuites.Count);
-				Assert.AreEqual(4, streams[0].TestSuites[0].Tests.Count);
+			Assert.AreEqual(2, streams.Count);
 
-				Assert.AreEqual(4, streams[1].Tests.Count);
-				Assert.AreEqual(1, streams[1].TestMetadata.Count);
-				Assert.AreEqual(1, streams[1].TestSuites.Count);
-				Assert.AreEqual(4, streams[1].TestSuites[0].Tests.Count);
+			Assert.AreEqual(4, streams[0].Tests.Count);				
+			Assert.AreEqual(1, streams[0].TestMetadata.Count);
+			Assert.AreEqual(1, streams[0].TestSuites.Count);
+			Assert.AreEqual(4, streams[0].TestSuites[0].Tests.Count);
 
-				Assert.AreEqual(streams[0].TestMetadata[0].Id, streams[1].TestMetadata[0].Id);
+			Assert.AreEqual(4, streams[1].Tests.Count);
+			Assert.AreEqual(1, streams[1].TestMetadata.Count);
+			Assert.AreEqual(1, streams[1].TestSuites.Count);
+			Assert.AreEqual(4, streams[1].TestSuites[0].Tests.Count);
 
-				ActionResult<List<GetTestDataRefResponse>> refResult = await TestDataController.GetTestDataRefAsync(streamIds);
-				Assert.IsNotNull(refResult);
-				Assert.IsNotNull(refResult.Value);
-				List<GetTestDataRefResponse> refs = refResult.Value;
+			Assert.AreEqual(streams[0].TestMetadata[0].Id, streams[1].TestMetadata[0].Id);
 
-				Assert.AreEqual(2, refs.Count);
+			ActionResult<List<GetTestDataRefResponse>> refResult = await TestDataController.GetTestDataRefAsync(streamIds);
+			Assert.IsNotNull(refResult);
+			Assert.IsNotNull(refResult.Value);
+			List<GetTestDataRefResponse> refs = refResult.Value;
 
-				Assert.AreEqual(3, refs[0].TestDataIds.Count);
-				Assert.AreEqual(3, refs[1].TestDataIds.Count);
+			Assert.AreEqual(2, refs.Count);
 
-			}
+			Assert.AreEqual(3, refs[0].TestDataIds.Count);
+			Assert.AreEqual(3, refs[1].TestDataIds.Count);			
 
 		}
 
diff --git a/Engine/Source/Programs/Horde/Horde.Build/Jobs/TestData/TestDataCollection.cs b/Engine/Source/Programs/Horde/Horde.Build/Jobs/TestData/TestDataCollection.cs
--- a/Engine/Source/Programs/Horde/Horde.Build/Jobs/TestData/TestDataCollection.cs
+++ b/Engine/Source/Programs/Horde/Horde.Build/Jobs/TestData/TestDataCollection.cs
@@ -333,26 +333,21 @@ namespace Horde.Build.Jobs.TestData
 			indexes.Add(keys => keys.Ascending(x => x.JobId).Ascending(x => x.StepId).Ascending(x => x.Key), unique: true);
 			_testDataDocuments = mongoService.GetCollection<TestDataDocument>("TestData", indexes);
 
-			if (_settings.CurrentValue.FeatureFlags.EnableTestDataV2)
-			{
-				_logger.LogInformation("EnableTestDataV2 is true, creating new test data collections");
-
-				List<MongoIndex<TestMetaDocument>> metaIndexes = new List<MongoIndex<TestMetaDocument>>();
-				metaIndexes.Add(keys => keys.Ascending(x => x.ProjectName).Ascending(x => x.Platforms).Ascending(x => x.Configurations));
-				_testMeta = mongoService.GetCollection<TestMetaDocument>("TestData.Meta", metaIndexes);
+			List<MongoIndex<TestMetaDocument>> metaIndexes = new List<MongoIndex<TestMetaDocument>>();
+			metaIndexes.Add(keys => keys.Ascending(x => x.ProjectName).Ascending(x => x.Platforms).Ascending(x => x.Configurations).Ascending(x => x.BuildTargets).Ascending(x => x.RHI), unique: true);
+			_testMeta = mongoService.GetCollection<TestMetaDocument>("TestData.Meta", metaIndexes);
 
-				List<MongoIndex<TestDocument>> testIndexes = new List<MongoIndex<TestDocument>>();
-				testIndexes.Add(keys => keys.Ascending(x => x.StreamId));
-				_tests = mongoService.GetCollection<TestDocument>("TestData.Tests", testIndexes);
+			List<MongoIndex<TestDocument>> testIndexes = new List<MongoIndex<TestDocument>>();
+			testIndexes.Add(keys => keys.Ascending(x => x.StreamId));
+			_tests = mongoService.GetCollection<TestDocument>("TestData.Tests", testIndexes);
 
-				List<MongoIndex<TestSuiteDocument>> testSuiteIndexes = new List<MongoIndex<TestSuiteDocument>>();
-				testSuiteIndexes.Add(keys => keys.Ascending(x => x.StreamId));
-				_testSuites = mongoService.GetCollection<TestSuiteDocument>("TestData.TestSuites", testSuiteIndexes);
+			List<MongoIndex<TestSuiteDocument>> testSuiteIndexes = new List<MongoIndex<TestSuiteDocument>>();
+			testSuiteIndexes.Add(keys => keys.Ascending(x => x.StreamId));
+			_testSuites = mongoService.GetCollection<TestSuiteDocument>("TestData.TestSuites", testSuiteIndexes);
 
-				List<MongoIndex<TestDataRefDocument>> testRefIndexes = new List<MongoIndex<TestDataRefDocument>>();
-				testRefIndexes.Add(keys => keys.Ascending(x => x.StreamId).Ascending(x => x.Metadata).Descending(x => x.BuildChangeList).Ascending(x => x.TestId).Ascending(x => x.SuiteId));
-				_testRefs = mongoService.GetCollection<TestDataRefDocument>("TestData.TestRefs", testRefIndexes);
-			}
+			List<MongoIndex<TestDataRefDocument>> testRefIndexes = new List<MongoIndex<TestDataRefDocument>>();
+			testRefIndexes.Add(keys => keys.Ascending(x => x.StreamId).Ascending(x => x.Metadata).Descending(x => x.BuildChangeList).Ascending(x => x.TestId).Ascending(x => x.SuiteId));
+			_testRefs = mongoService.GetCollection<TestDataRefDocument>("TestData.TestRefs", testRefIndexes);
 		}
 
 		/// <summary>
@@ -439,15 +434,7 @@ namespace Horde.Build.Jobs.TestData
 
 			try
 			{
-				if (_settings.CurrentValue.FeatureFlags.EnableTestDataV2)
-				{
-					_logger.LogInformation("Attempting to add test report data");
-					await AddTestReportData(job, step, documents);
-				}									
-				else
-				{
-					_logger.LogInformation("Test report data skipped due to disabled feature flag");
-				}
+				await AddTestReportData(job, step, documents);
 			}
 			catch (Exception ex)
 			{
diff --git a/Engine/Source/Programs/Horde/Horde.Build/ServerSettings.cs b/Engine/Source/Programs/Horde/Horde.Build/ServerSettings.cs
--- a/Engine/Source/Programs/Horde/Horde.Build/ServerSettings.cs
+++ b/Engine/Source/Programs/Horde/Horde.Build/ServerSettings.cs
@@ -180,11 +180,6 @@ namespace Horde.Build
 		/// </summary>
 		public bool LimitConcurrentLogChunkWriting { get; set; } = false;
 
-		/// <summary>
-		/// Enable new test data collections and service, feature flag added 22/11/7
-		/// </summary>
-		public bool EnableTestDataV2 { get; set; } = false;
-
 	}
 
 	/// <summary>
