Horde: Add fallback property for identifying OS name to display in agent view

#skipci

diff --git a/Engine/Source/Programs/Horde/HordeDashboard/src/components/AgentView.tsx b/Engine/Source/Programs/Horde/HordeDashboard/src/components/AgentView.tsx
--- a/Engine/Source/Programs/Horde/HordeDashboard/src/components/AgentView.tsx
+++ b/Engine/Source/Programs/Horde/HordeDashboard/src/components/AgentView.tsx
@@ -160,6 +160,12 @@ function getAgentCapability(agent: AgentData, inProp: string) {
    return null;
 }
 
+function getAgentOs(agent: AgentData) {
+	const osDist = getAgentCapability(agent, "OSDistribution");
+	const osFamily = getAgentCapability(agent, "OSFamily");
+	return osDist ?? osFamily;
+}
+
 function getTaskTime(agent: GetAgentResponse): number {
 
    let max: moment.Duration | undefined
@@ -1797,8 +1803,7 @@ export const AgentView: React.FC = observer(() => {
                }
                break;
             case 'OS':
-               const OS = getAgentCapability(agent, "OSDistribution");
-               toCheck = OS ?? "Unknown";
+               toCheck = getAgentOs(agent) ?? "Unknown";
                break;
             case 'CPU':
                const CPU = getAgentCapability(agent, "CPU");
@@ -1865,7 +1870,7 @@ export const AgentView: React.FC = observer(() => {
             }
             return (left.comment ?? "").toLocaleLowerCase().localeCompare(right.comment ?? "");
          case 'systemInfoOS':
-            return (getAgentCapability(left, "OSDistribution") ?? "Unknown").toLocaleLowerCase().localeCompare((getAgentCapability(right, "OSDistribution") ?? "Unknown"));
+            return ((getAgentOs(left) ?? "Unknown").toLocaleLowerCase().localeCompare((getAgentOs(right) ?? "Unknown")));
          case 'systemInfoCPU':
             return (getAgentCapability(left, "CPU") ?? "Unknown").toLocaleLowerCase().localeCompare((getAgentCapability(right, "CPU") ?? "Unknown"));
          case 'systemInfoRAM':
@@ -2089,23 +2094,23 @@ export const AgentView: React.FC = observer(() => {
       return <Icon iconName="FullCircle" className={className} />;
    }
 
-   function getPropFromDevice(agent: AgentData, inProp: string) {
-      const unknownElement = <Stack styles={{ root: { height: '100%' } }} horizontal horizontalAlign={"center"}><Stack.Item align={"center"}>Unknown</Stack.Item></Stack>;
-      let field = getAgentCapability(agent, inProp);
-      if (field === null)
-         return unknownElement;
-
-      if (inProp.indexOf("RAM") !== -1) {
-         field += " GB";
-      }
-      return (
-         <Stack styles={{ root: { height: '100%' } }} horizontal horizontalAlign={"center"}>
-            <Stack.Item align={"center"} className={agentStyles.ellipsesStackItem}>
-               <Text title={field} key={`sysInfo_${agent.id}_${inProp}`}>{`${field}`}</Text>
-            </Stack.Item>
-         </Stack>
-      );
-   }
+	function getPropFromDevice(agent: AgentData, propKey: string, propValue: string | null = null) {
+		const unknownElement = <Stack styles={{ root: { height: '100%' } }} horizontal horizontalAlign={"center"}><Stack.Item align={"center"}>Unknown</Stack.Item></Stack>;
+		let prop = propValue ?? getAgentCapability(agent, propKey);
+		if (prop === null)
+			return unknownElement;
+
+		if (propKey.indexOf("RAM") !== -1) {
+			prop += " GB";
+		}
+		return (
+			<Stack styles={{ root: { height: '100%' } }} horizontal horizontalAlign={"center"}>
+				<Stack.Item align={"center"} className={agentStyles.ellipsesStackItem}>
+					<Text title={prop} key={`sysInfo_${agent.id}_${propKey}`}>{`${prop}`}</Text>
+				</Stack.Item>
+			</Stack>
+		);
+	}
 
    // when an actual item is drawn
    function onRenderAgentListItem(agent: AgentData, index?: number, column?: IColumn) {
@@ -2380,7 +2385,7 @@ export const AgentView: React.FC = observer(() => {
                </Stack.Item>
             </Stack>);
          case 'systemInfoOS':
-            return getPropFromDevice(agent, "OSDistribution");
+            return getPropFromDevice(agent, "OSDistribution", getAgentOs(agent));
          case 'systemInfoCPU':
             return getPropFromDevice(agent, "CPU");
          case 'systemInfoRAM':
