Horde: Prefer native client in agent operations.

diff --git a/Engine/Source/Programs/Horde/Horde.Agent/Utility/WorkspaceInfo.cs b/Engine/Source/Programs/Horde/Horde.Agent/Utility/WorkspaceInfo.cs
--- a/Engine/Source/Programs/Horde/Horde.Agent/Utility/WorkspaceInfo.cs
+++ b/Engine/Source/Programs/Horde/Horde.Agent/Utility/WorkspaceInfo.cs
@@ -141,7 +141,7 @@ namespace Horde.Agent.Utility
 			}
 
 			// Create the connection
-			IPerforceConnection perforce = await PerforceConnection.CreateAsync(new PerforceSettings(serverAndPort, userName), logger);
+			IPerforceConnection perforce = await PerforceConnection.CreateAsync(new PerforceSettings(serverAndPort, userName) { PreferNativeClient = true }, logger);
 			if (userName != null)
 			{
 				if (password != null)
