Horde: Fix compile errors.

diff --git a/Engine/Source/Programs/Horde/Horde.Build/Perforce/ReplicationService.cs b/Engine/Source/Programs/Horde/Horde.Build/Perforce/ReplicationService.cs
--- a/Engine/Source/Programs/Horde/Horde.Build/Perforce/ReplicationService.cs
+++ b/Engine/Source/Programs/Horde/Horde.Build/Perforce/ReplicationService.cs
@@ -521,7 +521,6 @@ namespace Horde.Build.Perforce
 
 			public async Task AppendAsync(ReadOnlyMemory<byte> data, ChunkingOptions options, TreeWriter writer, CancellationToken cancellationToken)
 			{
-				_hash.AppendData(data.Span);
 				_node = await _node.AppendAsync(data, options, writer, cancellationToken);
 			}
 		}
@@ -704,92 +703,75 @@ namespace Horde.Build.Perforce
 					Stopwatch gcTimer = new Stopwatch();
 
 					Dictionary<int, FileWriter> handles = new Dictionary<int, FileWriter>();
-					try
+					await foreach (PerforceResponse response in perforce.StreamCommandAsync("sync", Array.Empty<string>(), syncPaths, null, typeof(SyncRecord), true, default))
 					{
-						await foreach (PerforceResponse response in perforce.StreamCommandAsync("sync", Array.Empty<string>(), syncPaths, null, typeof(SyncRecord), true, default))
+						PerforceError? error = response.Error;
+						if (error != null)
 						{
-							PerforceError? error = response.Error;
-							if (error != null)
+							if (error.Generic == PerforceGenericCode.Empty)
 							{
-								if (error.Generic == PerforceGenericCode.Empty)
-								{
-									continue;
-								}
-								else
-								{
-									throw new ReplicationException($"Perforce error while replicating content - {error}");
-								}
+								continue;
+							}
+							else
+							{
+								throw new ReplicationException($"Perforce error while replicating content - {error}");
 							}
+						}
 
-							processTimer.Start();
+						processTimer.Start();
 
-							PerforceIo? io = response.Io;
-							if (io != null)
+						PerforceIo? io = response.Io;
+						if (io != null)
+						{
+							if (io.Command == PerforceIoCommand.Open)
 							{
-								if (io.Command == PerforceIoCommand.Open)
-								{
-									Utf8String file = GetClientRelativePath(io.Payload, clientInfo.Client.Root);
-									int offset = GetFileOffset(file);
-
-									long fileSize = 0;
-									if (!pathToDirectory.TryGetValue(file.Substring(0, offset), out DirectoryToSync? directory))
-									{
-										throw new ReplicationException($"Unable to find directory for {file}");
-									}
-									if (!directory.FileNameToSize.TryGetValue(file.Substring(offset), out fileSize))
-									{
-										throw new ReplicationException($"Unable to find file entry for {file}");
-									}
+								Utf8String file = GetClientRelativePath(io.Payload, clientInfo.Client.Root);
+								int offset = GetFileOffset(file);
 
-									handles[io.File] = new FileWriter(file, fileSize);
+								long fileSize = 0;
+								if (!pathToDirectory.TryGetValue(file.Substring(0, offset), out DirectoryToSync? directory))
+								{
+									throw new ReplicationException($"Unable to find directory for {file}");
 								}
-								else if (io.Command == PerforceIoCommand.Write)
+								if (!directory.FileNameToSize.TryGetValue(file.Substring(offset), out fileSize))
 								{
-									FileWriter file = handles[io.File];
-									await file.AppendAsync(io.Payload, _settings.Commits.Chunking, writer, cancellationToken);
+									throw new ReplicationException($"Unable to find file entry for {file}");
 								}
-								else if (io.Command == PerforceIoCommand.Close)
+
+								handles[io.File] = new FileWriter(file, fileSize);
+							}
+							else if (io.Command == PerforceIoCommand.Write)
+							{
+								FileWriter file = handles[io.File];
+								await file.AppendAsync(io.Payload, _settings.Commits.Chunking, writer, cancellationToken);
+							}
+							else if (io.Command == PerforceIoCommand.Close)
+							{
+								FileWriter? file = null;
+								if (handles.Remove(io.File, out file))
 								{
-									FileWriter? file = null;
-									try
+									FileNode node = file.Node;
+									if (node.Length != file.Size)
 									{
-										if (handles.Remove(io.File, out file))
-										{
-											FileNode node = file.Node;
-											if (node.Length != file.Size)
-											{
-												throw new ReplicationException($"Invalid size for replicated file '{file.Path}'. Expected {file.Size}, got {node.Length}.");
-											}
-
-											FileEntry entry = await root.AddFileByPathAsync(reader, file.Path, FileEntryFlags.None, node, cancellationToken);
-											await writer.WriteAsync(entry, cancellationToken);
-										}
+										throw new ReplicationException($"Invalid size for replicated file '{file.Path}'. Expected {file.Size}, got {node.Length}.");
 									}
-									finally
-									{
-										file?.Dispose();
-									}
-								}
-								else if (io.Command == PerforceIoCommand.Unlink)
-								{
-									Utf8String file = GetClientRelativePath(io.Payload, clientInfo.Client.Root);
-									await root.DeleteFileByPathAsync(reader, file, cancellationToken);
-								}
-								else
-								{
-									_logger.LogWarning("Unhandled command code {Code}", io.Command);
+
+									FileEntry entry = await root.AddFileByPathAsync(reader, file.Path, FileEntryFlags.None, node, cancellationToken);
+									await writer.WriteAsync(entry, cancellationToken);
 								}
 							}
-
-							processTimer.Stop();
-						}
-					}
-					finally
-					{
-						foreach (FileWriter handle in handles.Values)
-						{
-							handle.Dispose();
+							else if (io.Command == PerforceIoCommand.Unlink)
+							{
+								Utf8String file = GetClientRelativePath(io.Payload, clientInfo.Client.Root);
+								await root.DeleteFileByPathAsync(reader, file, cancellationToken);
+							}
+							else
+							{
+								_logger.LogWarning("Unhandled command code {Code}", io.Command);
+							}
 						}
+
+						processTimer.Stop();
 					}
 
 					TimeSpan stallTime = (perforce as NativePerforceConnection)?.StallTime ?? TimeSpan.Zero;
