Horde: Fix TypeScript warnings

#skipci

diff --git a/Engine/Source/Programs/Horde/HordeDashboard/src/components/AdminToken.tsx b/Engine/Source/Programs/Horde/HordeDashboard/src/components/AdminToken.tsx
--- a/Engine/Source/Programs/Horde/HordeDashboard/src/components/AdminToken.tsx
+++ b/Engine/Source/Programs/Horde/HordeDashboard/src/components/AdminToken.tsx
@@ -1,7 +1,7 @@
 // Copyright Epic Games, Inc. All Rights Reserved.
 
 import React, { useState } from 'react';
-import { Text, Stack, Spinner, SpinnerSize } from '@fluentui/react';
+import {Text, Stack, Spinner, SpinnerSize, PrimaryButton} from '@fluentui/react';
 import backend from '../backend';
 import {copyToClipboard} from "../base/utilities/clipboard";
 
@@ -24,7 +24,7 @@ export const AdminToken: React.FC = () => {
 		<Stack horizontalAlign="center" styles={{ root: { width: "100%", padding: 12 } }}>
 			<Stack styles={{ root: { width: 1440 } }} tokens={{ childrenGap: 18 }}>
 				<h1>Horde Access Token</h1>
-				{token && <a href="#" onClick={() => copyToClipboard(token)}>Copy to clipboard</a>}
+				{token && <PrimaryButton onClick={() => copyToClipboard(token)} text="Copy to clipboard" style={{ width: 200 }} />}
 				{!token && <div>Error getting token</div>}
 				<textarea readOnly={true} style={{ wordBreak: "break-word", height: 400 }} value={token ? token : tokenError}></textarea>
 			</Stack>
diff --git a/Engine/Source/Programs/Horde/HordeDashboard/src/components/PoolView.tsx b/Engine/Source/Programs/Horde/HordeDashboard/src/components/PoolView.tsx
--- a/Engine/Source/Programs/Horde/HordeDashboard/src/components/PoolView.tsx
+++ b/Engine/Source/Programs/Horde/HordeDashboard/src/components/PoolView.tsx
@@ -822,10 +822,10 @@ const AutoScalerPanel: React.FC = () => {
                      setState({ ...state, modified: true, strategy: option!.data })
                   }} />
                   <SpinButton styles={{ root: { width: 128 } }} disabled={disabled} label="Minimum Agents" value={state.minAgents?.toString() ?? UNSET_VALUE} labelPosition={Position.top} onChange={(ev, value) => {
-                     setState({ ...state, modified: true, minAgents: value == UNSET_VALUE ? undefined : parseInt(value ?? "0") })
+                     setState({ ...state, modified: true, minAgents: value === UNSET_VALUE ? undefined : parseInt(value ?? "0") })
                   }} />
                   <SpinButton styles={{ root: { width: 128 } }} disabled={disabled} label="Reserve Agents" value={state.reserveAgents?.toString() ?? UNSET_VALUE} labelPosition={Position.top} onChange={(ev, value) => {
-                     setState({ ...state, modified: true, reserveAgents: value == UNSET_VALUE ? undefined : parseInt(value ?? "0") })
+                     setState({ ...state, modified: true, reserveAgents: value === UNSET_VALUE ? undefined : parseInt(value ?? "0") })
                   }} />
                </Stack>
                <Stack>
