Horde: Always copy data files to output directory, to avoid bad local state persisting on dev machines.

diff --git a/Engine/Source/Programs/Horde/Horde.Agent/Horde.Agent.csproj b/Engine/Source/Programs/Horde/Horde.Agent/Horde.Agent.csproj
--- a/Engine/Source/Programs/Horde/Horde.Agent/Horde.Agent.csproj
+++ b/Engine/Source/Programs/Horde/Horde.Agent/Horde.Agent.csproj
@@ -53,31 +53,31 @@
   <ItemGroup>
     <Content Include="appsettings.json" Condition="Exists('appsettings.json')">
       <Generator></Generator>
-      <CopyToOutputDirectory>PreserveNewest</CopyToOutputDirectory>
-      <CopyToPublishDirectory>PreserveNewest</CopyToPublishDirectory>
+      <CopyToOutputDirectory>Always</CopyToOutputDirectory>
+      <CopyToPublishDirectory>Always</CopyToPublishDirectory>
     </Content>
     <Content Include="appsettings.Development.json" Condition="Exists('appsettings.Development.json')">
       <Generator></Generator>
-      <CopyToOutputDirectory>PreserveNewest</CopyToOutputDirectory>
-      <CopyToPublishDirectory>PreserveNewest</CopyToPublishDirectory>
+      <CopyToOutputDirectory>Always</CopyToOutputDirectory>
+      <CopyToPublishDirectory>Always</CopyToPublishDirectory>
       <DependentUpon>appsettings.json</DependentUpon>
     </Content>
     <Content Include="appsettings.Local.json" Condition="Exists('appsettings.Local.json')">
       <Generator></Generator>
-      <CopyToOutputDirectory>PreserveNewest</CopyToOutputDirectory>
-      <CopyToPublishDirectory>PreserveNewest</CopyToPublishDirectory>
+      <CopyToOutputDirectory>Always</CopyToOutputDirectory>
+      <CopyToPublishDirectory>Always</CopyToPublishDirectory>
       <DependentUpon>appsettings.json</DependentUpon>
     </Content>
     <Content Include="appsettings.Production.json" Condition="Exists('appsettings.Production.json')">
       <Generator></Generator>
-      <CopyToOutputDirectory>PreserveNewest</CopyToOutputDirectory>
-      <CopyToPublishDirectory>PreserveNewest</CopyToPublishDirectory>
+      <CopyToOutputDirectory>Always</CopyToOutputDirectory>
+      <CopyToPublishDirectory>Always</CopyToPublishDirectory>
       <DependentUpon>appsettings.json</DependentUpon>
     </Content>
     <Content Include="appsettings.User.json" Condition="Exists('appsettings.User.json')">
       <Generator></Generator>
-      <CopyToOutputDirectory>PreserveNewest</CopyToOutputDirectory>
-      <CopyToPublishDirectory>PreserveNewest</CopyToPublishDirectory>
+      <CopyToOutputDirectory>Always</CopyToOutputDirectory>
+      <CopyToPublishDirectory>Always</CopyToPublishDirectory>
       <DependentUpon>appsettings.json</DependentUpon>
     </Content>
   </ItemGroup>
diff --git a/Engine/Source/Programs/Horde/Horde.Build/Horde.Build.csproj b/Engine/Source/Programs/Horde/Horde.Build/Horde.Build.csproj
--- a/Engine/Source/Programs/Horde/Horde.Build/Horde.Build.csproj
+++ b/Engine/Source/Programs/Horde/Horde.Build/Horde.Build.csproj
@@ -115,14 +115,14 @@
 
   <ItemGroup>
     <Content Include="..\ThirdParty\**" Condition="$([MSBuild]::IsOSPlatform('Windows'))">
-      <CopyToOutputDirectory>PreserveNewest</CopyToOutputDirectory>
+      <CopyToOutputDirectory>Always</CopyToOutputDirectory>
       <Link>ThirdParty/%(RecursiveDir)/%(Filename)%(Extension)</Link>
     </Content>
   </ItemGroup>
 
   <ItemGroup>
     <Content Include="Defaults\ue5.project.png">
-      <CopyToOutputDirectory>PreserveNewest</CopyToOutputDirectory>
+      <CopyToOutputDirectory>Always</CopyToOutputDirectory>
     </Content>
   </ItemGroup>
 
