Horde: Fix static analysis warnings.

diff --git a/Engine/Source/Programs/Horde/Horde.Agent/Utility/TempStorage.cs b/Engine/Source/Programs/Horde/Horde.Agent/Utility/TempStorage.cs
--- a/Engine/Source/Programs/Horde/Horde.Agent/Utility/TempStorage.cs
+++ b/Engine/Source/Programs/Horde/Horde.Agent/Utility/TempStorage.cs
@@ -320,7 +320,7 @@ namespace Horde.Storage.Utility
 		/// <summary>
 		/// Construct a static Xml serializer to avoid throwing an exception searching for the reflection info at runtime
 		/// </summary>
-		static readonly XmlSerializer s_serializer = XmlSerializer.FromTypes(new Type[]{ typeof(TempStorageBlockManifest) })[0];
+		static readonly XmlSerializer s_serializer = XmlSerializer.FromTypes(new Type[]{ typeof(TempStorageBlockManifest) })[0]!;
 
 		/// <summary>
 		/// Construct an empty temp storage manifest
@@ -425,7 +425,7 @@ namespace Horde.Storage.Utility
 		/// <summary>
 		/// Construct a static Xml serializer to avoid throwing an exception searching for the reflection info at runtime
 		/// </summary>
-		static readonly XmlSerializer s_serializer = XmlSerializer.FromTypes(new Type[]{ typeof(TempStorageTagManifest) })[0];
+		static readonly XmlSerializer s_serializer = XmlSerializer.FromTypes(new Type[]{ typeof(TempStorageTagManifest) })[0]!;
 
 		/// <summary>
 		/// Construct an empty file list for deserialization
