Horde: Add an "Analyze" configuration to Horde.Agent.

diff --git a/Engine/Source/Programs/Horde/Horde.Agent/Horde.Agent.csproj b/Engine/Source/Programs/Horde/Horde.Agent/Horde.Agent.csproj
--- a/Engine/Source/Programs/Horde/Horde.Agent/Horde.Agent.csproj
+++ b/Engine/Source/Programs/Horde/Horde.Agent/Horde.Agent.csproj
@@ -10,6 +10,7 @@
     <PublishRepositoryUrl>true</PublishRepositoryUrl>
     <EmbedUntrackedSources>false</EmbedUntrackedSources>
     <DebugType>pdbonly</DebugType>
+    <Configurations>Debug;Release;Analyze</Configurations>
   </PropertyGroup>
 
   <PropertyGroup Condition="'$(Configuration)|$(Platform)'=='Debug|AnyCPU'">
@@ -17,6 +18,13 @@
     <TreatWarningsAsErrors>true</TreatWarningsAsErrors>
   </PropertyGroup>
 
+  <PropertyGroup Condition=" '$(Configuration)|$(Platform)' == 'Analyze|AnyCPU' ">
+    <RunAnalyzersDuringBuild>True</RunAnalyzersDuringBuild>
+    <TreatWarningsAsErrors>false</TreatWarningsAsErrors>
+    <GenerateDocumentationFile>True</GenerateDocumentationFile>
+    <DocumentationFile></DocumentationFile>
+  </PropertyGroup>
+
   <ItemGroup>
     <None Remove="HordeAgent.xml" />
   </ItemGroup>
