Horde: Fix compiler warning

diff --git a/Engine/Source/Programs/Horde/Horde.Build/Server/DebugController.cs b/Engine/Source/Programs/Horde/Horde.Build/Server/DebugController.cs
--- a/Engine/Source/Programs/Horde/Horde.Build/Server/DebugController.cs
+++ b/Engine/Source/Programs/Horde/Horde.Build/Server/DebugController.cs
@@ -421,7 +421,11 @@ namespace Horde.Build.Server
 			}
 
 			using IDisposable logScope = _logger.BeginScope(args);
+			
+			// Ignore warning as we explicitly want to build this message manually
+#pragma warning disable CA2254 // Template should be a static expression
 			_logger.Log(logLevelInternal, exception, message);
+#pragma warning restore CA2254
 			
 			return Ok($"Log message generated logLevel={logLevelInternal} messageLen={messageLen} exceptionMessageLen={exceptionMessageLen} argCount={argCount} argLen={argLen}");
 		}
