Horde: Fix analyzer warnings on test collection
#rnx

diff --git a/Engine/Source/Programs/Horde/Horde.Build/Jobs/TestData/TestDataCollection.cs b/Engine/Source/Programs/Horde/Horde.Build/Jobs/TestData/TestDataCollection.cs
--- a/Engine/Source/Programs/Horde/Horde.Build/Jobs/TestData/TestDataCollection.cs
+++ b/Engine/Source/Programs/Horde/Horde.Build/Jobs/TestData/TestDataCollection.cs
@@ -923,7 +923,7 @@ namespace Horde.Build.Jobs.TestData
 
 							if (testDoc == null)
 							{
-								_logger.LogWarning("Skipping adding suite test {TestName} metaId: {MetaId} testId: {TestId} jobId: {JobId} stepId: {StepId}", test.Name, metaId?.Value, testDoc?.Id, job.Id, step.Id);
+								_logger.LogWarning("Skipping adding suite test {TestName} metaId: {MetaId} jobId: {JobId} stepId: {StepId}", test.Name, metaId!.Value, job.Id, step.Id);
 								continue;
 							}
 
