Horde: Use job change when the test data does not contain a build change list
#rnx

diff --git a/Engine/Source/Programs/Horde/Horde.Build/Jobs/TestData/TestDataCollection.cs b/Engine/Source/Programs/Horde/Horde.Build/Jobs/TestData/TestDataCollection.cs
--- a/Engine/Source/Programs/Horde/Horde.Build/Jobs/TestData/TestDataCollection.cs
+++ b/Engine/Source/Programs/Horde/Horde.Build/Jobs/TestData/TestDataCollection.cs
@@ -814,7 +814,7 @@ namespace Horde.Build.Jobs.TestData
 				TestDataRefDocument testRef = new TestDataRefDocument(job.StreamId);
 				testRef.StreamId = test.StreamId;
 				testRef.Metadata = metaId.Value;
-				testRef.BuildChangeList = testData.BuildChangeList;
+				testRef.BuildChangeList = testData.BuildChangeList <= 0 ? job.Change : testData.BuildChangeList;
 				testRef.Duration = TimeSpan.FromSeconds(testData.TotalDurationSeconds);
 				testRef.TestId = test.Id;
 				testRef.TestDataIds.Add(testDataId);
