Horde: Render null tags skippig only if undefined
#rnx
#skipci

diff --git a/Engine/Source/Programs/Horde/HordeDashboard/src/components/LogRender.tsx b/Engine/Source/Programs/Horde/HordeDashboard/src/components/LogRender.tsx
--- a/Engine/Source/Programs/Horde/HordeDashboard/src/components/LogRender.tsx
+++ b/Engine/Source/Programs/Horde/HordeDashboard/src/components/LogRender.tsx
@@ -113,14 +113,18 @@ const renderTags = (line: LogLine, lineNumber: number | undefined, logStyle: any
       if (property && typeof (property) !== "string" && typeof (property) !== "number" && typeof (property) !== "boolean" && property !== null) {
          record = property as Record<string, string>;
       } else {
-         if (property === null || property === undefined) {
+         if (property === undefined) {
             return null;
          }
-         text = property.toString();
+         if (property === null) {
+            text = "null";
+         } else {
+            text = property.toString();
+         }
       }
 
       if (record) {
-         
+
          type = record.type ?? record["$type"] ?? "";
          text = record.text ?? record["$text"] ?? "";
 
