Horde: Add cancellation tokens to log writes.

diff --git a/Engine/Source/Programs/Horde/Horde.Agent/Commands/Utilities/BuildGraphCommand.cs b/Engine/Source/Programs/Horde/Horde.Agent/Commands/Utilities/BuildGraphCommand.cs
--- a/Engine/Source/Programs/Horde/Horde.Agent/Commands/Utilities/BuildGraphCommand.cs
+++ b/Engine/Source/Programs/Horde/Horde.Agent/Commands/Utilities/BuildGraphCommand.cs
@@ -48,23 +48,23 @@ namespace Horde.Agent.Commands.Utilities
 				_outputStream.Dispose();
 			}
 
-			public Task SetOutcomeAsync(JobStepOutcome outcome) => Task.CompletedTask;
+			public Task SetOutcomeAsync(JobStepOutcome outcome, CancellationToken cancellationToken) => Task.CompletedTask;
 
-			public async Task WriteEventsAsync(List<CreateEventRequest> events)
+			public async Task WriteEventsAsync(List<CreateEventRequest> events, CancellationToken cancellationToken)
 			{
 				foreach (CreateEventRequest request in events)
 				{
 					JsonSerializerOptions options = new JsonSerializerOptions();
 					options.Converters.Add(new JsonStringEnumConverter());
-					await JsonSerializer.SerializeAsync(_eventStream, request, options);
+					await JsonSerializer.SerializeAsync(_eventStream, request, options, cancellationToken);
 					_eventStream.Write(Encoding.UTF8.GetBytes(Environment.NewLine));
 				}
 			}
 
-			public async Task WriteOutputAsync(WriteOutputRequest request)
+			public async Task WriteOutputAsync(WriteOutputRequest request, CancellationToken cancellationToken)
 			{
 				PrintJson(request.Data.Span);
-				await _outputStream.WriteAsync(request.Data.Memory);
+				await _outputStream.WriteAsync(request.Data.Memory, cancellationToken);
 			}
 
 			void PrintJson(ReadOnlySpan<byte> span)
diff --git a/Engine/Source/Programs/Horde/Horde.Agent/Utility/JsonRpcLogger.cs b/Engine/Source/Programs/Horde/Horde.Agent/Utility/JsonRpcLogger.cs
--- a/Engine/Source/Programs/Horde/Horde.Agent/Utility/JsonRpcLogger.cs
+++ b/Engine/Source/Programs/Horde/Horde.Agent/Utility/JsonRpcLogger.cs
@@ -22,9 +22,9 @@ namespace Horde.Agent.Parser
 
 	interface IJsonRpcLogSink
 	{
-		Task WriteEventsAsync(List<CreateEventRequest> events);
-		Task WriteOutputAsync(WriteOutputRequest request);
-		Task SetOutcomeAsync(JobStepOutcome outcome);
+		Task WriteEventsAsync(List<CreateEventRequest> events, CancellationToken cancellationToken);
+		Task WriteOutputAsync(WriteOutputRequest request, CancellationToken cancellationToken);
+		Task SetOutcomeAsync(JobStepOutcome outcome, CancellationToken cancellationToken);
 	}
 
 	sealed class JsonRpcLogSink : IJsonRpcLogSink
@@ -45,25 +45,26 @@ namespace Horde.Agent.Parser
 		}
 
 		/// <inheritdoc/>
-		public async Task WriteEventsAsync(List<CreateEventRequest> events)
+		public async Task WriteEventsAsync(List<CreateEventRequest> events, CancellationToken cancellationToken)
 		{
-			await _rpcClient.InvokeAsync(x => x.CreateEventsAsync(new CreateEventsRequest(events)), new RpcContext(), CancellationToken.None);
+			await _rpcClient.InvokeAsync(x => x.CreateEventsAsync(new CreateEventsRequest(events)), new RpcContext(), cancellationToken);
 		}
 
 		/// <inheritdoc/>
-		public async Task WriteOutputAsync(WriteOutputRequest request)
+		public async Task WriteOutputAsync(WriteOutputRequest request, CancellationToken cancellationToken)
 		{
-			await _rpcClient.InvokeAsync(x => x.WriteOutputAsync(request), new RpcContext(), CancellationToken.None);
+			await _rpcClient.InvokeAsync(x => x.WriteOutputAsync(request), new RpcContext(), cancellationToken);
 		}
 
-		public async Task SetOutcomeAsync(JobStepOutcome outcome)
+		/// <inheritdoc/>
+		public async Task SetOutcomeAsync(JobStepOutcome outcome, CancellationToken cancellationToken)
 		{
 			// Update the outcome of this jobstep
 			if (_jobId != null && _jobBatchId != null && _jobStepId != null)
 			{
 				try
 				{
-					await _rpcClient.InvokeAsync(x => x.UpdateStepAsync(new UpdateStepRequest(_jobId, _jobBatchId, _jobStepId, JobStepState.Unspecified, outcome)), new RpcContext(), CancellationToken.None);
+					await _rpcClient.InvokeAsync(x => x.UpdateStepAsync(new UpdateStepRequest(_jobId, _jobBatchId, _jobStepId, JobStepState.Unspecified, outcome)), new RpcContext(), cancellationToken);
 				}
 				catch (Exception ex)
 				{
@@ -302,7 +303,7 @@ namespace Horde.Agent.Parser
 					byte[] data = writer.WrittenSpan.ToArray();
 					try
 					{
-						await _sink.WriteOutputAsync(new WriteOutputRequest(_logId, offset, initialLineIndex, data, false));
+						await _sink.WriteOutputAsync(new WriteOutputRequest(_logId, offset, initialLineIndex, data, false), CancellationToken.None);
 					}
 					catch (Exception ex)
 					{
@@ -316,7 +317,7 @@ namespace Horde.Agent.Parser
 				{
 					try
 					{
-						await _sink.WriteEventsAsync(events);
+						await _sink.WriteEventsAsync(events, CancellationToken.None);
 					}
 					catch (Exception ex)
 					{
@@ -329,7 +330,7 @@ namespace Horde.Agent.Parser
 				{
 					try
 					{
-						await _sink.SetOutcomeAsync(Outcome);
+						await _sink.SetOutcomeAsync(Outcome, CancellationToken.None);
 					}
 					catch (Exception ex)
 					{
@@ -343,7 +344,7 @@ namespace Horde.Agent.Parser
 				{
 					try
 					{
-						await _sink.WriteOutputAsync(new WriteOutputRequest(_logId, offset, lineIndex, Array.Empty<byte>(), true));
+						await _sink.WriteOutputAsync(new WriteOutputRequest(_logId, offset, lineIndex, Array.Empty<byte>(), true), CancellationToken.None);
 					}
 					catch (Exception ex)
 					{
