Horde: Extract test document version accounting for unversioned extra items, minor hygiene
#rnx

diff --git a/Engine/Source/Programs/Horde/Horde.Build.Tests/TestDataTests.cs b/Engine/Source/Programs/Horde/Horde.Build.Tests/TestDataTests.cs
--- a/Engine/Source/Programs/Horde/Horde.Build.Tests/TestDataTests.cs
+++ b/Engine/Source/Programs/Horde/Horde.Build.Tests/TestDataTests.cs
@@ -480,8 +480,7 @@ namespace Horde.Build.Tests
 			@"        },",
 			@"        {",
 			@"            ""Key"": ""Automated Test Session Result Details::b420bdde-c030-4add-81d2-3a8404ab3e45"",",			
-			@"            ""Data"": {",
-			@"	              ""Version"" : 1,",
+			@"            ""Data"": {",			
 			@"                ""cbdb55ea"": {",
 			@"                    ""Events"": [",
 			@"                        {",
diff --git a/Engine/Source/Programs/Horde/Horde.Build/Jobs/TestData/TestDataCollection.cs b/Engine/Source/Programs/Horde/Horde.Build/Jobs/TestData/TestDataCollection.cs
--- a/Engine/Source/Programs/Horde/Horde.Build/Jobs/TestData/TestDataCollection.cs
+++ b/Engine/Source/Programs/Horde/Horde.Build/Jobs/TestData/TestDataCollection.cs
@@ -312,7 +312,6 @@ namespace Horde.Build.Jobs.TestData
 
 		readonly ILogger _logger;
 
-		readonly IOptionsMonitor<ServerSettings> _settings;
 
 		/// <summary>
 		/// Constructor
@@ -320,13 +319,9 @@ namespace Horde.Build.Jobs.TestData
 		/// <param name="mongoService"></param>
 		/// <param name="logger"></param>
 		/// <param name="settings"></param>
-#pragma warning disable CS8618 // Disabled for feature flag: Non-nullable field must contain a non-null value when exiting constructor. Consider declaring as nullable.
 		public TestDataCollection(MongoService mongoService, ILogger<TestDataCollection> logger, IOptionsMonitor<ServerSettings> settings)
-#pragma warning restore CS8618 
 		{
-
 			_logger = logger;
-			_settings = settings;
 
 			List<MongoIndex<TestDataDocument>> indexes = new List<MongoIndex<TestDataDocument>>();
 			indexes.Add(keys => keys.Ascending(x => x.StreamId).Ascending(x => x.Change).Ascending(x => x.Key));
@@ -671,16 +666,27 @@ namespace Horde.Build.Jobs.TestData
 			List<AutomatedTestSessionData> sessions = new List<AutomatedTestSessionData>();
 			List<UnrealAutomatedTestData> tests = new List<UnrealAutomatedTestData>();
 
+			// Get test document version
+			int version = 0;
 			foreach (TestDataDocument item in documents)
 			{
 				BsonDocument testData = item.Data;
-
-				int version = testData.GetValue("Version", new BsonInt32(0)).AsInt32;
-				if (version < 1)
+				version = testData.GetValue("Version", new BsonInt32(0)).AsInt32;
+				if (version > 0)
 				{
-					_logger.LogWarning("Test data does not have version and needs to be updated in stream for job {JobId} step {StepId}", job.Id, step.Id);
-					return;
+					break;
 				}
+			}
+
+			if (version < 1)
+			{
+				_logger.LogWarning("Test data does not have version and needs to be updated in stream for job {JobId} step {StepId}", job.Id, step.Id);
+				return;
+			}
+
+			foreach (TestDataDocument item in documents)
+			{
+				BsonDocument testData = item.Data;
 
 				BsonValue? value;
 				if (!testData.TryGetPropertyValue("Type", BsonType.String, out value))
