Horde: Send metrics for instance starts in AWS fleet manager

diff --git a/Engine/Source/Programs/Horde/Horde.Build.Tests/Fleet/AwsRecyclingFleetManagerTest.cs b/Engine/Source/Programs/Horde/Horde.Build.Tests/Fleet/AwsRecyclingFleetManagerTest.cs
--- a/Engine/Source/Programs/Horde/Horde.Build.Tests/Fleet/AwsRecyclingFleetManagerTest.cs
+++ b/Engine/Source/Programs/Horde/Horde.Build.Tests/Fleet/AwsRecyclingFleetManagerTest.cs
@@ -12,6 +12,7 @@ using Horde.Build.Agents.Fleet;
 using Microsoft.VisualStudio.TestTools.UnitTesting;
 using Horde.Build.Agents;
 using Horde.Build.Agents.Fleet.Providers;
+using Horde.Build.Utilities;
 using Microsoft.Extensions.Logging;
 using JsonSerializer = System.Text.Json.JsonSerializer;
 
@@ -195,7 +196,8 @@ namespace Horde.Build.Tests.Fleet
 
 			ILogger<AwsRecyclingFleetManager> logger = loggerFactory.CreateLogger<AwsRecyclingFleetManager>();
 			IPool pool = await PoolService.CreatePoolAsync("testPool", null, true, 0, 0, sizeStrategy: PoolSizeStrategy.NoOp);
-			AwsRecyclingFleetManager manager = new (ec2, AgentCollection, settings, logger);
+			using NoOpDogStatsd dogStatsd = new ();
+			AwsRecyclingFleetManager manager = new (ec2, AgentCollection, dogStatsd, settings, logger);
 			await manager.ExpandPoolAsync(pool, new List<IAgent>(), numRequestedInstances, CancellationToken.None);
 		}
 	}
diff --git a/Engine/Source/Programs/Horde/Horde.Build/Agents/Fleet/FleetManagerFactory.cs b/Engine/Source/Programs/Horde/Horde.Build/Agents/Fleet/FleetManagerFactory.cs
--- a/Engine/Source/Programs/Horde/Horde.Build/Agents/Fleet/FleetManagerFactory.cs
+++ b/Engine/Source/Programs/Horde/Horde.Build/Agents/Fleet/FleetManagerFactory.cs
@@ -8,6 +8,7 @@ using Horde.Build.Agents.Fleet.Providers;
 using Microsoft.Extensions.DependencyInjection;
 using Microsoft.Extensions.Logging;
 using Microsoft.Extensions.Options;
+using StatsdClient;
 
 namespace Horde.Build.Agents.Fleet;
 
@@ -32,6 +33,7 @@ public interface IFleetManagerFactory
 public sealed class FleetManagerFactory : IFleetManagerFactory
 {
 	private readonly IAgentCollection _agentCollection;
+	private readonly IDogStatsd _dogStatsd;
 	private readonly IServiceProvider _provider;
 	private readonly IOptionsMonitor<ServerSettings> _settings;
 	private readonly ILoggerFactory _loggerFactory;
@@ -42,9 +44,10 @@ public sealed class FleetManagerFactory : IFleetManagerFactory
 	/// <summary>
 	/// Constructor
 	/// </summary>
-	public FleetManagerFactory(IAgentCollection agentCollection, IServiceProvider provider, IOptionsMonitor<ServerSettings> settings, ILoggerFactory loggerFactory)
+	public FleetManagerFactory(IAgentCollection agentCollection, IDogStatsd dogStatsd, IServiceProvider provider, IOptionsMonitor<ServerSettings> settings, ILoggerFactory loggerFactory)
 	{
 		_agentCollection = agentCollection;
+		_dogStatsd = dogStatsd;
 		_provider = provider;
 		_settings = settings;
 		_loggerFactory = loggerFactory;
@@ -64,7 +67,7 @@ public sealed class FleetManagerFactory : IFleetManagerFactory
 			FleetManagerType.AwsReuse =>
 				new AwsReuseFleetManager(GetAwsEc2(type), _agentCollection, DeserializeSettings<AwsReuseFleetManagerSettings>(config), _loggerFactory.CreateLogger<AwsReuseFleetManager>()),
 			FleetManagerType.AwsRecycle =>
-				new AwsRecyclingFleetManager(GetAwsEc2(type), _agentCollection, DeserializeSettings<AwsRecyclingFleetManagerSettings>(config), _loggerFactory.CreateLogger<AwsRecyclingFleetManager>()),
+				new AwsRecyclingFleetManager(GetAwsEc2(type), _agentCollection, _dogStatsd, DeserializeSettings<AwsRecyclingFleetManagerSettings>(config), _loggerFactory.CreateLogger<AwsRecyclingFleetManager>()),
 			FleetManagerType.AwsAsg =>
 				new AwsAsgFleetManager(GetAwsAutoScaling(type), DeserializeSettings<AwsAsgSettings>(config), _loggerFactory.CreateLogger<AwsAsgFleetManager>()),
 			_ => throw new ArgumentException("Unknown fleet manager type " + type)
diff --git a/Engine/Source/Programs/Horde/Horde.Build/Agents/Fleet/Providers/AwsRecyclingFleetManager.cs b/Engine/Source/Programs/Horde/Horde.Build/Agents/Fleet/Providers/AwsRecyclingFleetManager.cs
--- a/Engine/Source/Programs/Horde/Horde.Build/Agents/Fleet/Providers/AwsRecyclingFleetManager.cs
+++ b/Engine/Source/Programs/Horde/Horde.Build/Agents/Fleet/Providers/AwsRecyclingFleetManager.cs
@@ -14,6 +14,7 @@ using Horde.Build.Agents.Pools;
 using Microsoft.Extensions.Logging;
 using OpenTracing;
 using OpenTracing.Util;
+using StatsdClient;
 
 namespace Horde.Build.Agents.Fleet.Providers;
 
@@ -89,15 +90,17 @@ public sealed class AwsRecyclingFleetManager : IFleetManager
 	
 	private readonly IAmazonEC2 _ec2;
 	private readonly IAgentCollection _agentCollection;
+	private readonly IDogStatsd _dogStatsd;
 	private readonly ILogger _logger;
 
 	/// <summary>
 	/// Constructor
 	/// </summary>
-	public AwsRecyclingFleetManager(IAmazonEC2 ec2, IAgentCollection agentCollection, AwsRecyclingFleetManagerSettings settings, ILogger<AwsRecyclingFleetManager> logger)
+	public AwsRecyclingFleetManager(IAmazonEC2 ec2, IAgentCollection agentCollection, IDogStatsd dogStatsd, AwsRecyclingFleetManagerSettings settings, ILogger<AwsRecyclingFleetManager> logger)
 	{
 		_ec2 = ec2;
 		_agentCollection = agentCollection;
+		_dogStatsd = dogStatsd;
 		_logger = logger;
 		Settings = settings;
 	}
@@ -254,22 +257,30 @@ public sealed class AwsRecyclingFleetManager : IFleetManager
 		bool success = false;
 		foreach (InstanceType? instanceType in instanceTypes)
 		{
+			string[] metricTags = { "instanceType:" + instanceType, "az:" + instances[0].Placement.AvailabilityZone };
 			try
 			{
 				await StartInstancesAsync(instances, instanceType, cancellationToken);
 				success = true;
+				_dogStatsd.Increment("horde.fleet.awsRecycling.startInstance.success", tags: metricTags);
 				break;
 			}
 			catch (AwsInsufficientCapacityException)
 			{
-				_logger.LogInformation("Insufficient capacity for {InstanceType}", instanceType);
+				_dogStatsd.Increment("horde.fleet.awsRecycling.startInstance.errorCapacity", tags: metricTags);
+				_logger.LogInformation("Insufficient capacity for {InstanceType}", instanceType?.ToString());
+			}
+			catch (Exception)
+			{
+				_dogStatsd.Increment("horde.fleet.awsRecycling.startInstance.error", tags: metricTags);
+				throw;
 			}
 		}
 
 		scope.Span.SetTag("Success", success);
 		if (!success)
 		{
-			_logger.LogError("Unable to start instances. Insufficient capacity for all instance types tried.");
+			_logger.LogError("Unable to start instances. Insufficient capacity for all instance types tried");
 		}
 	}
 
