Horde: Fix event matcher for editor exit codes not matching negative values correctly.

diff --git a/Engine/Source/Programs/AutomationTool/AutomationUtils/Matchers/ExitCodeEventMatcher.cs b/Engine/Source/Programs/AutomationTool/AutomationUtils/Matchers/ExitCodeEventMatcher.cs
--- a/Engine/Source/Programs/AutomationTool/AutomationUtils/Matchers/ExitCodeEventMatcher.cs
+++ b/Engine/Source/Programs/AutomationTool/AutomationUtils/Matchers/ExitCodeEventMatcher.cs
@@ -14,7 +14,7 @@ namespace AutomationUtils.Matchers
 	class ExitCodeEventMatcher : ILogEventMatcher
 	{
 		static readonly Regex s_pattern = new Regex(
-			@"Editor terminated with exit code [1-9]|(Error executing.+)(tool returned code)(.+)");
+			@"Editor terminated with exit code -?[1-9]|(Error executing.+)(tool returned code)(.+)");
 
 		public LogEventMatch? Match(ILogCursor cursor)
 		{
diff --git a/Engine/Source/Programs/Horde/Horde.Agent.Tests/LogParserTests.cs b/Engine/Source/Programs/Horde/Horde.Agent.Tests/LogParserTests.cs
--- a/Engine/Source/Programs/Horde/Horde.Agent.Tests/LogParserTests.cs
+++ b/Engine/Source/Programs/Horde/Horde.Agent.Tests/LogParserTests.cs
@@ -114,6 +114,22 @@ namespace Horde.Agent.Tests
 			CheckEventGroup(logEvents, 2, 1, LogLevel.Error, KnownLogEvents.ExitCode);
 		}
 
+		[TestMethod]
+		public void ExitCodeEventMatcher3()
+		{
+			string[] lines =
+			{
+				@"Took 5406.5523184s to run UnrealEditor-Cmd.exe, ExitCode=-1073741819",
+				@"Editor terminated with exit code -1073741819 while running Cook for D:\build\++UE5\Sync\Samples\Games\AncientGame\AncientGame.uproject; see log d:\build\++UE5\Sync\Engine\Programs\AutomationTool\Saved\Logs\Cook-2022.11.11-08.02.19.txt",
+				@"AutomationTool executed for 1h 33m 5s",
+				@"AutomationTool exiting with ExitCode=1 (Error_Unknown)",
+				@"BUILD FAILED"
+			};
+
+			List<LogEvent> logEvents = Parse(lines);
+			CheckEventGroup(logEvents, 1, 1, LogLevel.Error, KnownLogEvents.ExitCode);
+		}
+
 		[TestMethod]
 		public void CrashEventMatcher()
 		{
