Horde: Test data and api revision, mainly removing stream id as was generating a lot of redundant metadata (for stream index) and adding finer grained queries
#rnx

diff --git a/Engine/Source/Programs/Horde/Horde.Build.Tests/TestDataTests.cs b/Engine/Source/Programs/Horde/Horde.Build.Tests/TestDataTests.cs
--- a/Engine/Source/Programs/Horde/Horde.Build.Tests/TestDataTests.cs
+++ b/Engine/Source/Programs/Horde/Horde.Build.Tests/TestDataTests.cs
@@ -17,6 +17,7 @@ using Microsoft.AspNetCore.Http;
 using Microsoft.AspNetCore.Mvc;
 using System.Security.Claims;
 using Horde.Build.Logs;
+using System.Linq;
 
 namespace Horde.Build.Tests
 {
@@ -179,7 +180,14 @@ namespace Horde.Build.Tests
 			await TestDataCollection.AddAsync(job, step, data.ToArray());
 			await TestDataCollection.AddAsync(job2, step2, data.ToArray());
 
-			ActionResult<List<GetTestStreamResponse>> streamResult = await TestDataController.GetTestStreamsAsync(streamIds);
+			ActionResult<List<GetTestMetaResponse>> metaResult = await TestDataController.GetTestMetaAsync();
+			Assert.IsNotNull(metaResult);
+			Assert.IsNotNull(metaResult.Value);
+			List<GetTestMetaResponse> meta = metaResult.Value;
+			Assert.AreEqual(1, meta.Count);
+
+			DateTimeOffset minCreateTime = new DateTimeOffset(DateTime.UtcNow.Subtract(TimeSpan.FromHours(1)));
+			ActionResult<List<GetTestStreamResponse>> streamResult = await TestDataController.GetTestStreamsAsync(streamIds, meta.Select(x => x.Id).ToArray(), minCreateTime);
 			Assert.IsNotNull(streamResult);
 			Assert.IsNotNull(streamResult.Value);
 			List<GetTestStreamResponse> streams = streamResult.Value;
@@ -194,7 +202,7 @@ namespace Horde.Build.Tests
 
 			Assert.AreEqual(streams[0].TestMetadata[0].Id, streams[1].TestMetadata[0].Id);
 
-			ActionResult<List<GetTestDataRefResponse>> refResult = await TestDataController.GetTestDataRefAsync(streamIds);
+			ActionResult<List<GetTestDataRefResponse>> refResult = await TestDataController.GetTestDataRefAsync(streamIds, meta.Select(x => x.Id).ToArray());
 			Assert.IsNotNull(refResult);
 			Assert.IsNotNull(refResult.Value);
 			List<GetTestDataRefResponse> refs = refResult.Value;
@@ -228,7 +236,22 @@ namespace Horde.Build.Tests
 			await TestDataCollection.AddAsync(job, step, data.ToArray());
 			await TestDataCollection.AddAsync(job2, step2, data.ToArray());
 
-			ActionResult<List<GetTestStreamResponse>> streamResult = await TestDataController.GetTestStreamsAsync(streamIds);
+			ActionResult<List<GetTestMetaResponse>> metaResult = await TestDataController.GetTestMetaAsync();
+			Assert.IsNotNull(metaResult);
+			Assert.IsNotNull(metaResult.Value);
+			List<GetTestMetaResponse> meta = metaResult.Value;
+			Assert.AreEqual(1, meta.Count);
+			Assert.AreEqual(1, meta[0].BuildTargets.Count);
+			Assert.AreEqual("Client", meta[0].BuildTargets[0]);
+			Assert.AreEqual(1, meta[0].Platforms.Count);
+			Assert.AreEqual("Win64", meta[0].Platforms[0]);
+			Assert.AreEqual(1, meta[0].Configurations.Count);
+			Assert.AreEqual("Development", meta[0].Configurations[0]);
+			Assert.AreEqual("EngineTest", meta[0].ProjectName);
+			Assert.AreEqual("default", meta[0].RHI);
+
+			DateTimeOffset minCreateTime = new DateTimeOffset(DateTime.UtcNow.Subtract(TimeSpan.FromHours(1)));
+			ActionResult<List<GetTestStreamResponse>> streamResult = await TestDataController.GetTestStreamsAsync(streamIds, meta.Select(x => x.Id).ToArray(), minCreateTime);
 			Assert.IsNotNull(streamResult);
 			Assert.IsNotNull(streamResult.Value);
 			List<GetTestStreamResponse> streams = streamResult.Value;
@@ -247,7 +270,7 @@ namespace Horde.Build.Tests
 
 			Assert.AreEqual(streams[0].TestMetadata[0].Id, streams[1].TestMetadata[0].Id);
 
-			ActionResult<List<GetTestDataRefResponse>> refResult = await TestDataController.GetTestDataRefAsync(streamIds);
+			ActionResult<List<GetTestDataRefResponse>> refResult = await TestDataController.GetTestDataRefAsync(streamIds, meta.Select(x => x.Id).ToArray());
 			Assert.IsNotNull(refResult);
 			Assert.IsNotNull(refResult.Value);
 			List<GetTestDataRefResponse> refs = refResult.Value;
diff --git a/Engine/Source/Programs/Horde/Horde.Build/Jobs/TestData/ITestData.cs b/Engine/Source/Programs/Horde/Horde.Build/Jobs/TestData/ITestData.cs
--- a/Engine/Source/Programs/Horde/Horde.Build/Jobs/TestData/ITestData.cs
+++ b/Engine/Source/Programs/Horde/Horde.Build/Jobs/TestData/ITestData.cs
@@ -79,11 +79,6 @@ namespace Horde.Build.Jobs.TestData
 	/// </summary>
 	public interface ITest
 	{
-		/// <summary>
-		/// The associated stream
-		/// </summary>
-		StreamId StreamId { get; }
-
 		/// <summary>
 		/// The fully qualified name of the test 
 		/// </summary>
@@ -94,6 +89,11 @@ namespace Horde.Build.Jobs.TestData
 		/// </summary>
 		string? DisplayName { get; }
 
+		/// <summary>
+		/// The name of the associated suite if any
+		/// </summary>
+		string? SuiteName { get; }
+
 		/// <summary>
 		/// The meta data for the test 
 		/// </summary>
@@ -105,11 +105,6 @@ namespace Horde.Build.Jobs.TestData
 	/// </summary>
 	public interface ITestSuite
 	{
-		/// <summary>
-		/// The associated stream
-		/// </summary>
-		StreamId StreamId { get; }
-
 		/// <summary>
 		/// The name of the test suite
 		/// </summary>
@@ -129,7 +124,7 @@ namespace Horde.Build.Jobs.TestData
 		/// <summary>
 		/// The test id
 		/// </summary>
-		TestId Id { get; }
+		TestId TestId { get; }
 
 		/// <summary>
 		/// The ourcome of the suite test
diff --git a/Engine/Source/Programs/Horde/Horde.Build/Jobs/TestData/ITestDataCollection.cs b/Engine/Source/Programs/Horde/Horde.Build/Jobs/TestData/ITestDataCollection.cs
--- a/Engine/Source/Programs/Horde/Horde.Build/Jobs/TestData/ITestDataCollection.cs
+++ b/Engine/Source/Programs/Horde/Horde.Build/Jobs/TestData/ITestDataCollection.cs
@@ -24,8 +24,11 @@ namespace Horde.Build.Jobs.TestData
 		/// Gets the tests running in provided streams
 		/// </summary>
 		/// <param name="streamIds"></param>
+		/// <param name="metaIds"></param>
+		/// <param name="minCreateTime"></param>
+		/// <param name="maxCreateTime"></param>
 		/// <returns></returns>
-		Task<List<TestStream>> FindTestStreams(StreamId[] streamIds);
+		Task<List<TestStream>> FindTestStreams(StreamId[] streamIds, TestMetaId[] metaIds, DateTime minCreateTime, DateTime maxCreateTime);
 
 		/// <summary>
 		/// 
@@ -71,6 +74,19 @@ namespace Horde.Build.Jobs.TestData
 		/// <returns>The stream document</returns>
 		Task<List<ITestData>> FindAsync(StreamId? streamId, int? minChange, int? maxChange, JobId? jobId, SubResourceId? stepId, string? key = null, int index = 0, int count = 10);
 
+
+		/// <summary>
+		/// Find test meta data
+		/// </summary>
+		/// <param name="projectNames"></param>
+		/// <param name="platforms"></param>
+		/// <param name="configurations"></param>
+		/// <param name="buildTargets"></param>
+		/// <param name="rhi"></param>
+		/// <param name="metaIds"></param>
+		/// <returns></returns>
+		Task<List<ITestMeta>> FindTestMeta(string[]? projectNames = null, string[]? platforms = null, string[]? configurations = null, string[]? buildTargets = null, string? rhi = null, TestMetaId[]? metaIds = null);
+
 		/// <summary>
 		/// Delete the test data
 		/// </summary>
diff --git a/Engine/Source/Programs/Horde/Horde.Build/Jobs/TestData/TestDataCollection.cs b/Engine/Source/Programs/Horde/Horde.Build/Jobs/TestData/TestDataCollection.cs
--- a/Engine/Source/Programs/Horde/Horde.Build/Jobs/TestData/TestDataCollection.cs
+++ b/Engine/Source/Programs/Horde/Horde.Build/Jobs/TestData/TestDataCollection.cs
@@ -123,13 +123,13 @@ namespace Horde.Build.Jobs.TestData
 		{
 			[BsonRequired, BsonId]
 			public TestId Id { get; set; }
-
-			[BsonRequired, BsonElement("sid")]
-			public StreamId StreamId { get; set; }
 			
 			[BsonRequired, BsonElement("n")]
 			public string Name { get; set; }
 
+			[BsonIgnoreIfNull, BsonElement("sn")]
+			public string? SuiteName { get; set; }
+
 			[BsonIgnoreIfNull, BsonElement("dn")]
 			public string? DisplayName { get; set; }
 
@@ -143,12 +143,12 @@ namespace Horde.Build.Jobs.TestData
 				Metadata = new List<TestMetaId>();
 			}
 
-			public TestDocument(StreamId streamId, string name, List<TestMetaId> meta, string? displayName = null)
+			public TestDocument(string name, List<TestMetaId> meta, string? suiteName = null, string? displayName = null)
 			{
-				Id = TestId.GenerateNewId();
-				StreamId = streamId;
+				Id = TestId.GenerateNewId();				
 				Name = name;
 				Metadata = meta;
+				SuiteName = suiteName;
 				DisplayName = displayName;
 			}
 		}
@@ -157,10 +157,7 @@ namespace Horde.Build.Jobs.TestData
 		{
 			[BsonRequired, BsonId]
 			public TestSuiteId Id { get; set; }
-
-			[BsonRequired, BsonElement("sid")]
-			public StreamId StreamId { get; set; }
-
+			
 			[BsonRequired, BsonElement("n")]
 			public string Name { get; set; }
 
@@ -174,10 +171,9 @@ namespace Horde.Build.Jobs.TestData
 				Tests = new List<TestId>();
 			}
 
-			public TestSuiteDocument(StreamId streamId, string name, List<TestId> tests)
+			public TestSuiteDocument(string name, List<TestId> tests)
 			{
 				Id = TestSuiteId.GenerateNewId();
-				StreamId = streamId;
 				Name = name;
 				Tests = tests;
 			}
@@ -186,7 +182,7 @@ namespace Horde.Build.Jobs.TestData
 		class SuiteTestData : ISuiteTestData
 		{
 			[BsonRequired, BsonElement("tid")]
-			public TestId Id { get; set; }
+			public TestId TestId { get; set; }
 
 			[BsonRequired, BsonElement("o")]
 			public TestOutcome Outcome { get; set; }
@@ -199,7 +195,7 @@ namespace Horde.Build.Jobs.TestData
 
 			public SuiteTestData(TestId id, TestOutcome outcome, TimeSpan duration, string uid)
 			{
-				Id = id;
+				TestId = id;
 				Outcome = outcome;
 				Duration = duration;
 				UID = uid;
@@ -330,19 +326,19 @@ namespace Horde.Build.Jobs.TestData
 
 			List<MongoIndex<TestMetaDocument>> metaIndexes = new List<MongoIndex<TestMetaDocument>>();
 			metaIndexes.Add(keys => keys.Ascending(x => x.ProjectName).Ascending(x => x.Platforms).Ascending(x => x.Configurations).Ascending(x => x.BuildTargets).Ascending(x => x.RHI), unique: true);
-			_testMeta = mongoService.GetCollection<TestMetaDocument>("TestData.Meta", metaIndexes);
+			_testMeta = mongoService.GetCollection<TestMetaDocument>("TestData.MetaV2", metaIndexes);
 
 			List<MongoIndex<TestDocument>> testIndexes = new List<MongoIndex<TestDocument>>();
-			testIndexes.Add(keys => keys.Ascending(x => x.StreamId));
-			_tests = mongoService.GetCollection<TestDocument>("TestData.Tests", testIndexes);
+			testIndexes.Add(keys => keys.Ascending(x => x.Name).Ascending(x => x.Metadata));
+			_tests = mongoService.GetCollection<TestDocument>("TestData.TestsV2", testIndexes);
 
 			List<MongoIndex<TestSuiteDocument>> testSuiteIndexes = new List<MongoIndex<TestSuiteDocument>>();
-			testSuiteIndexes.Add(keys => keys.Ascending(x => x.StreamId));
-			_testSuites = mongoService.GetCollection<TestSuiteDocument>("TestData.TestSuites", testSuiteIndexes);
+			testSuiteIndexes.Add(keys => keys.Ascending(x => x.Name).Ascending(x => x.Tests));
+			_testSuites = mongoService.GetCollection<TestSuiteDocument>("TestData.TestSuitesV2", testSuiteIndexes);
 
 			List<MongoIndex<TestDataRefDocument>> testRefIndexes = new List<MongoIndex<TestDataRefDocument>>();
 			testRefIndexes.Add(keys => keys.Ascending(x => x.StreamId).Ascending(x => x.Metadata).Descending(x => x.BuildChangeList).Ascending(x => x.TestId).Ascending(x => x.SuiteId));
-			_testRefs = mongoService.GetCollection<TestDataRefDocument>("TestData.TestRefs", testRefIndexes);
+			_testRefs = mongoService.GetCollection<TestDataRefDocument>("TestData.TestRefsV2", testRefIndexes);
 		}
 
 		/// <summary>
@@ -486,50 +482,254 @@ namespace Horde.Build.Jobs.TestData
 			await _testDataDocuments.DeleteOneAsync<TestDataDocument>(x => x.Id == id);
 		}
 
-		private async Task<List<TestMetaDocument>> FindTestMeta(TestMetaId[] metaIds)
+		/// <inheritdoc/>
+		public async Task<List<ITestMeta>> FindTestMeta(string[]? projectNames = null, string[]? platforms = null, string[]? configurations = null, string[]? buildTargets = null, string? rhi = null, TestMetaId[]? metaIds = null)
 		{
 			FilterDefinitionBuilder<TestMetaDocument> filterBuilder = Builders<TestMetaDocument>.Filter;
 			FilterDefinition<TestMetaDocument> filter = FilterDefinition<TestMetaDocument>.Empty;
-			filter &= filterBuilder.In(x => x.Id, metaIds);
 
-			return await _testMeta.Find(filter).ToListAsync();
+			if (metaIds != null && metaIds.Length > 0)
+			{
+				if (metaIds.Length == 1)
+				{
+					filter &= filterBuilder.Eq(x => x.Id, metaIds[0]);
+				}
+				else
+				{
+					filter &= filterBuilder.In(x => x.Id, metaIds);
+				}
+			}
+
+			if (projectNames != null && projectNames.Length > 0)
+			{
+				if (projectNames.Length == 1)
+				{
+					filter &= filterBuilder.Eq(x => x.ProjectName, projectNames[0]);
+				}
+				else
+				{
+					filter &= filterBuilder.In(x => x.ProjectName, projectNames);
+				}
+			}
+
+			if (platforms != null && platforms.Length > 0)
+			{
+				if (platforms.Length == 1)
+				{
+					filter &= filterBuilder.AnyEq(x => x.Platforms, platforms[0]);
+				}
+				else
+				{
+					filter &= filterBuilder.AnyIn(x => x.Platforms, platforms);
+				}
+			}
+
+			if (configurations != null && configurations.Length > 0)
+			{
+				if (configurations.Length == 1)
+				{
+					filter &= filterBuilder.AnyEq(x => x.Configurations, configurations[0]);
+				}
+				else
+				{
+					filter &= filterBuilder.AnyIn(x => x.Configurations, configurations);
+				}
+			}
+
+			if (buildTargets != null && buildTargets.Length > 0)
+			{
+				if (buildTargets.Length == 1)
+				{
+					filter &= filterBuilder.AnyEq(x => x.BuildTargets, buildTargets[0]);
+				}
+				else
+				{
+					filter &= filterBuilder.AnyIn(x => x.BuildTargets, buildTargets);
+				}
+			}
+
+			if (rhi != null)
+			{
+				filter &= filterBuilder.Eq(x => x.RHI, rhi);
+			}
+
+			List<TestMetaDocument> result = await _testMeta.Find(filter).ToListAsync();
+
+			return result.ConvertAll<ITestMeta>(x => x);			
 		}
 
-		private async Task<List<TestDocument>> FindTests(StreamId[] streamIds)
+		private async Task<List<TestDocument>> FindTests(bool? suiteTests = null, string[]? projectNames = null, string[]? testNames = null, string[]? suiteNames = null, TestMetaId[]? metaIds = null, TestId[]? testIds = null)
 		{
+			List<ITestMeta> metaData = new List<ITestMeta>();
+
+			if ((projectNames != null && projectNames.Length > 0) || (metaIds != null && metaIds.Length > 0))
+			{
+				metaData = await FindTestMeta(projectNames, metaIds: metaIds);
+			}
+			
+			TestMetaId[] queryIds = metaData.Select(x => x.Id).ToArray();
+
 			FilterDefinitionBuilder<TestDocument> filterBuilder = Builders<TestDocument>.Filter;
 			FilterDefinition<TestDocument> filter = FilterDefinition<TestDocument>.Empty;
-			filter &= filterBuilder.In(x => x.StreamId, streamIds);
+
+			if (queryIds.Length > 0)
+			{
+				if (queryIds.Length == 1)
+				{
+					filter &= filterBuilder.AnyEq(x => x.Metadata, queryIds[0]);
+				}
+				else
+				{
+					filter &= filterBuilder.AnyIn(x => x.Metadata, queryIds);
+				}
+			}
+
+			if (testNames != null)
+			{
+				if (testNames.Length == 1)
+				{
+					filter &= filterBuilder.Eq(x => x.Name, testNames[0]);
+				}
+				else
+				{
+					filter &= filterBuilder.In(x => x.Name, testNames);
+				}
+			}
+
+			if (testIds != null)
+			{
+				if (testIds.Length == 1)
+				{
+					filter &= filterBuilder.Eq(x => x.Id, testIds[0]);
+				}
+				else
+				{
+					filter &= filterBuilder.In(x => x.Id, testIds);
+				}
+			}			
+			if (suiteTests == false)
+			{
+				filter &= filterBuilder.Eq(x => x.SuiteName, null);
+			}
+			else if (suiteTests == true)
+			{
+				filter &= filterBuilder.Ne(x => x.SuiteName, null);
+
+				if (suiteNames != null)
+				{
+					if (suiteNames.Length == 1)
+					{
+						filter &= filterBuilder.Eq(x => x.SuiteName, suiteNames[0]);
+					}
+					else
+					{
+						filter &= filterBuilder.In(x => x.SuiteName, suiteNames);
+					}
+				}
+			}
 
 			return await _tests.Find(filter).ToListAsync();
 		}
 
-		private async Task<List<TestSuiteDocument>> FindTestSuites(StreamId[] streamIds)
+		private async Task<List<TestSuiteDocument>> FindTestSuites(string[]? projectNames = null, string[]? suiteNames = null, TestMetaId[]? metaIds = null, TestId[]? testIds = null)
 		{
+			// get the tests which satsify the environment
+			List<TestDocument> tests = await FindTests(true, projectNames, null, suiteNames, metaIds, testIds);
+
+			TestId[] testIdQuery = tests.Select(t => t.Id).Distinct().ToArray();
+
+			if (testIdQuery.Length == 0)
+			{
+				return new List<TestSuiteDocument>();
+			}
+
 			FilterDefinitionBuilder<TestSuiteDocument> filterBuilder = Builders<TestSuiteDocument>.Filter;
 			FilterDefinition<TestSuiteDocument> filter = FilterDefinition<TestSuiteDocument>.Empty;
-			filter &= filterBuilder.In(x => x.StreamId, streamIds);
+
+			if (testIdQuery.Length == 1)
+			{
+				filter &= filterBuilder.AnyEq(x => x.Tests, testIdQuery[0]);
+			}
+			else
+			{
+				filter &= filterBuilder.AnyIn(x => x.Tests, testIdQuery);
+			}
 
 			return await _testSuites.Find(filter).ToListAsync();
 		}
 
 		/// <inheritdoc/>
-		public async Task<List<TestStream>> FindTestStreams(StreamId[] streamIds)
-		{
-			List<TestDocument> tests = await FindTests(streamIds);
-			List<TestSuiteDocument> suites = await FindTestSuites(streamIds);
+		public async Task<List<TestStream>> FindTestStreams(StreamId[] streamIds, TestMetaId[] metaIds, DateTime minCreateTime, DateTime maxCreateTime)
+		{		
+			List<ITestDataRef> refs = await FindTestRefs(streamIds, metaIds: metaIds, minCreateTime: minCreateTime, maxCreateTime: maxCreateTime);
+			List<ITestMeta> metaData = await FindTestMeta(metaIds: metaIds);
+
+			HashSet<TestId> allTestIds = new HashSet<TestId>();
+			HashSet<TestId> testIds = new HashSet<TestId>();
+			HashSet<TestId> suiteTestIds = new HashSet<TestId>();
+
+			for (int i = 0; i < refs.Count; i++)
+			{
+				ITestDataRef testRef = refs[i];
+
+				if (testRef.TestId != null)
+				{
+					testIds.Add(testRef.TestId.Value);
+					allTestIds.Add(testRef.TestId.Value);
+				}
+				else if (testRef.SuiteTests != null)
+				{
+					foreach (ISuiteTestData t in testRef.SuiteTests)
+					{
+						suiteTestIds.Add(t.TestId);
+						allTestIds.Add(t.TestId);
+					}
+				}
+			}
+			List<TestDocument> allTests = await FindTests(testIds: allTestIds.ToArray());
+
+			List<TestDocument> tests = allTests.Where(t => testIds.Contains(t.Id)).ToList();
 
-			List<TestMetaId> metaIds = tests.SelectMany(t => t.Metadata).Distinct().ToList();
-			List<TestMetaDocument> metaData = await FindTestMeta(metaIds.ToArray());
+			List<TestSuiteDocument> suites = await FindTestSuites(testIds: suiteTestIds.ToArray());
+			List<TestDocument> suiteTests = allTests.Where(t => suiteTestIds.Contains(t.Id)).ToList();
+		
+			List<TestMetaId> allMetaIds = allTests.SelectMany(t => t.Metadata).Distinct().ToList();			
 			
 			List<TestStream> testStreams = new List<TestStream>();
 			for (int i = 0; i < streamIds.Length; i++)
 			{
 				StreamId id = streamIds[i];
-				List<ITest> streamTests = tests.Where(x => x.StreamId == id).ToList<ITest>();
-				List<ITestSuite> streamSuites = suites.Where(x => x.StreamId == id).ToList<ITestSuite>();
+
+				List<ITestDataRef> streamRefs = refs.Where(r => r.StreamId == id).ToList();
+
+				allTestIds.Clear();
+				testIds.Clear();
+				suiteTestIds.Clear();
+
+				for (int j = 0; j < streamRefs.Count; j++)
+				{
+					ITestDataRef testRef = streamRefs[j];
+
+					if (testRef.TestId != null)
+					{
+						testIds.Add(testRef.TestId.Value);
+						allTestIds.Add(testRef.TestId.Value);
+					}
+					else if (testRef.SuiteTests != null)
+					{
+						foreach (ISuiteTestData t in testRef.SuiteTests)
+						{
+							suiteTestIds.Add(t.TestId);
+							allTestIds.Add(t.TestId);
+						}
+					}
+				}
+
+				List<ITest> streamTests = allTests.Where(x => allTestIds.Contains(x.Id)).ToList<ITest>();
+				List<ITestSuite> streamSuites = suites.Where(x => allTestIds.Intersect(x.Tests).Any()).ToList<ITestSuite>();
+				
 				List<TestMetaId> streamMetaIds = streamTests.SelectMany(t => t.Metadata).Distinct().ToList();
-				List<ITestMeta> streamMeta = metaData.Where(x => streamMetaIds.Contains(x.Id)).ToList<ITestMeta>();
+				List<ITestMeta> streamMeta = metaData.Where(x => streamMetaIds.Contains(x.Id)).ToList();
 				testStreams.Add(new TestStream(id, streamTests, streamSuites, streamMeta));
 			}
 
@@ -605,17 +805,36 @@ namespace Horde.Build.Jobs.TestData
 
 		}
 
-		async Task<TestDocument?> AddOrUpdateTest(IJob job, TestMetaId metaId, string testName, string? displayName = null)
-		{			
-			// register the test
-			TestDocument? test = await _tests.Find(x => x.StreamId == job.StreamId && x.Name == testName).FirstOrDefaultAsync();
-			if (test == null)
+		async Task<TestDocument?> AddOrUpdateTest(TestMetaId metaId, string testName, string? suiteName = null,  string? displayName = null)
+		{
+
+			List<TestDocument> tests;
+
+			if (suiteName != null)
 			{
-				test = new TestDocument(job.StreamId, testName, new List<TestMetaId>() { metaId }, displayName);
+				tests = await FindTests(true, metaIds: new TestMetaId[] { metaId }, testNames: new string[] { testName }, suiteNames: new string[] {suiteName});
+			}
+			else
+			{
+				tests = await FindTests(false, metaIds: new TestMetaId[] { metaId }, testNames: new string[] { testName });
+			}
+
+			if (tests.Count > 1)
+			{
+				throw new Exception($"Duplicate tests found for MetaId: {metaId}, TestName: {testName}, SuiteName: {suiteName}, Count: {tests.Count}");
+			}
+
+			TestDocument? test;
+
+			if (tests.Count == 0)
+			{
+				test = new TestDocument( testName, new List<TestMetaId>() { metaId }, suiteName, displayName);
 				await _tests.InsertOneAsync(test);
 			}
 			else
 			{
+				test = tests[0];
+
 				if (!test.Metadata.Contains(metaId))
 				{
 					test.Metadata.Add(metaId);
@@ -632,17 +851,28 @@ namespace Horde.Build.Jobs.TestData
 			return test;
 		}
 
-		private async Task<TestSuiteDocument?> AddOrUpdateTestSuite(IJob job, string suiteName, List<TestId> testIds)
+		private async Task<TestSuiteDocument?> AddOrUpdateTestSuite(string suiteName, TestMetaId metaId, List<TestId> testIds)
 		{
-			TestSuiteDocument? suite = await _testSuites.Find(x => x.StreamId == job.StreamId && x.Name == suiteName).FirstOrDefaultAsync();
 
-			if (suite == null)
+			TestSuiteDocument? suite;
+
+			List<TestSuiteDocument> suiteTests = await FindTestSuites(suiteNames: new string[] { suiteName }, metaIds: new TestMetaId[] {metaId});
+
+			if (suiteTests.Count > 1)
 			{
-				suite = new TestSuiteDocument(job.StreamId, suiteName, testIds);
+				throw new Exception($"Duplicate tests found for MetaId: {metaId}, SuiteName: {suiteName}, Count: {suiteTests.Count}");
+			}
+
+
+			if (suiteTests.Count == 0)
+			{
+				suite = new TestSuiteDocument(suiteName, testIds);
 				await _testSuites.InsertOneAsync(suite);
 			}
 			else
 			{
+				suite = suiteTests[0];
+
 				List<TestId> newIds = suite.Tests.Union(testIds).ToList();
 
 				if (newIds.Count != suite.Tests.Count)
@@ -656,6 +886,7 @@ namespace Horde.Build.Jobs.TestData
 					UpdateDefinition<TestSuiteDocument> update = updateBuilder.Set(x => x.Tests, suite.Tests);
 					await _testSuites.UpdateOneAsync(efilter, update);
 				}
+
 			}
 
 			return suite;
@@ -810,7 +1041,7 @@ namespace Horde.Build.Jobs.TestData
 					throw new Exception($"Unable to add or update simple test report for {testData.TestName}, unable to generate meta data");
 				}
 
-				TestDocument? test = await AddOrUpdateTest(job, metaId.Value, testData.TestName);
+				TestDocument? test = await AddOrUpdateTest(metaId.Value, testData.TestName);
 
 				if (test == null)
 				{
@@ -818,7 +1049,7 @@ namespace Horde.Build.Jobs.TestData
 				}
 
 				TestDataRefDocument testRef = new TestDataRefDocument(job.StreamId);
-				testRef.StreamId = test.StreamId;
+				testRef.StreamId = job.StreamId;
 				testRef.Metadata = metaId.Value;
 				testRef.BuildChangeList = testData.BuildChangeList <= 0 ? job.Change : testData.BuildChangeList;
 				testRef.Duration = TimeSpan.FromSeconds(testData.TotalDurationSeconds);
@@ -925,7 +1156,7 @@ namespace Horde.Build.Jobs.TestData
 								displayName = utest.TestDisplayName;
 							}
 
-							TestDocument? testDoc = await AddOrUpdateTest(job, metaId!.Value, test.Name, displayName);
+							TestDocument? testDoc = await AddOrUpdateTest(metaId!.Value, test.Name, suite, displayName);
 
 							if (testDoc == null)
 							{
@@ -941,7 +1172,7 @@ namespace Horde.Build.Jobs.TestData
 					// register test suite
 					if (suiteTestIds.Count > 0)
 					{
-						TestSuiteDocument? testSuite = await AddOrUpdateTestSuite(job, suite, suiteTestIds.ToList());
+						TestSuiteDocument? testSuite = await AddOrUpdateTestSuite(suite, metaId.Value, suiteTestIds.ToList());
 
 						if (testSuite == null)
 						{
diff --git a/Engine/Source/Programs/Horde/Horde.Build/Jobs/TestData/TestDataController.cs b/Engine/Source/Programs/Horde/Horde.Build/Jobs/TestData/TestDataController.cs
--- a/Engine/Source/Programs/Horde/Horde.Build/Jobs/TestData/TestDataController.cs
+++ b/Engine/Source/Programs/Horde/Horde.Build/Jobs/TestData/TestDataController.cs
@@ -10,6 +10,7 @@ using Horde.Build.Utilities;
 using Microsoft.AspNetCore.Authorization;
 using Microsoft.AspNetCore.Mvc;
 using MongoDB.Bson;
+using EpicGames.Core;
 
 namespace Horde.Build.Jobs.TestData
 {
@@ -42,7 +43,6 @@ namespace Horde.Build.Jobs.TestData
 
 		readonly StreamService _streamService;
 
-
 		/// <summary>
 		/// Constructor
 		/// </summary>
@@ -58,15 +58,45 @@ namespace Horde.Build.Jobs.TestData
 			_streamService = streamService;
 		}
 
+		/// <summary>
+		/// Get metadata 
+		/// </summary>
+		/// <param name="projects"></param>
+		/// <param name="platforms"></param>
+		/// <param name="targets"></param>
+		/// <param name="configurations"></param>
+		/// <returns></returns>
+		[HttpGet]
+		[Route("/api/v2/testdata/metadata")]
+		[ProducesResponseType(typeof(List<GetTestMetaResponse>), 200)]
+		public async Task<ActionResult<List<GetTestMetaResponse>>> GetTestMetaAsync(
+			[FromQuery(Name = "projects")] string[]? projects = null,
+			[FromQuery(Name = "platform")] string[]? platforms = null,
+			[FromQuery(Name = "target")] string[]? targets = null,
+			[FromQuery(Name = "configuration")] string[]? configurations = null)
+		{
+			List<GetTestMetaResponse> responses = new List<GetTestMetaResponse>();
+			List<ITestMeta> metaData = await _testDataService.FindTestMeta(projects, platforms, configurations, targets);
+			metaData.ForEach(m => responses.Add(new GetTestMetaResponse(m)));
+			return responses;
+		}
+
 		/// <summary>
 		/// Get stream test data for the provided ids
 		/// </summary>
 		/// <param name="streamIds"></param>
+		/// <param name="metaIds"></param>
+		/// <param name="minCreateTime"></param>
+		/// <param name="maxCreateTime"></param>
 		/// <returns></returns>
 		[HttpGet]
 		[Route("/api/v2/testdata/streams")]
 		[ProducesResponseType(typeof(List<GetTestStreamResponse>), 200)]
-		public async Task<ActionResult<List<GetTestStreamResponse>>> GetTestStreamsAsync([FromQuery(Name = "Id")] string[] streamIds)
+		public async Task<ActionResult<List<GetTestStreamResponse>>> GetTestStreamsAsync(
+			[FromQuery(Name = "Id")] string[] streamIds,
+			[FromQuery(Name = "Mid")] string[] metaIds,
+			[FromQuery] DateTimeOffset minCreateTime,
+			[FromQuery] DateTimeOffset? maxCreateTime = null)
 		{
 			StreamId[] streamIdValues = Array.ConvertAll(streamIds, x => new StreamId(x));
 
@@ -91,7 +121,7 @@ namespace Horde.Build.Jobs.TestData
 				return responses;
 			}
 
-			List<TestStream> streams = await _testDataService.FindTestStreams(queryStreams.ToArray());
+			List<TestStream> streams = await _testDataService.FindTestStreams(queryStreams.ToArray(), metaIds.ConvertAll(x => TestMetaId.Parse(x)).ToArray(), minCreateTime.UtcDateTime, maxCreateTime?.UtcDateTime);
 			streams.ForEach(s => responses.Add(new GetTestStreamResponse(s.StreamId, s.Tests, s.TestSuites, s.TestMeta)));
 			return responses;
 		}
@@ -102,7 +132,7 @@ namespace Horde.Build.Jobs.TestData
 		/// <param name="streamIds"></param>
 		/// <param name="testIds"></param>
 		/// <param name="suiteIds"></param>
-		/// <param name="platformIds"></param>
+		/// <param name="metaIds"></param>
 		/// <param name="minCreateTime"></param>
 		/// <param name="maxCreateTime"></param>
 		/// <param name="minChange"></param>
@@ -113,9 +143,9 @@ namespace Horde.Build.Jobs.TestData
 		[ProducesResponseType(typeof(List<GetTestDataRefResponse>), 200)]
 		public async Task<ActionResult<List<GetTestDataRefResponse>>> GetTestDataRefAsync(
 			[FromQuery(Name = "Id")] string[] streamIds,
+			[FromQuery(Name = "Mid")] string[] metaIds,
 			[FromQuery(Name = "Tid")] string[]? testIds = null,
-			[FromQuery(Name = "Sid")] string[]? suiteIds = null,
-			[FromQuery(Name = "Pid")] string[]? platformIds = null,
+			[FromQuery(Name = "Sid")] string[]? suiteIds = null,			
 			[FromQuery] DateTimeOffset? minCreateTime = null,
 			[FromQuery] DateTimeOffset? maxCreateTime = null,
 			[FromQuery] int? minChange = null, 
@@ -144,14 +174,13 @@ namespace Horde.Build.Jobs.TestData
 				return responses;
 			}
 
-			List<ITestDataRef> dataRefs = await _testDataService.FindTestRefs(queryStreams.ToArray(), platformIds, testIds, suiteIds, minCreateTime?.UtcDateTime, maxCreateTime?.UtcDateTime, minChange, maxChange);
+			List<ITestDataRef> dataRefs = await _testDataService.FindTestRefs(queryStreams.ToArray(), metaIds.ConvertAll(x => TestMetaId.Parse(x)).ToArray(), testIds, suiteIds, minCreateTime?.UtcDateTime, maxCreateTime?.UtcDateTime, minChange, maxChange);
 
 			dataRefs.ForEach(d => responses.Add(new GetTestDataRefResponse(d)));
 
 			return responses;
 		}
 
-
 		/// <summary>
 		/// Creates a new TestData document
 		/// </summary>
diff --git a/Engine/Source/Programs/Horde/Horde.Build/Jobs/TestData/TestDataMessages.cs b/Engine/Source/Programs/Horde/Horde.Build/Jobs/TestData/TestDataMessages.cs
--- a/Engine/Source/Programs/Horde/Horde.Build/Jobs/TestData/TestDataMessages.cs
+++ b/Engine/Source/Programs/Horde/Horde.Build/Jobs/TestData/TestDataMessages.cs
@@ -186,9 +186,9 @@ namespace Horde.Build.Jobs.TestData
 		public string? DisplayName { get; set; }
 
 		/// <summary>
-		/// The name of the test 
+		/// The name of the test suite
 		/// </summary>
-		public string StreamId { get; set; }
+		public string? SuiteName { get; set; }
 
 		/// <summary>
 		/// The platforms the test runs on
@@ -199,7 +199,7 @@ namespace Horde.Build.Jobs.TestData
 		{
 			Name = test.Name;
 			DisplayName = test.DisplayName;
-			StreamId = test.StreamId.ToString();
+			SuiteName = test.SuiteName?.ToString();
 			Metadata = test.Metadata.Select(x => x.ToString()).ToList();
 		}
 	}
@@ -214,11 +214,6 @@ namespace Horde.Build.Jobs.TestData
 		/// </summary>
 		public string Name { get; set; }
 
-		/// <summary>
-		/// The name of the test 
-		/// </summary>
-		public string StreamId { get; set; }
-
 		/// <summary>
 		/// The tests in the suite
 		/// </summary>
@@ -227,7 +222,6 @@ namespace Horde.Build.Jobs.TestData
 		internal GetTestSuiteResponse(ITestSuite suite)
 		{
 			Name = suite.Name;
-			StreamId = suite.StreamId.ToString();
 			Tests = suite.Tests.ToList();
 		}
 	}
@@ -272,7 +266,6 @@ namespace Horde.Build.Jobs.TestData
 			TestSuites = suites.Select(t => new GetTestSuiteResponse(t)).ToList();
 			TestMetadata = metaData.Select(m => new GetTestMetaResponse(m)).ToList();
 		}
-
 	}
 
 	/// <summary>
@@ -283,7 +276,7 @@ namespace Horde.Build.Jobs.TestData
 		/// <summary>
 		/// The test id
 		/// </summary>
-		public string Id { get; set; }
+		public string TestId { get; set; }
 
 		/// <summary>
 		/// The ourcome of the suite test
@@ -306,14 +299,13 @@ namespace Horde.Build.Jobs.TestData
 		/// <param name="data"></param>
 		public GetSuiteTestDataResponse(ISuiteTestData data)
 		{
-			Id = data.Id.ToString();
+			TestId = data.TestId.ToString();
 			Outcome = data.Outcome;
 			Duration = data.Duration;
-			UID =data.UID;
+			UID = data.UID;
 		}
 	}
 
-
 	/// <summary>
 	/// Data ref 
 	/// </summary>
@@ -364,7 +356,6 @@ namespace Horde.Build.Jobs.TestData
 		/// </summary>
 		public List<GetSuiteTestDataResponse>? SuiteTests { get; set; }
 
-
 		/// <summary>
 		/// Constructor
 		/// </summary>
diff --git a/Engine/Source/Programs/Horde/Horde.Build/Jobs/TestData/TestDataService.cs b/Engine/Source/Programs/Horde/Horde.Build/Jobs/TestData/TestDataService.cs
--- a/Engine/Source/Programs/Horde/Horde.Build/Jobs/TestData/TestDataService.cs
+++ b/Engine/Source/Programs/Horde/Horde.Build/Jobs/TestData/TestDataService.cs
@@ -16,7 +16,6 @@ namespace Horde.Build.Jobs.TestData
 	using TestSuiteId = ObjectId<ITestSuite>;
 	using TestMetaId = ObjectId<ITestMeta>;
 
-
 	/// <summary>
 	/// Device management service
 	/// </summary>
@@ -25,7 +24,6 @@ namespace Horde.Build.Jobs.TestData
 		
 		readonly ITestDataCollection _testData;
 
-
 		/// <summary>
 		/// Device service constructor
 		/// </summary>
@@ -41,13 +39,36 @@ namespace Horde.Build.Jobs.TestData
 
 
 		/// <summary>
-		/// Finds test streams
+		/// Find test streams
 		/// </summary>
 		/// <param name="streamIds"></param>
+		/// <param name="metaIds"></param>
+		/// <param name="minCreateTime"></param>
+		/// <param name="maxCreateTime"></param>
+		/// <returns></returns>
+		public async Task<List<TestStream>> FindTestStreams(StreamId[] streamIds, TestMetaId[] metaIds, DateTime minCreateTime, DateTime? maxCreateTime = null)
+		{
+			if (maxCreateTime == null)
+			{
+				maxCreateTime = DateTime.UtcNow;
+			}
+
+			return await _testData.FindTestStreams(streamIds, metaIds, minCreateTime, maxCreateTime.Value);
+		}
+
+		/// <summary>
+		/// Find test meta data
+		/// </summary>
+		/// <param name="projectNames"></param>
+		/// <param name="platforms"></param>
+		/// <param name="configurations"></param>
+		/// <param name="buildTargets"></param>
+		/// <param name="rhi"></param>
+		/// <param name="metaIds"></param>
 		/// <returns></returns>
-		public async Task<List<TestStream>> FindTestStreams(StreamId[] streamIds)
+		public async Task<List<ITestMeta>> FindTestMeta(string[]? projectNames = null, string[]? platforms = null, string[]? configurations = null, string[]? buildTargets = null, string? rhi = null, TestMetaId[]? metaIds = null)
 		{
-			return await _testData.FindTestStreams(streamIds);
+			return await _testData.FindTestMeta(projectNames, platforms, configurations, buildTargets, rhi, metaIds);
 		}
 
 		/// <summary>
@@ -62,13 +83,12 @@ namespace Horde.Build.Jobs.TestData
 		/// <param name="minChange"></param>
 		/// <param name="maxChange"></param>
 		/// <returns></returns>
-		public async Task<List<ITestDataRef>> FindTestRefs(StreamId[] streamIds, string[]? metaIds = null, string[]? testIds = null, string[]? suiteIds = null, DateTime? minCreateTime = null, DateTime? maxCreateTime = null, int? minChange = null, int? maxChange = null)
+		public async Task<List<ITestDataRef>> FindTestRefs(StreamId[] streamIds, TestMetaId[] metaIds, string[]? testIds = null, string[]? suiteIds = null, DateTime? minCreateTime = null, DateTime? maxCreateTime = null, int? minChange = null, int? maxChange = null)
 		{
-			TestMetaId[]? mids = metaIds?.ConvertAll(x => new TestMetaId(x)).ToArray();
 			TestId[]? tids = testIds?.ConvertAll(x => new TestId(x)).ToArray();
 			TestSuiteId[]? sids = suiteIds?.ConvertAll(x => new TestSuiteId(x)).ToArray();
 
-			return await _testData.FindTestRefs(streamIds, mids, tids, sids, minCreateTime, maxCreateTime, minChange, maxChange);
+			return await _testData.FindTestRefs(streamIds, metaIds, tids, sids, minCreateTime, maxCreateTime, minChange, maxChange);
 		}
 	}
 }
