Horde: fix typo on test meta data endpoint
#rnx
#skipci

diff --git a/Engine/Source/Programs/Horde/Horde.Build/Jobs/TestData/TestDataController.cs b/Engine/Source/Programs/Horde/Horde.Build/Jobs/TestData/TestDataController.cs
--- a/Engine/Source/Programs/Horde/Horde.Build/Jobs/TestData/TestDataController.cs
+++ b/Engine/Source/Programs/Horde/Horde.Build/Jobs/TestData/TestDataController.cs
@@ -70,7 +70,7 @@ namespace Horde.Build.Jobs.TestData
 		[Route("/api/v2/testdata/metadata")]
 		[ProducesResponseType(typeof(List<GetTestMetaResponse>), 200)]
 		public async Task<ActionResult<List<GetTestMetaResponse>>> GetTestMetaAsync(
-			[FromQuery(Name = "projects")] string[]? projects = null,
+			[FromQuery(Name = "project")] string[]? projects = null,
 			[FromQuery(Name = "platform")] string[]? platforms = null,
 			[FromQuery(Name = "target")] string[]? targets = null,
 			[FromQuery(Name = "configuration")] string[]? configurations = null)
