Horde: Initial pass at supporting log storage through generic storage backend. Toggled via the -ManageSharedStorage option used for managed temp storage, does not yet implement tailing.

diff --git a/Engine/Source/Programs/Horde/Horde.Agent.Tests/JsonRpcLoggerTests.cs b/Engine/Source/Programs/Horde/Horde.Agent.Tests/JsonRpcLoggerTests.cs
--- a/Engine/Source/Programs/Horde/Horde.Agent.Tests/JsonRpcLoggerTests.cs
+++ b/Engine/Source/Programs/Horde/Horde.Agent.Tests/JsonRpcLoggerTests.cs
@@ -4,8 +4,18 @@ using System;
 using System.Buffers;
 using System.Collections.Generic;
 using System.Text;
+using System.Threading;
+using System.Threading.Tasks;
 using EpicGames.Core;
+using EpicGames.Horde.Storage;
+using EpicGames.Horde.Storage.Backends;
+using EpicGames.Horde.Logs;
 using Horde.Agent.Parser;
+using Horde.Agent.Utility;
+using Horde.Common.Rpc;
+using HordeCommon;
+using HordeCommon.Rpc;
+using Microsoft.Extensions.Caching.Memory;
 using Microsoft.Extensions.Logging;
 using Microsoft.Extensions.Logging.Abstractions;
 using Microsoft.VisualStudio.TestTools.UnitTesting;
@@ -91,5 +101,97 @@ namespace Horde.Agent.Tests
 			Assert.AreEqual(1, result.Length);
 			Assert.AreEqual(result[0], msg);
 		}
+
+		class FakeJsonRpcLoggerBackend : JsonRpcAndStorageLogSink
+		{
+			public NodeLocator Locator { get; private set; }
+
+			public FakeJsonRpcLoggerBackend(IRpcConnection connection, string logId, IJsonRpcLogSink inner, IStorageClient store, ILogger logger)
+				: base(connection, logId, inner, store, logger)
+			{
+			}
+
+			protected override Task UpdateLogAsync(NodeLocator locator, int lineCount, CancellationToken cancellationToken)
+			{
+				Locator = locator;
+				return Task.CompletedTask;
+			}
+
+			protected override async Task<int> UpdateLogTailAsync(int tailNext, ReadOnlyMemory<byte> tailData, CancellationToken cancellationToken)
+			{
+				await Task.Delay(TimeSpan.FromSeconds(2.0), cancellationToken);
+				return -1;
+			}
+		}
+
+		class FakeLogSink : IJsonRpcLogSink
+		{
+			public Task SetOutcomeAsync(JobStepOutcome outcome, CancellationToken cancellationToken) => Task.CompletedTask;
+
+			public Task WriteEventsAsync(List<CreateEventRequest> events, CancellationToken cancellationToken) => Task.CompletedTask;
+
+			public Task WriteOutputAsync(WriteOutputRequest request, CancellationToken cancellationToken) => Task.CompletedTask;
+		}
+
+		[TestMethod]
+		public async Task StorageLoggerTest()
+		{
+			using MemoryCache cache = new MemoryCache(new MemoryCacheOptions());
+			MemoryStorageClient store = new MemoryStorageClient();
+
+			TreeReader reader = new TreeReader(store, cache, NullLogger.Instance);
+
+			FakeLogSink innerSink = new FakeLogSink();
+
+			const int Count = 20000;
+
+			LogNode file;
+			await using (FakeJsonRpcLoggerBackend sink = new FakeJsonRpcLoggerBackend(null!, "foo", innerSink, store, NullLogger.Instance))
+			{
+				await using (JsonRpcLogger logger = new JsonRpcLogger(sink, "foo", null, NullLogger.Instance))
+				{
+					for (int idx = 0; idx < Count; idx++)
+					{
+						logger.LogInformation("Testing {Number}", idx);
+					}
+				}
+				file = await reader.ReadNodeAsync<LogNode>(sink.Locator);
+			}
+
+			// Check the index text
+			List<Utf8String> extractedIndexText = new List<Utf8String>();
+
+			LogIndexNode index = await file.IndexRef.ExpandAsync(reader);
+			foreach (LogChunkRef block in index.PlainTextChunkRefs)
+			{
+				LogChunkNode text = await block.ExpandAsync(reader);
+				extractedIndexText.AddRange(text.Lines);
+			}
+
+			Assert.AreEqual(Count, extractedIndexText.Count);
+			for (int idx = 0; idx < Count; idx++)
+			{
+				Assert.AreEqual(extractedIndexText[idx], new Utf8String($"Testing {idx}"));
+			}
+
+			// Check the body text
+			List<string> extractedBodyText = new List<string>();
+
+			foreach (LogChunkRef blockRef in file.TextChunkRefs)
+			{
+				LogChunkNode blockText = await blockRef.ExpandAsync(reader);
+				foreach (Utf8String line in blockText.Lines)
+				{
+					LogEvent logEvent = LogEvent.Read(line.Span);
+					extractedBodyText.Add(logEvent.ToString());
+				}
+			}
+
+			Assert.AreEqual(Count, extractedBodyText.Count);
+			for (int idx = 0; idx < Count; idx++)
+			{
+				Assert.AreEqual(extractedBodyText[idx], $"Testing {idx}");
+			}
+		}
 	}
 }
diff --git a/Engine/Source/Programs/Horde/Horde.Agent/Execution/JobExecutor.cs b/Engine/Source/Programs/Horde/Horde.Agent/Execution/JobExecutor.cs
--- a/Engine/Source/Programs/Horde/Horde.Agent/Execution/JobExecutor.cs
+++ b/Engine/Source/Programs/Horde/Horde.Agent/Execution/JobExecutor.cs
@@ -166,6 +166,8 @@ namespace Horde.Agent.Execution
 
 		protected Dictionary<string, string> _envVars = new Dictionary<string, string>(StringComparer.OrdinalIgnoreCase);
 
+		public bool UseNewLogger { get; private set; }
+
 		public JobExecutor(ISession session, string jobId, string batchId, string agentTypeName, IHttpClientFactory httpClientFactory)
 		{
 			_session = session;
@@ -239,6 +241,10 @@ namespace Horde.Agent.Execution
 				{
 					_preprocessScript = true;
 				}
+				else if (argument.Equals("-UseNewLogger", StringComparison.OrdinalIgnoreCase))
+				{
+					UseNewLogger = true;
+				}
 				else if (argument.StartsWith(TargetArgumentPrefix, StringComparison.OrdinalIgnoreCase))
 				{
 					_targets.Add(argument.Substring(TargetArgumentPrefix.Length));
@@ -670,6 +676,11 @@ namespace Horde.Agent.Execution
 			TreeReader reader = new TreeReader(storage, cache, logger);
 			logger.LogInformation("Using Horde-managed shared storage via {SharedStorageDir}", sharedStorageDir);
 
+			JsonRpcLogger rpcLogger = (JsonRpcLogger)logger;
+			await using JsonRpcAndStorageLogSink sink = new JsonRpcAndStorageLogSink(RpcConnection, rpcLogger._logId, rpcLogger._sink, storage, logger);
+			await using JsonRpcLogger newRpcLogger = new JsonRpcLogger(sink, rpcLogger._logId, rpcLogger._warnings, rpcLogger._inner);
+			logger = newRpcLogger;
+
 			// Create the mapping of tag names to file sets
 			Dictionary<string, HashSet<FileReference>> tagNameToFileSet = new Dictionary<string, HashSet<FileReference>>();
 
diff --git a/Engine/Source/Programs/Horde/Horde.Agent/Horde.Agent.csproj b/Engine/Source/Programs/Horde/Horde.Agent/Horde.Agent.csproj
--- a/Engine/Source/Programs/Horde/Horde.Agent/Horde.Agent.csproj
+++ b/Engine/Source/Programs/Horde/Horde.Agent/Horde.Agent.csproj
@@ -133,6 +133,8 @@
   <ItemGroup>
     <Protobuf Include="..\Horde.Build\Protos\horde\horde_rpc.proto" ProtoRoot="..\Horde.Build\Protos" GrpcServices="Client" Access="Public" />
     <Protobuf Include="..\Horde.Build\Protos\horde\horde_rpc_messages.proto" ProtoRoot="..\Horde.Build\Protos" Access="Public" />
+    <Protobuf Include="..\Horde.Build\Protos\horde\log_rpc.proto" ProtoRoot="..\Horde.Build\Protos" GrpcServices="Client" Access="Public" />
+    <Protobuf Include="..\Horde.Build\Protos\horde\log_rpc_messages.proto" ProtoRoot="..\Horde.Build\Protos" Access="Public" />
     <Protobuf Include="..\Horde.Build\Protos\horde\common\agent_status.proto" ProtoRoot="..\Horde.Build\Protos" Access="Public" />
     <Protobuf Include="..\Horde.Build\Protos\horde\common\bucket_id.proto" ProtoRoot="..\Horde.Build\Protos" Access="Public" />
     <Protobuf Include="..\Horde.Build\Protos\horde\common\event_severity.proto" ProtoRoot="..\Horde.Build\Protos" Access="Public" />
diff --git a/Engine/Source/Programs/Horde/Horde.Agent/Utility/JsonRpcLogger.cs b/Engine/Source/Programs/Horde/Horde.Agent/Utility/JsonRpcLogger.cs
--- a/Engine/Source/Programs/Horde/Horde.Agent/Utility/JsonRpcLogger.cs
+++ b/Engine/Source/Programs/Horde/Horde.Agent/Utility/JsonRpcLogger.cs
@@ -3,6 +3,7 @@
 using System;
 using System.Buffers;
 using System.Collections.Generic;
+using System.Diagnostics;
 using System.Linq;
 using System.Text;
 using System.Text.Json;
@@ -10,12 +11,20 @@ using System.Text.Json.Nodes;
 using System.Threading;
 using System.Threading.Channels;
 using System.Threading.Tasks;
+using System.Xml.Linq;
 using EpicGames.Core;
+using EpicGames.Horde.Logs;
+using EpicGames.Horde.Storage;
+using Google.Protobuf;
+using Grpc.Core;
 using Horde.Agent.Utility;
+using Horde.Common.Rpc;
 using HordeCommon;
 using HordeCommon.Rpc;
 using Microsoft.Extensions.Logging;
 
+using static Horde.Common.Rpc.LogRpc;
+
 namespace Horde.Agent.Parser
 {
 	using JsonObject = System.Text.Json.Nodes.JsonObject;
@@ -74,6 +83,172 @@ namespace Horde.Agent.Parser
 		}
 	}
 
+	class JsonRpcAndStorageLogSink : IJsonRpcLogSink, IAsyncDisposable
+	{
+		const int FlushLength = 1024 * 1024;
+
+		readonly IRpcConnection _connection;
+		readonly string _logId;
+		readonly LogBuilder _builder;
+		readonly IJsonRpcLogSink _inner;
+		readonly TreeWriter _writer;
+		readonly ILogger _logger;
+
+		int _bufferLength;
+
+		// Background task
+		readonly object _lockObject = new object();
+
+		// Tailing task
+		readonly CancellationTokenSource _tailCancellationSource;
+		readonly Task _tailTask;
+		AsyncEvent _newTailDataEvent = new AsyncEvent();
+
+		public JsonRpcAndStorageLogSink(IRpcConnection connection, string logId, IJsonRpcLogSink inner, IStorageClient store, ILogger logger)
+		{
+			_connection = connection;
+			_logId = logId;
+			_builder = new LogBuilder(LogFormat.Json, logger);
+			_inner = inner;
+			_writer = new TreeWriter(store, prefix: logId);
+			_logger = logger;
+
+			_tailCancellationSource = new CancellationTokenSource();
+			_tailTask = Task.Run(() => TickTailAsync(_tailCancellationSource.Token));
+		}
+
+		public async ValueTask DisposeAsync()
+		{
+			_tailCancellationSource.Cancel();
+			try
+			{
+				await _tailTask;
+			}
+			catch (OperationCanceledException)
+			{
+			}
+			catch (Exception ex)
+			{
+				_logger.LogError(ex, "Exception on log tailing task ({LogId}): {Message}", _logId, ex.Message);
+			}
+			_tailCancellationSource.Dispose();
+		}
+
+		async Task TickTailAsync(CancellationToken cancellationToken)
+		{
+			int tailNext = -1;
+			for (; ;)
+			{
+				Task newTailDataTask;
+
+				// Get the data to send to the server
+				ReadOnlyMemory<byte> tailData = ReadOnlyMemory<byte>.Empty;
+				lock (_lockObject)
+				{
+					if (tailNext != -1)
+					{
+						tailNext = Math.Max(tailNext, _builder.LineCount);
+						tailData = _builder.ReadTailData(tailNext, 16 * 1024);
+					}
+
+					newTailDataTask = _newTailDataEvent.Task;
+					_newTailDataEvent = new AsyncEvent();
+				}
+
+				// If we don't have any updates for the server, wait until we do.
+				if (tailNext != -1 && tailData.IsEmpty)
+				{
+					await newTailDataTask;
+					continue;
+				}
+
+				// Update the next tailing position
+				int newTailNext = await UpdateLogTailAsync(tailNext, tailData, cancellationToken);
+				if (newTailNext != tailNext)
+				{
+					tailNext = newTailNext;
+					_logger.LogInformation("Modified tail position for log {LogId} to {TailNext}", _logId, tailNext);
+				}
+			}
+		}
+
+		/// <inheritdoc/>
+		public Task SetOutcomeAsync(JobStepOutcome outcome, CancellationToken cancellationToken) => _inner.SetOutcomeAsync(outcome, cancellationToken);
+
+		/// <inheritdoc/>
+		public async Task WriteEventsAsync(List<CreateEventRequest> events, CancellationToken cancellationToken)
+		{
+			await _inner.WriteEventsAsync(events, cancellationToken);
+		}
+
+		/// <inheritdoc/>
+		public async Task WriteOutputAsync(WriteOutputRequest request, CancellationToken cancellationToken)
+		{
+			await _inner.WriteOutputAsync(request, cancellationToken);
+
+			_builder.WriteData(request.Data.Memory);
+			_bufferLength += request.Data.Length;
+
+			if (request.Flush || _bufferLength > FlushLength)
+			{
+				NodeLocator locator = await _builder.FlushAsync(_writer, cancellationToken);
+				await UpdateLogAsync(locator, _builder.LineCount, cancellationToken);
+				_bufferLength = 0;
+			}
+
+			_newTailDataEvent.Set();
+		}
+
+		#region RPC calls
+
+		protected virtual async Task UpdateLogAsync(NodeLocator locator, int lineCount, CancellationToken cancellationToken)
+		{
+			UpdateLogRequest request = new UpdateLogRequest();
+			request.LineCount = lineCount;
+			request.BlobLocator = locator.ToString();
+			await _connection.InvokeAsync((LogRpcClient client) => client.UpdateLogAsync(request, cancellationToken: cancellationToken), cancellationToken);
+		}
+
+		protected virtual async Task<int> UpdateLogTailAsync(int tailNext, ReadOnlyMemory<byte> tailData, CancellationToken cancellationToken)
+		{
+			DateTime deadline = DateTime.UtcNow.AddMinutes(2.0);
+			using (IRpcClientRef<LogRpcClient> clientRef = await _connection.GetClientRefAsync<LogRpcClient>(cancellationToken))
+			{
+				using (AsyncDuplexStreamingCall<UpdateLogTailRequest, UpdateLogTailResponse> call = clientRef.Client.UpdateLogTail(deadline: deadline, cancellationToken: cancellationToken))
+				{
+					// Write the request to the server
+					UpdateLogTailRequest request = new UpdateLogTailRequest();
+					request.TailNext = tailNext;
+					request.TailData = UnsafeByteOperations.UnsafeWrap(tailData);
+					await call.RequestStream.WriteAsync(request, cancellationToken);
+
+					// Wait until the server responds or we need to trigger a new update
+					Task<bool> moveNextAsync = call.ResponseStream.MoveNext();
+
+					Task task = await Task.WhenAny(moveNextAsync, clientRef.DisposingTask);
+					if (task == clientRef.DisposingTask)
+					{
+						_logger.LogDebug("Cancelling long poll from client side (server migration)");
+					}
+
+					// Close the request stream to indicate that we're finished
+					await call.RequestStream.CompleteAsync();
+
+					// Wait for a response or a new update to come in, then close the request stream
+					UpdateLogTailResponse? response = null;
+					while (await moveNextAsync)
+					{
+						response = call.ResponseStream.Current;
+						moveNextAsync = call.ResponseStream.MoveNext();
+					}
+					return response?.TailNext ?? -1;
+				}
+			}
+		}
+
+		#endregion
+	}
+
 	/// <summary>
 	/// Class to handle uploading log data to the server in the background
 	/// </summary>
@@ -96,10 +271,10 @@ namespace Horde.Agent.Parser
 			}
 		}
 
-		readonly IJsonRpcLogSink _sink;
-		readonly string _logId;
-		readonly bool _warnings;
-		readonly ILogger _inner;
+		internal readonly IJsonRpcLogSink _sink;
+		internal readonly string _logId;
+		internal readonly bool _warnings;
+		internal readonly ILogger _inner;
 		readonly Channel<JsonLogEvent> _dataChannel;
 		Task? _dataWriter;
 
diff --git a/Engine/Source/Programs/Horde/Horde.Build/Horde.Build.csproj b/Engine/Source/Programs/Horde/Horde.Build/Horde.Build.csproj
--- a/Engine/Source/Programs/Horde/Horde.Build/Horde.Build.csproj
+++ b/Engine/Source/Programs/Horde/Horde.Build/Horde.Build.csproj
@@ -74,6 +74,8 @@
   <ItemGroup>
     <Protobuf Include="Protos\horde\horde_rpc.proto" ProtoRoot="Protos" GrpcServices="Server" />
     <Protobuf Include="Protos\horde\horde_rpc_messages.proto" ProtoRoot="Protos" />
+    <Protobuf Include="Protos\horde\log_rpc.proto" ProtoRoot="Protos" GrpcServices="Server" />
+    <Protobuf Include="Protos\horde\log_rpc_messages.proto" ProtoRoot="Protos" />
     <Protobuf Include="Protos\horde\common\agent_status.proto" ProtoRoot="Protos" />
     <Protobuf Include="Protos\horde\common\bucket_id.proto" ProtoRoot="Protos" />
     <Protobuf Include="Protos\horde\common\event_severity.proto" ProtoRoot="Protos" />
diff --git a/Engine/Source/Programs/Horde/Horde.Build/Logs/LogRpcService.cs b/Engine/Source/Programs/Horde/Horde.Build/Logs/LogRpcService.cs
new file mode 100644
--- /dev/null
+++ b/Engine/Source/Programs/Horde/Horde.Build/Logs/LogRpcService.cs
@@ -0,0 +1,61 @@
+// Copyright Epic Games, Inc. All Rights Reserved.
+
+using Microsoft.AspNetCore.Authorization;
+using Horde.Common.Rpc;
+using System.Threading.Tasks;
+using Grpc.Core;
+using Horde.Build.Utilities;
+using EpicGames.Horde.Storage;
+using Horde.Build.Storage;
+
+namespace Horde.Build.Logs
+{
+	using LogId = ObjectId<ILogFile>;
+
+	/// <summary>
+	/// Implements the Horde gRPC service for bots updating their status and dequeing work
+	/// </summary>
+	[Authorize]
+	public class LogRpcService : LogRpc.LogRpcBase
+	{
+		readonly LogFileService _logFileService;
+		readonly StorageService _storageService;
+
+		/// <summary>
+		/// Constructor
+		/// </summary>
+		public LogRpcService(LogFileService logFileService, StorageService storageService)
+		{
+			_logFileService = logFileService;
+			_storageService = storageService;
+		}
+
+		/// <inheritdoc/>
+		public override async Task<UpdateLogResponse> UpdateLog(UpdateLogRequest request, ServerCallContext context)
+		{
+			ILogFile? logFile = await _logFileService.GetCachedLogFileAsync(new LogId(request.LogId), context.CancellationToken);
+			if (logFile == null)
+			{
+				throw new StructuredRpcException(StatusCode.NotFound, "Resource not found");
+			}
+			if (!LogFileService.AuthorizeForSession(logFile, context.GetHttpContext().User))
+			{
+				throw new StructuredRpcException(StatusCode.PermissionDenied, "Access denied");
+			}
+
+			IStorageClient store = await _storageService.GetClientAsync(new NamespaceId("default"), context.CancellationToken);
+			await store.WriteRefTargetAsync(new RefName(request.LogId), NodeLocator.Parse(request.BlobLocator));
+			return new UpdateLogResponse();
+		}
+
+		/// <inheritdoc/>
+		public override async Task UpdateLogTail(IAsyncStreamReader<UpdateLogTailRequest> requestStream, IServerStreamWriter<UpdateLogTailResponse> responseStream, ServerCallContext context)
+		{
+			while (await requestStream.MoveNext())
+			{
+				// TODO: request tail data if log has been read recently
+			}
+			await responseStream.WriteAsync(new UpdateLogTailResponse { TailNext = -1 }); 
+		}
+	}
+}
diff --git a/Engine/Source/Programs/Horde/Horde.Build/Startup.cs b/Engine/Source/Programs/Horde/Horde.Build/Startup.cs
--- a/Engine/Source/Programs/Horde/Horde.Build/Startup.cs
+++ b/Engine/Source/Programs/Horde/Horde.Build/Startup.cs
@@ -976,6 +976,7 @@ namespace Horde.Build
 			{
 				endpoints.MapGrpcService<HealthService>();
 				endpoints.MapGrpcService<RpcService>();
+				endpoints.MapGrpcService<LogRpcService>();
 
 				endpoints.MapGrpcReflectionService();
 
diff --git a/Engine/Source/Programs/Shared/EpicGames.Horde/Logs/LogChunkNode.cs b/Engine/Source/Programs/Shared/EpicGames.Horde/Logs/LogChunkNode.cs
--- a/Engine/Source/Programs/Shared/EpicGames.Horde/Logs/LogChunkNode.cs
+++ b/Engine/Source/Programs/Shared/EpicGames.Horde/Logs/LogChunkNode.cs
@@ -321,6 +321,7 @@ namespace EpicGames.Horde.Logs
 		/// <returns>New chunk data object</returns>
 		public void Append(ReadOnlySpan<byte> textData)
 		{
+			CreateOutputSpace(textData.Length);
 			textData.CopyTo(_data.AsSpan(_length, textData.Length));
 
 			int prevLength = _length;
@@ -334,14 +335,7 @@ namespace EpicGames.Horde.Logs
 		/// </summary>
 		public void AppendJsonAsPlainText(ReadOnlySpan<byte> inputLine, ILogger logger)
 		{
-			// Make sure the output buffer is large enough
-			int requiredSpace = inputLine.Length;
-			if (_length + requiredSpace > _data.Length)
-			{
-				Array.Resize(ref _data, _length + requiredSpace);
-			}
-
-			// Convert the line to plain text
+			CreateOutputSpace(inputLine.Length);
 			try
 			{
 				_length = ConvertToPlainText(inputLine, _data, _length);
@@ -369,6 +363,19 @@ namespace EpicGames.Horde.Logs
 			}
 		}
 
+		/// <summary>
+		/// Ensure there is a certain amount of space in the output buffer
+		/// </summary>
+		/// <param name="appendLength">Required space</param>
+		void CreateOutputSpace(int appendLength)
+		{
+			int requiredLength = _length + appendLength;
+			if (_data.Length < requiredLength)
+			{
+				Array.Resize(ref _data, requiredLength);
+			}
+		}
+
 		/// <summary>
 		/// Determines if the given line is empty
 		/// </summary>
@@ -511,7 +518,7 @@ namespace EpicGames.Horde.Logs
 	class LogChunkSequenceBuilder
 	{
 		readonly List<LogChunkNode> _chunks = new List<LogChunkNode>();
-		readonly LogChunkBuilder _nextChunkBuilder = new LogChunkBuilder();
+		readonly LogChunkBuilder _nextChunkBuilder;
 		int _flushedLength;
 		int _flushedLineCount;
 
@@ -542,6 +549,7 @@ namespace EpicGames.Horde.Logs
 		public LogChunkSequenceBuilder(int chunkSize)
 		{
 			ChunkSize = chunkSize;
+			_nextChunkBuilder = new LogChunkBuilder(chunkSize);
 		}
 
 		/// <summary>
diff --git a/Engine/Source/Programs/Shared/EpicGames.Horde/Storage/NamespaceId.cs b/Engine/Source/Programs/Shared/EpicGames.Horde/Storage/NamespaceId.cs
--- a/Engine/Source/Programs/Shared/EpicGames.Horde/Storage/NamespaceId.cs
+++ b/Engine/Source/Programs/Shared/EpicGames.Horde/Storage/NamespaceId.cs
@@ -5,6 +5,7 @@ using System.ComponentModel;
 using System.Globalization;
 using System.Text.Json;
 using System.Text.Json.Serialization;
+using EpicGames.Core;
 
 namespace EpicGames.Horde.Storage
 {
@@ -18,34 +19,34 @@ namespace EpicGames.Horde.Storage
 		/// <summary>
 		/// The text representing this id
 		/// </summary>
-		readonly StringId _inner;
+		public Utf8String Text { get; }
 
 		/// <summary>
 		/// Constructor
 		/// </summary>
-		/// <param name="input">Unique id for the string</param>
-		public NamespaceId(string input)
+		/// <param name="text">Unique id for the namespace</param>
+		public NamespaceId(Utf8String text)
 		{
-			_inner = new StringId(input);
+			Text = text;
 		}
 
 		/// <inheritdoc/>
-		public override bool Equals(object? obj) => obj is NamespaceId id && _inner.Equals(id._inner);
+		public override bool Equals(object? obj) => obj is NamespaceId id && Text.Equals(id.Text);
 
 		/// <inheritdoc/>
-		public override int GetHashCode() => _inner.GetHashCode();
+		public override int GetHashCode() => Text.GetHashCode();
 
 		/// <inheritdoc/>
-		public bool Equals(NamespaceId other) => _inner.Equals(other._inner);
+		public bool Equals(NamespaceId other) => Text.Equals(other.Text);
 
 		/// <inheritdoc/>
-		public override string ToString() => _inner.ToString();
+		public override string ToString() => Text.ToString();
 
 		/// <inheritdoc cref="StringId.op_Equality"/>
-		public static bool operator ==(NamespaceId left, NamespaceId right) => left._inner == right._inner;
+		public static bool operator ==(NamespaceId left, NamespaceId right) => left.Text == right.Text;
 
 		/// <inheritdoc cref="StringId.op_Inequality"/>
-		public static bool operator !=(NamespaceId left, NamespaceId right) => left._inner != right._inner;
+		public static bool operator !=(NamespaceId left, NamespaceId right) => left.Text != right.Text;
 	}
 
 	/// <summary>
@@ -54,10 +55,10 @@ namespace EpicGames.Horde.Storage
 	sealed class NamespaceIdJsonConverter : JsonConverter<NamespaceId>
 	{
 		/// <inheritdoc/>
-		public override NamespaceId Read(ref Utf8JsonReader reader, Type typeToConvert, JsonSerializerOptions options) => new NamespaceId(reader.GetString() ?? String.Empty);
+		public override NamespaceId Read(ref Utf8JsonReader reader, Type typeToConvert, JsonSerializerOptions options) => new NamespaceId(new Utf8String(reader.GetUtf8String().ToArray()));
 
 		/// <inheritdoc/>
-		public override void Write(Utf8JsonWriter writer, NamespaceId value, JsonSerializerOptions options) => writer.WriteStringValue(value.ToString());
+		public override void Write(Utf8JsonWriter writer, NamespaceId value, JsonSerializerOptions options) => writer.WriteStringValue(value.Text.Span);
 	}
 
 	/// <summary>
