Horde: Do not auto-scale pools during maintenance window

diff --git a/Engine/Source/Programs/Horde/Horde.Build/Agents/Fleet/FleetService.cs b/Engine/Source/Programs/Horde/Horde.Build/Agents/Fleet/FleetService.cs
--- a/Engine/Source/Programs/Horde/Horde.Build/Agents/Fleet/FleetService.cs
+++ b/Engine/Source/Programs/Horde/Horde.Build/Agents/Fleet/FleetService.cs
@@ -12,6 +12,7 @@ using Horde.Build.Agents.Leases;
 using Horde.Build.Agents.Pools;
 using Horde.Build.Jobs;
 using Horde.Build.Jobs.Graphs;
+using Horde.Build.Server;
 using Horde.Build.Streams;
 using Horde.Build.Utilities;
 using HordeCommon;
@@ -42,6 +43,7 @@ namespace Horde.Build.Agents.Fleet
 		private readonly IJobCollection _jobCollection;
 		private readonly ILeaseCollection _leaseCollection;
 		private readonly IPoolCollection _poolCollection;
+		private readonly IDowntimeService _downtimeService;
 		private readonly StreamService _streamService;
 		private readonly IDogStatsd _dogStatsd;
 		private readonly IFleetManagerFactory _fleetManagerFactory;
@@ -63,6 +65,7 @@ namespace Horde.Build.Agents.Fleet
 			IJobCollection jobCollection,
 			ILeaseCollection leaseCollection,
 			IPoolCollection poolCollection,
+			IDowntimeService downtimeService,
 			StreamService streamService,
 			IDogStatsd dogStatsd,
 			IFleetManagerFactory fleetManagerFactory,
@@ -76,6 +79,7 @@ namespace Horde.Build.Agents.Fleet
 			_jobCollection = jobCollection;
 			_leaseCollection = leaseCollection;
 			_poolCollection = poolCollection;
+			_downtimeService = downtimeService;
 			_streamService = streamService;
 			_dogStatsd = dogStatsd;
 			_fleetManagerFactory = fleetManagerFactory;
@@ -114,6 +118,12 @@ namespace Horde.Build.Agents.Fleet
 			ISpan span = GlobalTracer.Instance.BuildSpan("FleetService.Tick").Start();
 			try
 			{
+				if (_downtimeService.IsDowntimeActive)
+				{
+					span.SetTag("IsDowntimeActive", _downtimeService.IsDowntimeActive);
+					return;
+				}
+				
 				List<PoolSizeResult> input = await GetPoolSizeDataAsync();
 				List<PoolSizeResult> output = await CalculatePoolSizesAsync(input, stoppingToken);
 				await ResizePoolsAsync(output, stoppingToken);
@@ -129,6 +139,12 @@ namespace Horde.Build.Agents.Fleet
 			ISpan span = GlobalTracer.Instance.BuildSpan("FleetService.TickHighFrequency").Start();
 			try
 			{
+				if (_downtimeService.IsDowntimeActive)
+				{
+					span.SetTag("IsDowntimeActive", _downtimeService.IsDowntimeActive);
+					return;
+				}
+				
 				// TODO: Re-enable high frequency scaling (only used for experimental scaling of remote execution agents at the moment)
 				await Task.Delay(0, stoppingToken);
 			}
