Horde: Add missing file

diff --git a/Engine/Source/Programs/Horde/Horde.Build.Tests/Fleet/FleetServiceTest.cs b/Engine/Source/Programs/Horde/Horde.Build.Tests/Fleet/FleetServiceTest.cs
--- a/Engine/Source/Programs/Horde/Horde.Build.Tests/Fleet/FleetServiceTest.cs
+++ b/Engine/Source/Programs/Horde/Horde.Build.Tests/Fleet/FleetServiceTest.cs
@@ -7,8 +7,6 @@ using System.Linq;
 using System.Text.Json;
 using System.Threading;
 using System.Threading.Tasks;
-using Amazon.AutoScaling;
-using Amazon.EC2;
 using Horde.Build.Agents;
 using Horde.Build.Agents.Fleet;
 using Horde.Build.Agents.Fleet.Providers;
@@ -18,7 +16,6 @@ using Microsoft.Extensions.DependencyInjection;
 using Microsoft.Extensions.Logging;
 using Microsoft.Extensions.Options;
 using Microsoft.VisualStudio.TestTools.UnitTesting;
-using Moq;
 using StatsdClient;
 
 namespace Horde.Build.Tests.Fleet
@@ -173,7 +170,7 @@ namespace Horde.Build.Tests.Fleet
 			serverSettingsOpt.Value.FleetManagerV2 = FleetManagerType.AwsReuse;
 			
 			FleetService service = new(
-				AgentCollection, GraphCollection, JobCollection, LeaseCollection, PoolCollection, StreamService, _dogStatsD,
+				AgentCollection, GraphCollection, JobCollection, LeaseCollection, PoolCollection, DowntimeService, StreamService, _dogStatsD,
 				new StubFleetManagerFactory(fleetManager), Clock, Cache, serverSettingsOpt, loggerFactory.CreateLogger<FleetService>());
 				
 			return service;
