Horde: Fix invalid cast when using new logger backend.

diff --git a/Engine/Source/Programs/Horde/Horde.Agent/Execution/JobExecutor.cs b/Engine/Source/Programs/Horde/Horde.Agent/Execution/JobExecutor.cs
--- a/Engine/Source/Programs/Horde/Horde.Agent/Execution/JobExecutor.cs
+++ b/Engine/Source/Programs/Horde/Horde.Agent/Execution/JobExecutor.cs
@@ -676,7 +676,8 @@ namespace Horde.Agent.Execution
 			TreeReader reader = new TreeReader(storage, cache, logger);
 			logger.LogInformation("Using Horde-managed shared storage via {SharedStorageDir}", sharedStorageDir);
 
-			JsonRpcLogger rpcLogger = (JsonRpcLogger)logger;
+			PerforceLogger perforceLogger = (PerforceLogger)logger;
+			JsonRpcLogger rpcLogger = (JsonRpcLogger)perforceLogger.Inner;
 			await using JsonRpcAndStorageLogSink sink = new JsonRpcAndStorageLogSink(RpcConnection, rpcLogger._logId, rpcLogger._sink, storage, logger);
 			await using JsonRpcLogger newRpcLogger = new JsonRpcLogger(sink, rpcLogger._logId, rpcLogger._warnings, rpcLogger._inner);
 			logger = newRpcLogger;
diff --git a/Engine/Source/Programs/Horde/Horde.Agent/Utility/PerforceLogger.cs b/Engine/Source/Programs/Horde/Horde.Agent/Utility/PerforceLogger.cs
--- a/Engine/Source/Programs/Horde/Horde.Agent/Utility/PerforceLogger.cs
+++ b/Engine/Source/Programs/Horde/Horde.Agent/Utility/PerforceLogger.cs
@@ -33,6 +33,11 @@ namespace Horde.Agent.Utility
 		readonly ILogger _inner;
 		readonly List<ClientView> _clients = new List<ClientView>();
 
+		/// <summary>
+		/// Accessor for the inner logger
+		/// </summary>
+		public ILogger Inner => _inner;
+
 		/// <summary>
 		/// Constructor
 		/// </summary>
