Horde: Handle multiple logger types when unwrapping the JSON RPC logger.

diff --git a/Engine/Source/Programs/Horde/Horde.Agent/Execution/JobExecutor.cs b/Engine/Source/Programs/Horde/Horde.Agent/Execution/JobExecutor.cs
--- a/Engine/Source/Programs/Horde/Horde.Agent/Execution/JobExecutor.cs
+++ b/Engine/Source/Programs/Horde/Horde.Agent/Execution/JobExecutor.cs
@@ -676,8 +676,32 @@ namespace Horde.Agent.Execution
 			TreeReader reader = new TreeReader(storage, cache, logger);
 			logger.LogInformation("Using Horde-managed shared storage via {SharedStorageDir}", sharedStorageDir);
 
-			PerforceLogger perforceLogger = (PerforceLogger)logger;
-			JsonRpcLogger rpcLogger = (JsonRpcLogger)perforceLogger.Inner;
+			ILogger nextLogger = logger;
+			for(; ;)
+			{
+				if (nextLogger is JsonRpcLogger)
+				{
+					break;
+				}
+				else if (nextLogger is PerforceLogger perforceLogger)
+				{
+					nextLogger = perforceLogger.Inner;
+				}
+				else if (nextLogger is ForwardingLogger forwardingLogger)
+				{
+					nextLogger = forwardingLogger.Loggers[0];
+				}
+				else if (nextLogger is DefaultLoggerIndentHandler indentLogger)
+				{
+					nextLogger = indentLogger.Inner;
+				}
+				else
+				{
+					throw new Exception($"Unknown logger type: {nextLogger.GetType().Name}");
+				}
+			}
+
+			JsonRpcLogger rpcLogger = (JsonRpcLogger)nextLogger;
 			await using JsonRpcAndStorageLogSink sink = new JsonRpcAndStorageLogSink(RpcConnection, rpcLogger._logId, rpcLogger._sink, storage, logger);
 			await using JsonRpcLogger newRpcLogger = new JsonRpcLogger(sink, rpcLogger._logId, rpcLogger._warnings, rpcLogger._inner);
 			logger = newRpcLogger;
diff --git a/Engine/Source/Programs/Shared/EpicGames.Core/ILoggerIndent.cs b/Engine/Source/Programs/Shared/EpicGames.Core/ILoggerIndent.cs
--- a/Engine/Source/Programs/Shared/EpicGames.Core/ILoggerIndent.cs
+++ b/Engine/Source/Programs/Shared/EpicGames.Core/ILoggerIndent.cs
@@ -147,6 +147,11 @@ namespace EpicGames.Core
 		/// </summary>
 		readonly List<Scope> _scopes = new List<Scope>();
 
+		/// <summary>
+		/// The inner logger
+		/// </summary>
+		public ILogger Inner => _inner;
+
 		/// <summary>
 		/// The current indent text
 		/// </summary>
