Horde: Add missing TreeNodeAttributes to temp storage classes. Also fix reference to type not exposed through DI.

diff --git a/Engine/Source/Programs/Horde/Horde.Agent/Utility/TempStorage.cs b/Engine/Source/Programs/Horde/Horde.Agent/Utility/TempStorage.cs
--- a/Engine/Source/Programs/Horde/Horde.Agent/Utility/TempStorage.cs
+++ b/Engine/Source/Programs/Horde/Horde.Agent/Utility/TempStorage.cs
@@ -526,6 +526,7 @@ namespace Horde.Storage.Utility
 	/// <summary>
 	/// Node representing a temp storage block
 	/// </summary>
+	[TreeNode("{73AE9604-45E6-473A-95D6-2C9732A0820E}")]
 	class TempStorageBlockNode : TreeNode
 	{
 		public TempStorageBlockManifest Manifest { get; set; }
@@ -578,6 +579,7 @@ namespace Horde.Storage.Utility
 	/// <summary>
 	/// Node representing a temp storage tag
 	/// </summary>
+	[TreeNode("{CB4DEA03-3FA5-42A1-890D-1BFD7435427C}")]
 	class TempStorageTagNode : TreeNode
 	{
 		public TempStorageTagManifest FileList { get; }
@@ -621,6 +623,7 @@ namespace Horde.Storage.Utility
 	/// <summary>
 	/// Node representing all the outputs from a node
 	/// </summary>
+	[TreeNode("{A1380FBC-06DB-4CD5-9CF2-ABEB5108E406}")]
 	class TempStorageNode : TreeNode
 	{
 		public Dictionary<string, TreeNodeRef<TempStorageTagNode>> Tags { get; } = new Dictionary<string, TreeNodeRef<TempStorageTagNode>>(StringComparer.OrdinalIgnoreCase);
diff --git a/Engine/Source/Programs/Horde/Horde.Build/Logs/LogRpcService.cs b/Engine/Source/Programs/Horde/Horde.Build/Logs/LogRpcService.cs
--- a/Engine/Source/Programs/Horde/Horde.Build/Logs/LogRpcService.cs
+++ b/Engine/Source/Programs/Horde/Horde.Build/Logs/LogRpcService.cs
@@ -18,13 +18,13 @@ namespace Horde.Build.Logs
 	[Authorize]
 	public class LogRpcService : LogRpc.LogRpcBase
 	{
-		readonly LogFileService _logFileService;
+		readonly ILogFileService _logFileService;
 		readonly StorageService _storageService;
 
 		/// <summary>
 		/// Constructor
 		/// </summary>
-		public LogRpcService(LogFileService logFileService, StorageService storageService)
+		public LogRpcService(ILogFileService logFileService, StorageService storageService)
 		{
 			_logFileService = logFileService;
 			_storageService = storageService;
