Horde: Add a slash between blob prefixes and identifiers.

diff --git a/Engine/Source/Programs/Horde/Horde.Agent/Utility/JsonRpcLogger.cs b/Engine/Source/Programs/Horde/Horde.Agent/Utility/JsonRpcLogger.cs
--- a/Engine/Source/Programs/Horde/Horde.Agent/Utility/JsonRpcLogger.cs
+++ b/Engine/Source/Programs/Horde/Horde.Agent/Utility/JsonRpcLogger.cs
@@ -110,7 +110,7 @@ namespace Horde.Agent.Parser
 			_logId = logId;
 			_builder = new LogBuilder(LogFormat.Json, logger);
 			_inner = inner;
-			_writer = new TreeWriter(store, prefix: logId);
+			_writer = new TreeWriter(store, prefix: $"logs/{logId}");
 			_logger = logger;
 
 			_tailCancellationSource = new CancellationTokenSource();
diff --git a/Engine/Source/Programs/Horde/Horde.Build/Storage/StorageService.cs b/Engine/Source/Programs/Horde/Horde.Build/Storage/StorageService.cs
--- a/Engine/Source/Programs/Horde/Horde.Build/Storage/StorageService.cs
+++ b/Engine/Source/Programs/Horde/Horde.Build/Storage/StorageService.cs
@@ -680,7 +680,7 @@ namespace Horde.Build.Storage
 			HostId hostId = HostId.Empty;
 
 			ObjectId id = ObjectId.GenerateNewId(_clock.UtcNow);
-			BlobId blobId = (prefix.Length > 0) ? new BlobId($"{prefix}{id}", BlobId.Validate.None) : new BlobId(id.ToString());
+			BlobId blobId = (prefix.Length > 0) ? new BlobId($"{prefix}/{id}") : new BlobId(id.ToString());
 
 			BlobInfo blobInfo = new BlobInfo(id, namespaceId, hostId, blobId);
 			await _blobCollection.InsertOneAsync(blobInfo, new InsertOneOptions { }, cancellationToken);
