Horde: Disable tailing calls to log RPC server, and output diagnostic messages when updating the corresponding ref.

diff --git a/Engine/Source/Programs/Horde/Horde.Agent/Utility/JsonRpcLogger.cs b/Engine/Source/Programs/Horde/Horde.Agent/Utility/JsonRpcLogger.cs
--- a/Engine/Source/Programs/Horde/Horde.Agent/Utility/JsonRpcLogger.cs
+++ b/Engine/Source/Programs/Horde/Horde.Agent/Utility/JsonRpcLogger.cs
@@ -203,6 +203,8 @@ namespace Horde.Agent.Parser
 
 		protected virtual async Task UpdateLogAsync(NodeLocator locator, int lineCount, CancellationToken cancellationToken)
 		{
+			_logger.LogDebug("Updating log {LogId} to line {LineCount}, blob {Locator}", _logId, lineCount, locator);
+
 			UpdateLogRequest request = new UpdateLogRequest();
 			request.LineCount = lineCount;
 			request.BlobLocator = locator.ToString();
@@ -225,7 +227,7 @@ namespace Horde.Agent.Parser
 					// Wait until the server responds or we need to trigger a new update
 					Task<bool> moveNextAsync = call.ResponseStream.MoveNext();
 
-					Task task = await Task.WhenAny(moveNextAsync, clientRef.DisposingTask);
+					Task task = await Task.WhenAny(moveNextAsync, clientRef.DisposingTask, Task.Delay(TimeSpan.FromMinutes(1.0), CancellationToken.None));
 					if (task == clientRef.DisposingTask)
 					{
 						_logger.LogDebug("Cancelling long poll from client side (server migration)");
diff --git a/Engine/Source/Programs/Horde/Horde.Build/Logs/LogRpcService.cs b/Engine/Source/Programs/Horde/Horde.Build/Logs/LogRpcService.cs
--- a/Engine/Source/Programs/Horde/Horde.Build/Logs/LogRpcService.cs
+++ b/Engine/Source/Programs/Horde/Horde.Build/Logs/LogRpcService.cs
@@ -7,6 +7,7 @@ using Grpc.Core;
 using Horde.Build.Utilities;
 using EpicGames.Horde.Storage;
 using Horde.Build.Storage;
+using Microsoft.Extensions.Logging;
 
 namespace Horde.Build.Logs
 {
@@ -20,14 +21,16 @@ namespace Horde.Build.Logs
 	{
 		readonly ILogFileService _logFileService;
 		readonly StorageService _storageService;
+		readonly ILogger<LogRpcService> _logger;
 
 		/// <summary>
 		/// Constructor
 		/// </summary>
-		public LogRpcService(ILogFileService logFileService, StorageService storageService)
+		public LogRpcService(ILogFileService logFileService, StorageService storageService, ILogger<LogRpcService> logger)
 		{
 			_logFileService = logFileService;
 			_storageService = storageService;
+			_logger = logger;
 		}
 
 		/// <inheritdoc/>
@@ -44,6 +47,7 @@ namespace Horde.Build.Logs
 			}
 
 			IStorageClient store = await _storageService.GetClientAsync(new NamespaceId("default"), context.CancellationToken);
+			_logger.LogInformation("Updating {LogId} to node {Locator}", request.LogId, request.BlobLocator);
 			await store.WriteRefTargetAsync(new RefName(request.LogId), NodeLocator.Parse(request.BlobLocator));
 			return new UpdateLogResponse();
 		}
