Horde: Add a test stream collection as discovering these from refs is too slow and doesn't cache readily
#rnx

diff --git a/Engine/Source/Programs/Horde/Horde.Build.Tests/TestDataTests.cs b/Engine/Source/Programs/Horde/Horde.Build.Tests/TestDataTests.cs
--- a/Engine/Source/Programs/Horde/Horde.Build.Tests/TestDataTests.cs
+++ b/Engine/Source/Programs/Horde/Horde.Build.Tests/TestDataTests.cs
@@ -186,8 +186,7 @@ namespace Horde.Build.Tests
 			List<GetTestMetaResponse> meta = metaResult.Value;
 			Assert.AreEqual(1, meta.Count);
 
-			DateTimeOffset minCreateTime = new DateTimeOffset(DateTime.UtcNow.Subtract(TimeSpan.FromHours(1)));
-			ActionResult<List<GetTestStreamResponse>> streamResult = await TestDataController.GetTestStreamsAsync(streamIds, meta.Select(x => x.Id).ToArray(), minCreateTime);
+			ActionResult<List<GetTestStreamResponse>> streamResult = await TestDataController.GetTestStreamsAsync(streamIds);
 			Assert.IsNotNull(streamResult);
 			Assert.IsNotNull(streamResult.Value);
 			List<GetTestStreamResponse> streams = streamResult.Value;
@@ -250,8 +249,7 @@ namespace Horde.Build.Tests
 			Assert.AreEqual("EngineTest", meta[0].ProjectName);
 			Assert.AreEqual("default", meta[0].RHI);
 
-			DateTimeOffset minCreateTime = new DateTimeOffset(DateTime.UtcNow.Subtract(TimeSpan.FromHours(1)));
-			ActionResult<List<GetTestStreamResponse>> streamResult = await TestDataController.GetTestStreamsAsync(streamIds, meta.Select(x => x.Id).ToArray(), minCreateTime);
+			ActionResult<List<GetTestStreamResponse>> streamResult = await TestDataController.GetTestStreamsAsync(streamIds);
 			Assert.IsNotNull(streamResult);
 			Assert.IsNotNull(streamResult.Value);
 			List<GetTestStreamResponse> streams = streamResult.Value;
diff --git a/Engine/Source/Programs/Horde/Horde.Build/Jobs/TestData/ITestData.cs b/Engine/Source/Programs/Horde/Horde.Build/Jobs/TestData/ITestData.cs
--- a/Engine/Source/Programs/Horde/Horde.Build/Jobs/TestData/ITestData.cs
+++ b/Engine/Source/Programs/Horde/Horde.Build/Jobs/TestData/ITestData.cs
@@ -79,6 +79,11 @@ namespace Horde.Build.Jobs.TestData
 	/// </summary>
 	public interface ITest
 	{
+		/// <summary>
+		/// The test id
+		/// </summary>
+		TestId Id { get; }
+
 		/// <summary>
 		/// The fully qualified name of the test 
 		/// </summary>
@@ -105,6 +110,11 @@ namespace Horde.Build.Jobs.TestData
 	/// </summary>
 	public interface ITestSuite
 	{
+		/// <summary>
+		/// The test suite id
+		/// </summary>
+		TestSuiteId Id { get; }
+
 		/// <summary>
 		/// The name of the test suite
 		/// </summary>
@@ -192,7 +202,29 @@ namespace Horde.Build.Jobs.TestData
 		/// </summary>
 		IReadOnlyList<ISuiteTestData>? SuiteTests { get; }
 	}
-	
+
+	/// <summary>
+	/// The tests and suites running in a given stream
+	/// </summary>
+	public interface ITestStream
+	{
+		/// <summary>
+		/// The stream id
+		/// </summary>
+		StreamId StreamId { get; }
+
+		/// <summary>
+		/// Test ids of tests running in the stream
+		/// </summary>
+		IReadOnlyList<TestId> Tests { get; }
+
+		/// <summary>
+		/// Test suite ids
+		/// </summary>
+		IReadOnlyList<TestSuiteId> TestSuites { get; } 
+	}
+
+
 	/// <summary>
 	/// Stores information about the results of a test
 	/// </summary>
diff --git a/Engine/Source/Programs/Horde/Horde.Build/Jobs/TestData/ITestDataCollection.cs b/Engine/Source/Programs/Horde/Horde.Build/Jobs/TestData/ITestDataCollection.cs
--- a/Engine/Source/Programs/Horde/Horde.Build/Jobs/TestData/ITestDataCollection.cs
+++ b/Engine/Source/Programs/Horde/Horde.Build/Jobs/TestData/ITestDataCollection.cs
@@ -24,11 +24,8 @@ namespace Horde.Build.Jobs.TestData
 		/// Gets the tests running in provided streams
 		/// </summary>
 		/// <param name="streamIds"></param>
-		/// <param name="metaIds"></param>
-		/// <param name="minCreateTime"></param>
-		/// <param name="maxCreateTime"></param>
 		/// <returns></returns>
-		Task<List<TestStream>> FindTestStreams(StreamId[] streamIds, TestMetaId[] metaIds, DateTime minCreateTime, DateTime maxCreateTime);
+		Task<List<ITestStream>> FindTestStreams(StreamId[] streamIds);
 
 		/// <summary>
 		/// 
@@ -44,6 +41,20 @@ namespace Horde.Build.Jobs.TestData
 		/// <returns></returns>
 		Task<List<ITestDataRef>> FindTestRefs(StreamId[] streamIds, TestMetaId[]? metaIds = null, TestId[]? testIds = null, TestSuiteId[]? suiteIds = null, DateTime? minCreateTime = null, DateTime? maxCreateTime = null, int? minChange = null, int? maxChange = null);
 
+		/// <summary>
+		/// Find tests
+		/// </summary>
+		/// <param name="testIds"></param>
+		/// <returns></returns>
+		Task<List<ITest>> FindTests(TestId[] testIds);
+
+		/// <summary>
+		/// Find test suites
+		/// </summary>
+		/// <param name="suiteIds"></param>
+		/// <returns></returns>
+		Task<List<ITestSuite>> FindTestSuites(TestSuiteId[] suiteIds);
+
 		/// <summary>
 		/// Creates a new test data document
 		/// </summary>
diff --git a/Engine/Source/Programs/Horde/Horde.Build/Jobs/TestData/TestDataCollection.cs b/Engine/Source/Programs/Horde/Horde.Build/Jobs/TestData/TestDataCollection.cs
--- a/Engine/Source/Programs/Horde/Horde.Build/Jobs/TestData/TestDataCollection.cs
+++ b/Engine/Source/Programs/Horde/Horde.Build/Jobs/TestData/TestDataCollection.cs
@@ -26,48 +26,7 @@ namespace Horde.Build.Jobs.TestData
 
 	using TestId = ObjectId<ITest>;
 	using TestSuiteId = ObjectId<ITestSuite>;
-	using TestMetaId = ObjectId<ITestMeta>;	
-
-	/// <summary>
-	/// Tests which are run in a  stream
-	/// </summary>
-	public class TestStream
-	{
-		/// <summary>
-		/// The associated stream tests are running in
-		/// </summary>
-		public StreamId StreamId { get; set; }
-
-		/// <summary>
-		/// The individual tests running in the stream
-		/// </summary>
-		public List<ITest> Tests { get; set; }
-
-		/// <summary>
-		/// The tests suites running in the stream
-		/// </summary>
-		public List<ITestSuite> TestSuites { get; set; }
-
-		/// <summary>
-		/// The test meta data (environments) running in the stream
-		/// </summary>
-		public List<ITestMeta> TestMeta { get; set; }
-
-		/// <summary>
-		/// Constructor
-		/// </summary>
-		/// <param name="streamId"></param>
-		/// <param name="tests"></param>
-		/// <param name="testSuites"></param>
-		/// <param name="testMeta"></param>
-		public TestStream(StreamId streamId, List<ITest> tests, List<ITestSuite> testSuites, List<ITestMeta> testMeta)
-		{
-			StreamId = streamId;
-			Tests = tests;
-			TestSuites = testSuites;
-			TestMeta = testMeta;
-		}
-	}
+	using TestMetaId = ObjectId<ITestMeta>;
 
 	/// <summary>
 	/// Collection of test data documents
@@ -248,6 +207,38 @@ namespace Horde.Build.Jobs.TestData
 			}
 		}
 
+		/// <summary>
+		/// Stores the tests and suites running in a stream
+		/// </summary>
+		class TestStreamDocument : ITestStream
+		{
+			[BsonRequired, BsonId]
+			public ObjectId Id { get; set; }
+
+			[BsonRequired]
+			public StreamId StreamId { get; set; }
+
+			[BsonRequired]
+			public List<TestId> Tests { get; set; } = new List<TestId>();
+			IReadOnlyList<TestId> ITestStream.Tests => Tests;
+
+			[BsonRequired]
+			public List<TestSuiteId> TestSuites { get; set; } = new List<TestSuiteId>();
+			IReadOnlyList<TestSuiteId> ITestStream.TestSuites => TestSuites;
+
+			private TestStreamDocument()
+			{
+
+			}
+
+			public TestStreamDocument(StreamId streamId)
+			{
+				Id = ObjectId.GenerateNewId();
+				StreamId = streamId;
+			}
+
+		}
+
 		/// <summary>
 		/// Information about a test data document
 		/// </summary>
@@ -306,6 +297,11 @@ namespace Horde.Build.Jobs.TestData
 		/// </summary>
 		readonly IMongoCollection<TestDataRefDocument> _testRefs;
 
+		/// <summary>
+		/// Test suite collection
+		/// </summary>
+		readonly IMongoCollection<TestStreamDocument> _testStreams;
+
 		readonly ILogger _logger;
 
 
@@ -339,6 +335,11 @@ namespace Horde.Build.Jobs.TestData
 			List<MongoIndex<TestDataRefDocument>> testRefIndexes = new List<MongoIndex<TestDataRefDocument>>();
 			testRefIndexes.Add(keys => keys.Ascending(x => x.StreamId).Ascending(x => x.Metadata).Descending(x => x.BuildChangeList).Ascending(x => x.TestId).Ascending(x => x.SuiteId));
 			_testRefs = mongoService.GetCollection<TestDataRefDocument>("TestData.TestRefsV2", testRefIndexes);
+
+			List<MongoIndex<TestStreamDocument>> streamIndexes = new List<MongoIndex<TestStreamDocument>>();
+			streamIndexes.Add(keys => keys.Ascending(x => x.StreamId), unique: true);
+			_testStreams = mongoService.GetCollection<TestStreamDocument>("TestData.TestStreamsV2", streamIndexes);
+
 		}
 
 		/// <summary>
@@ -558,6 +559,22 @@ namespace Horde.Build.Jobs.TestData
 			return result.ConvertAll<ITestMeta>(x => x);			
 		}
 
+		/// <inheritdoc/>
+		public async Task<List<ITest>> FindTests(TestId[] testIds)
+		{
+			List<TestDocument> results = await _tests.Find(Builders<TestDocument>.Filter.In(x => x.Id, testIds)).ToListAsync();
+			return results.ConvertAll<ITest>(x => x);
+
+		}
+
+		/// <inheritdoc/>
+		public async Task<List<ITestSuite>> FindTestSuites(TestSuiteId[] suiteIds)
+		{
+			List<TestSuiteDocument> results = await _testSuites.Find(Builders<TestSuiteDocument>.Filter.In(x => x.Id, suiteIds)).ToListAsync();
+			return results.ConvertAll<ITestSuite>(x => x);
+
+		}
+
 		private async Task<List<TestDocument>> FindTests(bool? suiteTests = null, string[]? projectNames = null, string[]? testNames = null, string[]? suiteNames = null, TestMetaId[]? metaIds = null, TestId[]? testIds = null)
 		{
 			List<ITestMeta> metaData = new List<ITestMeta>();
@@ -659,81 +676,66 @@ namespace Horde.Build.Jobs.TestData
 		}
 
 		/// <inheritdoc/>
-		public async Task<List<TestStream>> FindTestStreams(StreamId[] streamIds, TestMetaId[] metaIds, DateTime minCreateTime, DateTime maxCreateTime)
-		{		
-			List<ITestDataRef> refs = await FindTestRefs(streamIds, metaIds: metaIds, minCreateTime: minCreateTime, maxCreateTime: maxCreateTime);
-			List<ITestMeta> metaData = await FindTestMeta(metaIds: metaIds);
+		public async Task<List<ITestStream>> FindTestStreams(StreamId[] streamIds)
+		{
+			List<TestStreamDocument> results  = await _testStreams.Find(Builders<TestStreamDocument>.Filter.In(x => x.StreamId, streamIds)).ToListAsync();
+			return results.ConvertAll<ITestStream>(x => x);
+		}
 
-			HashSet<TestId> allTestIds = new HashSet<TestId>();
-			HashSet<TestId> testIds = new HashSet<TestId>();
-			HashSet<TestId> suiteTestIds = new HashSet<TestId>();
+		async Task AddTestRef(TestDataRefDocument testRef)
+		{
+			StreamId streamId = testRef.StreamId;
 
-			for (int i = 0; i < refs.Count; i++)
-			{
-				ITestDataRef testRef = refs[i];
+			List<TestStreamDocument> streams = await _testStreams.Find(Builders<TestStreamDocument>.Filter.Eq(x => x.StreamId, streamId)).ToListAsync();
 
+			if (streams.Count == 0)
+			{
+				TestStreamDocument streamDoc = new TestStreamDocument(streamId);
 				if (testRef.TestId != null)
 				{
-					testIds.Add(testRef.TestId.Value);
-					allTestIds.Add(testRef.TestId.Value);
+					streamDoc.Tests.Add(testRef.TestId.Value);
 				}
-				else if (testRef.SuiteTests != null)
+
+				if (testRef.SuiteId != null)
 				{
-					foreach (ISuiteTestData t in testRef.SuiteTests)
-					{
-						suiteTestIds.Add(t.TestId);
-						allTestIds.Add(t.TestId);
-					}
+					streamDoc.TestSuites.Add(testRef.SuiteId.Value);
 				}
-			}
-			List<TestDocument> allTests = await FindTests(testIds: allTestIds.ToArray());
-
-			List<TestDocument> tests = allTests.Where(t => testIds.Contains(t.Id)).ToList();
 
-			List<TestSuiteDocument> suites = await FindTestSuites(testIds: suiteTestIds.ToArray());
-			List<TestDocument> suiteTests = allTests.Where(t => suiteTestIds.Contains(t.Id)).ToList();
-		
-			List<TestMetaId> allMetaIds = allTests.SelectMany(t => t.Metadata).Distinct().ToList();			
-			
-			List<TestStream> testStreams = new List<TestStream>();
-			for (int i = 0; i < streamIds.Length; i++)
+				await _testStreams.InsertOneAsync(streamDoc);
+			}
+			else
 			{
-				StreamId id = streamIds[i];
-
-				List<ITestDataRef> streamRefs = refs.Where(r => r.StreamId == id).ToList();
+				UpdateDefinitionBuilder<TestStreamDocument> updateBuilder = Builders<TestStreamDocument>.Update;
+				List<UpdateDefinition<TestStreamDocument>> updates = new List<UpdateDefinition<TestStreamDocument>>();
 
-				allTestIds.Clear();
-				testIds.Clear();
-				suiteTestIds.Clear();
+				TestStreamDocument streamDoc = streams[0];
 
-				for (int j = 0; j < streamRefs.Count; j++)
+				if (testRef.TestId != null)
 				{
-					ITestDataRef testRef = streamRefs[j];
-
-					if (testRef.TestId != null)
+					if (!streamDoc.Tests.Contains(testRef.TestId.Value))
 					{
-						testIds.Add(testRef.TestId.Value);
-						allTestIds.Add(testRef.TestId.Value);
-					}
-					else if (testRef.SuiteTests != null)
+						streamDoc.Tests.Add(testRef.TestId.Value);
+						updates.Add(updateBuilder.Set(x => x.Tests, streamDoc.Tests));
+					}					
+				}
+
+				if (testRef.SuiteId != null)
+				{
+					if (!streamDoc.TestSuites.Contains(testRef.SuiteId.Value))
 					{
-						foreach (ISuiteTestData t in testRef.SuiteTests)
-						{
-							suiteTestIds.Add(t.TestId);
-							allTestIds.Add(t.TestId);
-						}
+						streamDoc.TestSuites.Add(testRef.SuiteId.Value);
+						updates.Add(updateBuilder.Set(x => x.TestSuites, streamDoc.TestSuites));
 					}
 				}
+				if (updates.Count > 0)
+				{
+					FilterDefinitionBuilder<TestStreamDocument> ebuilder = Builders<TestStreamDocument>.Filter;
+					FilterDefinition<TestStreamDocument> efilter = ebuilder.Eq(x => x.StreamId, streamDoc.StreamId);
+					await _testStreams.UpdateOneAsync(efilter, updateBuilder.Combine(updates));
 
-				List<ITest> streamTests = allTests.Where(x => allTestIds.Contains(x.Id)).ToList<ITest>();
-				List<ITestSuite> streamSuites = suites.Where(x => allTestIds.Intersect(x.Tests).Any()).ToList<ITestSuite>();
-				
-				List<TestMetaId> streamMetaIds = streamTests.SelectMany(t => t.Metadata).Distinct().ToList();
-				List<ITestMeta> streamMeta = metaData.Where(x => streamMetaIds.Contains(x.Id)).ToList();
-				testStreams.Add(new TestStream(id, streamTests, streamSuites, streamMeta));
+				}
 			}
-
-			return testStreams;
+			await _testRefs.InsertOneAsync(testRef);
 		}
 
 		/// <summary>
@@ -1069,7 +1071,7 @@ namespace Horde.Build.Jobs.TestData
 						break;
 				}
 
-				await _testRefs.InsertOneAsync(testRef);
+				await AddTestRef(testRef);
 
 			}
 			catch (Exception ex)
@@ -1213,7 +1215,7 @@ namespace Horde.Build.Jobs.TestData
 							testRef.SuiteTests.Add(new SuiteTestData(document.Id, outcome, TimeSpan.FromSeconds(test.TimeElapseSec), test.TestUID));
 						}
 
-						await _testRefs.InsertOneAsync(testRef);
+						await AddTestRef(testRef);
 					}
 					else
 					{
diff --git a/Engine/Source/Programs/Horde/Horde.Build/Jobs/TestData/TestDataController.cs b/Engine/Source/Programs/Horde/Horde.Build/Jobs/TestData/TestDataController.cs
--- a/Engine/Source/Programs/Horde/Horde.Build/Jobs/TestData/TestDataController.cs
+++ b/Engine/Source/Programs/Horde/Horde.Build/Jobs/TestData/TestDataController.cs
@@ -85,18 +85,11 @@ namespace Horde.Build.Jobs.TestData
 		/// Get stream test data for the provided ids
 		/// </summary>
 		/// <param name="streamIds"></param>
-		/// <param name="metaIds"></param>
-		/// <param name="minCreateTime"></param>
-		/// <param name="maxCreateTime"></param>
 		/// <returns></returns>
 		[HttpGet]
 		[Route("/api/v2/testdata/streams")]
 		[ProducesResponseType(typeof(List<GetTestStreamResponse>), 200)]
-		public async Task<ActionResult<List<GetTestStreamResponse>>> GetTestStreamsAsync(
-			[FromQuery(Name = "Id")] string[] streamIds,
-			[FromQuery(Name = "Mid")] string[] metaIds,
-			[FromQuery] DateTimeOffset minCreateTime,
-			[FromQuery] DateTimeOffset? maxCreateTime = null)
+		public async Task<ActionResult<List<GetTestStreamResponse>>> GetTestStreamsAsync([FromQuery(Name = "Id")] string[] streamIds)
 		{
 			StreamId[] streamIdValues = Array.ConvertAll(streamIds, x => new StreamId(x));
 
@@ -121,8 +114,99 @@ namespace Horde.Build.Jobs.TestData
 				return responses;
 			}
 
-			List<TestStream> streams = await _testDataService.FindTestStreams(queryStreams.ToArray(), metaIds.ConvertAll(x => TestMetaId.Parse(x)).ToArray(), minCreateTime.UtcDateTime, maxCreateTime?.UtcDateTime);
-			streams.ForEach(s => responses.Add(new GetTestStreamResponse(s.StreamId, s.Tests, s.TestSuites, s.TestMeta)));
+			HashSet<TestId> testIds = new HashSet<TestId>();
+			HashSet<TestMetaId> metaIds = new HashSet<TestMetaId>();
+
+			List<ITestStream> streams = await _testDataService.FindTestStreams(queryStreams.ToArray());
+
+			// flatten requested streams to single service queries
+		
+			HashSet<TestSuiteId> suiteIds = new HashSet<TestSuiteId>();
+			for (int i = 0; i < streams.Count; i++)
+			{
+				foreach(TestId testId in streams[i].Tests)
+				{
+					testIds.Add(testId);
+				}
+
+				foreach(TestSuiteId suiteId in streams[i].TestSuites)
+				{
+					suiteIds.Add(suiteId);
+				}				
+			}
+
+			List<ITestSuite> suites = new List<ITestSuite>();
+			if (suiteIds.Count > 0)
+			{
+				suites = await _testDataService.FindTestSuites(suiteIds.ToArray());
+			}
+
+			foreach(ITestSuite suite in suites)
+			{
+				foreach(TestId id in suite.Tests)
+				{
+					testIds.Add(id);
+				}
+			}
+
+			List<ITest> tests = new List<ITest>();
+			if (testIds.Count > 0)
+			{
+				tests = await _testDataService.FindTests(testIds.ToArray());
+			}
+
+			List<ITestMeta> metaData = new List<ITestMeta>();
+			foreach (ITest test in tests)
+			{
+				foreach (TestMetaId metaId in test.Metadata)
+				{
+					metaIds.Add(metaId);
+				}
+			}
+
+			if (metaIds.Count > 0)
+			{
+				metaData = await _testDataService.FindTestMeta(metaIds: metaIds.ToArray());
+			}
+
+			// generate individual stream responses
+			streams.ForEach(s => {
+				
+				List<ITest> streamTests = tests.Where(x => s.Tests.Contains(x.Id)).ToList();
+
+				List<ITestSuite> streamSuites = new List<ITestSuite>();
+				foreach (TestSuiteId suiteId in s.TestSuites)
+				{
+					ITestSuite? suite = suites.FirstOrDefault(x => x.Id == suiteId);
+					if (suite != null)
+					{
+						streamSuites.Add(suite);
+						foreach (TestId testId in suite.Tests)
+						{
+							ITest? test = tests.FirstOrDefault(x => x.Id == testId);
+							if (test != null)
+							{
+								streamTests.Add(test);
+							}
+						}
+					}
+				}
+
+				HashSet<TestMetaId> streamMetaIds = new HashSet<TestMetaId>();
+				foreach (ITest test in streamTests)
+				{
+					foreach (TestMetaId id in test.Metadata)
+					{
+						streamMetaIds.Add(id);
+					}						
+				}
+
+				List<ITestMeta> streamMetaData = metaData.Where(x => streamMetaIds.Contains(x.Id)).ToList();
+
+				responses.Add(new GetTestStreamResponse(s.StreamId, streamTests, streamSuites, streamMetaData));				
+
+			});
+
 			return responses;
 		}
 
diff --git a/Engine/Source/Programs/Horde/Horde.Build/Jobs/TestData/TestDataService.cs b/Engine/Source/Programs/Horde/Horde.Build/Jobs/TestData/TestDataService.cs
--- a/Engine/Source/Programs/Horde/Horde.Build/Jobs/TestData/TestDataService.cs
+++ b/Engine/Source/Programs/Horde/Horde.Build/Jobs/TestData/TestDataService.cs
@@ -22,19 +22,19 @@ namespace Horde.Build.Jobs.TestData
 	public sealed class TestDataService : IDisposable
 	{
 		
-		readonly ITestDataCollection _testData;
+		readonly ITestDataCollection _testData;		
 
 		/// <summary>
 		/// Device service constructor
 		/// </summary>
 		public TestDataService(ITestDataCollection testData)
 		{
-			_testData = testData;			
+			_testData = testData;		
 		}
 
 		/// <inheritdoc/>
 		public void Dispose()
-		{
+		{		
 		}
 
 
@@ -42,18 +42,31 @@ namespace Horde.Build.Jobs.TestData
 		/// Find test streams
 		/// </summary>
 		/// <param name="streamIds"></param>
-		/// <param name="metaIds"></param>
-		/// <param name="minCreateTime"></param>
-		/// <param name="maxCreateTime"></param>
 		/// <returns></returns>
-		public async Task<List<TestStream>> FindTestStreams(StreamId[] streamIds, TestMetaId[] metaIds, DateTime minCreateTime, DateTime? maxCreateTime = null)
+		public async Task<List<ITestStream>> FindTestStreams(StreamId[] streamIds)
+		{
+
+			return await _testData.FindTestStreams(streamIds);			
+		}
+
+		/// <summary>
+		/// Find tests
+		/// </summary>
+		/// <param name="testIds"></param>
+		/// <returns></returns>
+		public async Task<List<ITest>> FindTests(TestId[] testIds)
 		{
-			if (maxCreateTime == null)
-			{
-				maxCreateTime = DateTime.UtcNow;
-			}
+			return await _testData.FindTests(testIds);
+		}
 
-			return await _testData.FindTestStreams(streamIds, metaIds, minCreateTime, maxCreateTime.Value);
+		/// <summary>
+		/// Find test suites
+		/// </summary>
+		/// <param name="suiteIds"></param>
+		/// <returns></returns>
+		public async Task<List<ITestSuite>> FindTestSuites(TestSuiteId[] suiteIds)
+		{
+			return await _testData.FindTestSuites(suiteIds);
 		}
 
 		/// <summary>
@@ -91,4 +104,4 @@ namespace Horde.Build.Jobs.TestData
 			return await _testData.FindTestRefs(streamIds, metaIds, tids, sids, minCreateTime, maxCreateTime, minChange, maxChange);
 		}
 	}
-}
+}
\ No newline at end of file
