Horde: Allow setting the initial agent type per-stream.

diff --git a/Engine/Source/Programs/Horde/Horde.Build.Tests/Fixture.cs b/Engine/Source/Programs/Horde/Horde.Build.Tests/Fixture.cs
--- a/Engine/Source/Programs/Horde/Horde.Build.Tests/Fixture.cs
+++ b/Engine/Source/Programs/Horde/Horde.Build.Tests/Fixture.cs
@@ -62,7 +62,7 @@ namespace Horde.Build.Tests
 			fg.Labels = new List<ILabel>();
 
 			Template = await templateCollection.AddAsync("Test template");
-			Graph = await graphCollection.AddAsync(Template);
+			Graph = await graphCollection.AddAsync(Template, null);
 
 			TemplateRefId1 = new TemplateId("template1");
 			TemplateRefId2 = new TemplateId("template2");
diff --git a/Engine/Source/Programs/Horde/Horde.Build.Tests/JobCollectionTests.cs b/Engine/Source/Programs/Horde/Horde.Build.Tests/JobCollectionTests.cs
--- a/Engine/Source/Programs/Horde/Horde.Build.Tests/JobCollectionTests.cs
+++ b/Engine/Source/Programs/Horde/Horde.Build.Tests/JobCollectionTests.cs
@@ -75,7 +75,7 @@ namespace Horde.Build.Tests
 			Mock<ITemplate> templateMock = new Mock<ITemplate>(MockBehavior.Strict);
 			templateMock.SetupGet(x => x.InitialAgentType).Returns((string?)null);
 
-			IGraph baseGraph = await GraphCollection.AddAsync(templateMock.Object);
+			IGraph baseGraph = await GraphCollection.AddAsync(templateMock.Object, null);
 
 			CreateJobOptions options = new CreateJobOptions();
 			options.Arguments.Add("-Target=Publish Client");
@@ -150,7 +150,7 @@ namespace Horde.Build.Tests
 			Mock<ITemplate> templateMock = new Mock<ITemplate>(MockBehavior.Strict);
 			templateMock.SetupGet(x => x.InitialAgentType).Returns((string?)null);
 
-			IGraph baseGraph = await GraphCollection.AddAsync(templateMock.Object);
+			IGraph baseGraph = await GraphCollection.AddAsync(templateMock.Object, null);
 
 			CreateJobOptions options = new CreateJobOptions();
 			options.Arguments.Add("-Target=Step 1");
@@ -206,7 +206,7 @@ namespace Horde.Build.Tests
 			Mock<ITemplate> templateMock = new Mock<ITemplate>(MockBehavior.Strict);
 			templateMock.SetupGet(x => x.InitialAgentType).Returns((string?)null);
 
-			IGraph baseGraph = await GraphCollection.AddAsync(templateMock.Object);
+			IGraph baseGraph = await GraphCollection.AddAsync(templateMock.Object, null);
 
 			CreateJobOptions options = new CreateJobOptions();
 			options.Arguments.Add("-Target=Step 1");
@@ -257,7 +257,7 @@ namespace Horde.Build.Tests
 			Mock<ITemplate> templateMock = new Mock<ITemplate>(MockBehavior.Strict);
 			templateMock.SetupGet(x => x.InitialAgentType).Returns((string?)null);
 
-			IGraph baseGraph = await GraphCollection.AddAsync(templateMock.Object);
+			IGraph baseGraph = await GraphCollection.AddAsync(templateMock.Object, null);
 
 			CreateJobOptions options = new CreateJobOptions();
 			options.Arguments.Add("-Target=Step 1");
diff --git a/Engine/Source/Programs/Horde/Horde.Build.Tests/JobServiceTests.cs b/Engine/Source/Programs/Horde/Horde.Build.Tests/JobServiceTests.cs
--- a/Engine/Source/Programs/Horde/Horde.Build.Tests/JobServiceTests.cs
+++ b/Engine/Source/Programs/Horde/Horde.Build.Tests/JobServiceTests.cs
@@ -46,7 +46,7 @@ namespace Horde.Build.Tests
 			Assert.IsNotNull(project);
 
 			ITemplate template = await TemplateCollection.AddAsync("Test template");
-			IGraph graph = await GraphCollection.AddAsync(template);
+			IGraph graph = await GraphCollection.AddAsync(template, null);
 
 			TemplateId templateRefId1 = new TemplateId("template1");
 			TemplateId templateRefId2 = new TemplateId("template2");
@@ -216,7 +216,7 @@ namespace Horde.Build.Tests
 			stream = await CreateOrReplaceStreamAsync(streamId, stream, project!.Id, new StreamConfig { Name = "//UE5/Main" });
 
 			ITemplate template = await TemplateCollection.AddAsync("Test template");
-			IGraph graph = await GraphCollection.AddAsync(template);
+			IGraph graph = await GraphCollection.AddAsync(template, null);
 
 			NewGroup groupA = new NewGroup("win", new List<NewNode>());
 			groupA.Nodes.Add(new NewNode("Compile"));
diff --git a/Engine/Source/Programs/Horde/Horde.Build.Tests/SchedulerTests.cs b/Engine/Source/Programs/Horde/Horde.Build.Tests/SchedulerTests.cs
--- a/Engine/Source/Programs/Horde/Horde.Build.Tests/SchedulerTests.cs
+++ b/Engine/Source/Programs/Horde/Horde.Build.Tests/SchedulerTests.cs
@@ -435,7 +435,7 @@ namespace Horde.Build.Tests
 			stream = (await CreateOrReplaceStreamAsync(StreamId, stream, ProjectId, config))!;
 
 			// Create the TriggerNext step and mark it as complete
-			IGraph graphA = await GraphCollection.AddAsync(newTemplate1);
+			IGraph graphA = await GraphCollection.AddAsync(newTemplate1, null);
 			NewGroup groupA = new NewGroup("win", new List<NewNode> { new NewNode("TriggerNext") });
 			graphA = await GraphCollection.AppendAsync(graphA, new List<NewGroup> { groupA });
 
@@ -530,7 +530,7 @@ namespace Horde.Build.Tests
 			Assert.IsNotNull(template2);
 
 			// Create the graph
-			IGraph graphA = await GraphCollection.AddAsync(template1);
+			IGraph graphA = await GraphCollection.AddAsync(template1, null);
 			NewGroup groupA = new NewGroup("win", new List<NewNode> { new NewNode("TriggerNext") });
 			graphA = await GraphCollection.AppendAsync(graphA, new List<NewGroup> { groupA });
 
diff --git a/Engine/Source/Programs/Horde/Horde.Build/Jobs/Graphs/GraphCollection.cs b/Engine/Source/Programs/Horde/Horde.Build/Jobs/Graphs/GraphCollection.cs
--- a/Engine/Source/Programs/Horde/Horde.Build/Jobs/Graphs/GraphCollection.cs
+++ b/Engine/Source/Programs/Horde/Horde.Build/Jobs/Graphs/GraphCollection.cs
@@ -449,10 +449,10 @@ namespace Horde.Build.Jobs.Graphs
 		}
 
 		/// <inheritdoc/>
-		public async Task<IGraph> AddAsync(ITemplate template)
+		public async Task<IGraph> AddAsync(ITemplate template, string? streamInitialAgentType)
 		{
 			Node node = new Node(IJob.SetupNodeName, null, null, Array.Empty<NodeRef>(), Array.Empty<NodeRef>(), Priority.High, true, false, true, null, null, null);
-			NodeGroup group = new NodeGroup(template.InitialAgentType ?? "Win64", new List<Node> { node });
+			NodeGroup group = new NodeGroup(template.InitialAgentType ?? streamInitialAgentType ?? "Win64", new List<Node> { node });
 
 			GraphDocument graph = new GraphDocument(new List<NodeGroup> { group }, new List<Aggregate>(), new List<Label>());
 			await AddAsync(graph);
diff --git a/Engine/Source/Programs/Horde/Horde.Build/Jobs/Graphs/IGraphCollection.cs b/Engine/Source/Programs/Horde/Horde.Build/Jobs/Graphs/IGraphCollection.cs
--- a/Engine/Source/Programs/Horde/Horde.Build/Jobs/Graphs/IGraphCollection.cs
+++ b/Engine/Source/Programs/Horde/Horde.Build/Jobs/Graphs/IGraphCollection.cs
@@ -16,8 +16,9 @@ namespace Horde.Build.Jobs.Graphs
 		/// Adds a graph from a template
 		/// </summary>
 		/// <param name="template">The template</param>
+		/// <param name="streamInitialAgentType">Default agent type for the stream</param>
 		/// <returns>New graph</returns>
-		Task<IGraph> AddAsync(ITemplate template);
+		Task<IGraph> AddAsync(ITemplate template, string? streamInitialAgentType);
 
 		/// <summary>
 		/// Creates a graph by appending groups and aggregates to an existing graph.
diff --git a/Engine/Source/Programs/Horde/Horde.Build/Jobs/JobService.cs b/Engine/Source/Programs/Horde/Horde.Build/Jobs/JobService.cs
--- a/Engine/Source/Programs/Horde/Horde.Build/Jobs/JobService.cs
+++ b/Engine/Source/Programs/Horde/Horde.Build/Jobs/JobService.cs
@@ -1327,7 +1327,7 @@ namespace Horde.Build.Jobs
 						break;
 					}
 
-					IGraph triggerGraph = await _graphs.AddAsync(template);
+					IGraph triggerGraph = await _graphs.AddAsync(template, stream.Config.InitialAgentType);
 					_logger.LogInformation("Creating downstream job {ChainedJobId} from job {JobId}", chainedJobId, newJob.Id);
 
 					CreateJobOptions options = new CreateJobOptions(templateRef.Config);
diff --git a/Engine/Source/Programs/Horde/Horde.Build/Jobs/JobsController.cs b/Engine/Source/Programs/Horde/Horde.Build/Jobs/JobsController.cs
--- a/Engine/Source/Programs/Horde/Horde.Build/Jobs/JobsController.cs
+++ b/Engine/Source/Programs/Horde/Horde.Build/Jobs/JobsController.cs
@@ -164,7 +164,7 @@ namespace Horde.Build.Jobs
 			Priority priority = create.Priority ?? template.Priority ?? Priority.Normal;
 
 			// New groups for the job
-			IGraph graph = await _graphs.AddAsync(template);
+			IGraph graph = await _graphs.AddAsync(template, stream.Config.InitialAgentType);
 
 			// Get the commits for this stream
 			ICommitCollection commits = _commitService.GetCollection(stream);
diff --git a/Engine/Source/Programs/Horde/Horde.Build/Jobs/Schedules/ScheduleService.cs b/Engine/Source/Programs/Horde/Horde.Build/Jobs/Schedules/ScheduleService.cs
--- a/Engine/Source/Programs/Horde/Horde.Build/Jobs/Schedules/ScheduleService.cs
+++ b/Engine/Source/Programs/Horde/Horde.Build/Jobs/Schedules/ScheduleService.cs
@@ -446,7 +446,7 @@ namespace Horde.Build.Jobs.Schedules
 			}
 
 			// Register the graph for it
-			IGraph graph = await _graphs.AddAsync(template);
+			IGraph graph = await _graphs.AddAsync(template, stream.Config.InitialAgentType);
 
 			// We may need to submit a new change for any new jobs. This only makes sense if there's one change.
 			if (template.SubmitNewChange != null)
diff --git a/Engine/Source/Programs/Horde/Horde.Build/Streams/StreamConfig.cs b/Engine/Source/Programs/Horde/Horde.Build/Streams/StreamConfig.cs
--- a/Engine/Source/Programs/Horde/Horde.Build/Streams/StreamConfig.cs
+++ b/Engine/Source/Programs/Horde/Horde.Build/Streams/StreamConfig.cs
@@ -87,6 +87,11 @@ namespace Horde.Build.Streams
 		/// </summary>
 		public int Order { get; set; } = 128;
 
+		/// <summary>
+		/// Default initial agent type for templates
+		/// </summary>
+		public string? InitialAgentType { get; set; }
+
 		/// <summary>
 		/// Notification channel for all jobs in this stream
 		/// </summary>
