Horde: Remove Gauntlet regex matching in favor of structured logging
#rnx

diff --git a/Engine/Source/Programs/AutomationTool/AutomationUtils/Matchers/GauntletEventMatcher.cs b/Engine/Source/Programs/AutomationTool/AutomationUtils/Matchers/GauntletEventMatcher.cs
deleted file mode 100644
diff --git a/Engine/Source/Programs/Horde/Horde.Agent.Tests/LogParserTests.cs b/Engine/Source/Programs/Horde/Horde.Agent.Tests/LogParserTests.cs
--- a/Engine/Source/Programs/Horde/Horde.Agent.Tests/LogParserTests.cs
+++ b/Engine/Source/Programs/Horde/Horde.Agent.Tests/LogParserTests.cs
@@ -515,38 +515,6 @@ namespace Horde.Agent.Tests
 			Assert.AreEqual(0, logEvents.Count);
 		}
 
-		[TestMethod]
-		[Ignore]
-		public void GauntletGenericErrorMatcher()
-		{
-			string[] lines =
-			{
-				@"  Error: EngineTest.RunTests Group:HLOD (Win64 Development EditorGame) result=Failed",
-				@"    # EngineTest.RunTests Group:HLOD Report",
-				@"    ----------------------------------------",
-				@"    ### Process Role: Editor (Win64 Development)",
-				@"    ----------------------------------------",
-				@"    ##### Result: Abnormal Exit: Reason=3/24 tests failed, Code=-1",
-				@"    FatalErrors: 0, Ensures: 0, Errors: 8, Warnings: 20, Hash: 0",
-				@"    ##### Artifacts",
-				@"    Log: P:/Builds/Automation/Reports/++UE5+Main/EngineTest/++UE5+Main-CL-14167122/HLOD_Win64Editor\Saved_1\Editor\EditorOutput.log",
-				@"    Commandline: d:\Build\++UE5\Sync\EngineTest\EngineTest.uproject   -gauntlet  -unattended  -stdout  -AllowStdOutLogVerbosity  -gauntlet.heartbeatperiod=30  -NoWatchdog  -FORCELOGFLUSH  -CrashForUAT  -buildmachine  -ReportExportPath=""P:/Builds/Automation/Reports/++UE5+Main/EngineTest/++UE5+Main-CL-14167122/HLOD_Win64Editor""  -ExecCmds=""Automation RunTests Group:HLOD; Quit;""  -ddc=default  -userdir=""d:\Build\++UE5\Sync/Tests\DeviceCache\Win64\LocalDevice0_UserDir""",
-				@"    P:/Builds/Automation/Reports/++UE5+Main/EngineTest/++UE5+Main-CL-14167122/HLOD_Win64Editor\Saved_1\Editor",
-				@"    ----------------------------------------",
-				@"    ## Summary",
-				@"    ### EngineTest.RunTests Group:HLOD Failed",
-				@"    ### Editor: 3/24 tests failed",
-				@"    See below for logs and any callstacks",
-				@"    Context: Win64 Development EditorGame",
-				@"    FatalErrors: 0, Ensures: 0, Errors: 8, Warnings: 20",
-				@"    Result: Failed, ResultHash: 0",
-				@"    21 of 24 tests passed"
-			};
-
-			List<LogEvent> logEvents = Parse(lines);
-			CheckEventGroup(logEvents, 0, 20, LogLevel.Error, KnownLogEvents.Gauntlet);
-		}
-
 		[TestMethod]
 		public void DockerWarningMatcher()
 		{
@@ -582,98 +550,6 @@ namespace Horde.Agent.Tests
 			CheckEventGroup(logEvents, 2, 2, LogLevel.Warning, KnownLogEvents.Linker);
 		}
 
-		[TestMethod]
-		[Ignore]
-		public void GauntletErrorMatcher()
-		{
-			string[] lines =
-			{
-				@"  Error: EngineTest.RunTests Group:HLOD (Win64 Development EditorGame) result=Failed",
-				@"    # EngineTest.RunTests Group:HLOD Report",
-				@"    ----------------------------------------",
-				@"    ### Process Role: Editor (Win64 Development)",
-				@"    ----------------------------------------",
-				@"    ##### Result: Abnormal Exit: Reason=3/24 tests failed, Code=-1",
-				@"    FatalErrors: 0, Ensures: 0, Errors: 8, Warnings: 20, Hash: 0",
-				@"",
-				@"    ##### Artifacts",
-				@"    Log: P:/Builds/Automation/Reports/++UE5+Main/EngineTest/++UE5+Main-CL-14167122/HLOD_Win64Editor\Saved_1\Editor\EditorOutput.log",
-				@"",
-				@"    Commandline: d:\Build\++UE5\Sync\EngineTest\EngineTest.uproject   -gauntlet  -unattended  -stdout  -AllowStdOutLogVerbosity  -gauntlet.heartbeatperiod=30  -NoWatchdog  -FORCELOGFLUSH  -CrashForUAT  -buildmachine  -ReportExportPath=""P:/Builds/Automation/Reports/++UE5+Main/EngineTest/++UE5+Main-CL-14167122/HLOD_Win64Editor""  -ExecCmds=""Automation RunTests Group:HLOD; Quit;""  -ddc=default  -userdir=""d:\Build\++UE5\Sync/Tests\DeviceCache\Win64\LocalDevice0_UserDir""",
-				@"",
-				@"    P:/Builds/Automation/Reports/++UE5+Main/EngineTest/++UE5+Main-CL-14167122/HLOD_Win64Editor\Saved_1\Editor",
-				@"",
-				@"    ----------------------------------------",
-				@"    ## Summary",
-				@"    ### EngineTest.RunTests Group:HLOD Failed",
-				@"    ### Editor: 3/24 tests failed",
-				@"    See below for logs and any callstacks",
-				@"",
-				@"    Context: Win64 Development EditorGame",
-				@"",
-				@"    FatalErrors: 0, Ensures: 0, Errors: 8, Warnings: 20",
-				@"    Result: Failed, ResultHash: 0",
-				@"",
-				@"    21 of 24 tests passed",
-				@"",
-				@"    ### The following tests failed:",
-				@"    ##### SectionFlags: SectionFlags",
-				@"    * LogAutomationController: Building static mesh SectionFlags... [D:\Build\++UE5\Sync\Engine\Source\Runtime\Core\Private\Logging\LogMacros.cpp(92)]",
-				@"    * LogAutomationController: Building static mesh SectionFlags... [D:\Build\++UE5\Sync\Engine\Source\Runtime\Core\Private\Logging\LogMacros.cpp(92)]",
-				@"    * LogAutomationController: Err0r: Screenshot 'ActorMerging_SectionFlags_LOD_0_None' test failed, Screenshots were different!  Global Difference = 0.058361, Max Local Difference = 0.821376 [D:\Build\++UE5\Sync\Engine\Source\Runtime\Core\Public\Delegates\DelegateInstancesImpl.h(546)]",
-				@"",
-				@"    ##### SimpleMerge: SimpleMerge",
-				@"    * LogAutomationController: Building static mesh SM_TeapotHLOD... [D:\Build\++UE5\Sync\Engine\Source\Runtime\Core\Private\Logging\LogMacros.cpp(92)]",
-				@"    * LogAutomationController: Building static mesh SM_TeapotHLOD... [D:\Build\++UE5\Sync\Engine\Source\Runtime\Core\Private\Logging\LogMacros.cpp(92)]",
-				@"    * LogAutomationController: Screenshot 'ActorMerging_SimpleMeshMerge_LOD_0_None' was similar!  Global Difference = 0.000298, Max Local Difference = 0.010725 [D:\Build\++UE5\Sync\Engine\Source\Runtime\Core\Public\Delegates\DelegateInstancesImpl.h(546)]",
-				@"    * LogAutomationController: Err0r: Screenshot 'ActorMerging_SimpleMeshMerge_LOD_1_None' test failed, Screenshots were different!  Global Difference = 0.006954, Max Local Difference = 0.129438 [D:\Build\++UE5\Sync\Engine\Source\Runtime\Core\Public\Delegates\DelegateInstancesImpl.h(546)]",
-				@"    * LogAutomationController: Err0r: Screenshot 'ActorMerging_SimpleMeshMerge_LOD_2_None' test failed, Screenshots were different!  Global Difference = 0.007732, Max Local Difference = 0.127959 [D:\Build\++UE5\Sync\Engine\Source\Runtime\Core\Public\Delegates\DelegateInstancesImpl.h(546)]",
-				@"    * LogAutomationController: Err0r: Screenshot 'ActorMerging_SimpleMeshMerge_LOD_3_None' test failed, Screenshots were different!  Global Difference = 0.009140, Max Local Difference = 0.172337 [D:\Build\++UE5\Sync\Engine\Source\Runtime\Core\Public\Delegates\DelegateInstancesImpl.h(546)]",
-				@"    * LogAutomationController: Screenshot 'ActorMerging_SimpleMeshMerge_LOD_0_BaseColor' was similar!  Global Difference = 0.000000, Max Local Difference = 0.000000 [D:\Build\++UE5\Sync\Engine\Source\Runtime\Core\Public\Delegates\DelegateInstancesImpl.h(546)]",
-				@"    * LogAutomationController: Screenshot 'ActorMerging_SimpleMeshMerge_LOD_1_BaseColor' was similar!  Global Difference = 0.002068, Max Local Difference = 0.045858 [D:\Build\++UE5\Sync\Engine\Source\Runtime\Core\Public\Delegates\DelegateInstancesImpl.h(546)]",
-				@"    * LogAutomationController: Screenshot 'ActorMerging_SimpleMeshMerge_LOD_2_BaseColor' was similar!  Global Difference = 0.002377, Max Local Difference = 0.045858 [D:\Build\++UE5\Sync\Engine\Source\Runtime\Core\Public\Delegates\DelegateInstancesImpl.h(546)]",
-				@"    * LogAutomationController: Screenshot 'ActorMerging_SimpleMeshMerge_LOD_3_BaseColor' was similar!  Global Difference = 0.002647, Max Local Difference = 0.057322 [D:\Build\++UE5\Sync\Engine\Source\Runtime\Core\Public\Delegates\DelegateInstancesImpl.h(546)]",
-				@"",
-				@"    ##### SingleLODMerge: SingleLODMerge",
-				@"    * LogAutomationController: Building static mesh Pencil2... [D:\Build\++UE5\Sync\Engine\Source\Runtime\Core\Private\Logging\LogMacros.cpp(92)]",
-				@"    * LogAutomationController: Building static mesh Pencil2... [D:\Build\++UE5\Sync\Engine\Source\Runtime\Core\Private\Logging\LogMacros.cpp(92)]",
-				@"    * LogAutomationController: Err0r: Screenshot 'ActorMerging_SingleLODMerge_LOD_0_BaseColor' test failed, Screenshots were different!  Global Difference = 0.013100, Max Local Difference = 0.131657 [D:\Build\++UE5\Sync\Engine\Source\Runtime\Core\Public\Delegates\DelegateInstancesImpl.h(546)]",
-				@"",
-				@"    ### Links",
-				@"    View results here: http://automation.epicgames.net/reports/++UE5+Main/EngineTest/++UE5+Main-CL-14167122/HLOD_Win64Editor/index.html",
-				@"",
-				@"    Open results in UnrealEd from P:/Builds/Automation/Reports/++UE5+Main/EngineTest/++UE5+Main-CL-14167122/HLOD_Win64Editor",
-			};
-
-			List<LogEvent> logEvents = Parse(lines);
-			CheckEventGroup(logEvents, 0, 55, LogLevel.Error, KnownLogEvents.Gauntlet_TestEvent);
-
-			Assert.AreEqual("HLOD", logEvents[29].GetProperty("group").ToString());
-			Assert.AreEqual("SectionFlags", logEvents[29].GetProperty("name").ToString());
-			Assert.AreEqual("SectionFlags", logEvents[29].GetProperty("friendly_name").ToString());
-
-			Assert.AreEqual("HLOD", logEvents[34].GetProperty("group").ToString());
-			Assert.AreEqual("SimpleMerge", logEvents[34].GetProperty("name").ToString());
-			Assert.AreEqual("SimpleMerge", logEvents[34].GetProperty("friendly_name").ToString());
-
-			Assert.AreEqual("HLOD", logEvents[46].GetProperty("group").ToString());
-			Assert.AreEqual("SingleLODMerge", logEvents[46].GetProperty("name").ToString());
-			Assert.AreEqual("SingleLODMerge", logEvents[46].GetProperty("friendly_name").ToString());
-		}
-
-		[TestMethod]
-		[Ignore]
-		public void GauntletScreenshotErrorMatcher()
-		{
-			string text = @"  Error: LogAutomationController: Error: Screenshot 'ActorMerging_SectionFlags_LOD_0_None' test failed, Screenshots were different!  Global Difference = 0.058361, Max Local Difference = 0.821376 [D:\Build\++UE5\Sync\Engine\Source\Runtime\Core\Public\Delegates\DelegateInstancesImpl.h(546)]";
-
-			List<LogEvent> logEvents = Parse(text);
-			CheckEventGroup(logEvents, 0, 1, LogLevel.Error, KnownLogEvents.Gauntlet_UnrealEngineTestEvent);
-
-			LogEvent logEvent = logEvents[0];
-			Assert.AreEqual("ActorMerging_SectionFlags_LOD_0_None", logEvent.GetProperty("screenshot").ToString());
-		}
-
 		[TestMethod]
 		public void SystemicErrorMatcher()
 		{
diff --git a/Engine/Source/Programs/Horde/Horde.Build.Tests/IssueServiceTests.cs b/Engine/Source/Programs/Horde/Horde.Build.Tests/IssueServiceTests.cs
--- a/Engine/Source/Programs/Horde/Horde.Build.Tests/IssueServiceTests.cs
+++ b/Engine/Source/Programs/Horde/Horde.Build.Tests/IssueServiceTests.cs
@@ -1698,79 +1698,6 @@ namespace Horde.Build.Tests
 			}
 		}
 
-		[TestMethod]
-		[Ignore]
-		public async Task GauntletIssueTest()
-		{
-			string[] lines =
-			{
-				@"  Error: EngineTest.RunTests Group:HLOD (Win64 Development EditorGame) result=Failed",
-				@"    # EngineTest.RunTests Group:HLOD Report",
-				@"    ----------------------------------------",
-				@"    ### Process Role: Editor (Win64 Development)",
-				@"    ----------------------------------------",
-				@"    ##### Result: Abnormal Exit: Reason=3/24 tests failed, Code=-1",
-				@"    FatalErrors: 0, Ensures: 0, Errors: 8, Warnings: 20, Hash: 0",
-				@"    ##### Artifacts",
-				@"    Log: P:/Builds/Automation/Reports/++UE5+Main/EngineTest/++UE5+Main-CL-14167122/HLOD_Win64Editor\Saved_1\Editor\EditorOutput.log",
-				@"    Commandline: d:\Build\++UE5\Sync\EngineTest\EngineTest.uproject   -gauntlet  -unattended  -stdout  -AllowStdOutLogVerbosity  -gauntlet.heartbeatperiod=30  -NoWatchdog  -FORCELOGFLUSH  -CrashForUAT  -buildmachine  -ReportExportPath=""P:/Builds/Automation/Reports/++UE5+Main/EngineTest/++UE5+Main-CL-14167122/HLOD_Win64Editor""  -ExecCmds=""Automation RunTests Group:HLOD; Quit;""  -ddc=default  -userdir=""d:\Build\++UE5\Sync/Tests\DeviceCache\Win64\LocalDevice0_UserDir""",
-				@"    P:/Builds/Automation/Reports/++UE5+Main/EngineTest/++UE5+Main-CL-14167122/HLOD_Win64Editor\Saved_1\Editor",
-				@"    ----------------------------------------",
-				@"    ## Summary",
-				@"    ### EngineTest.RunTests Group:HLOD Failed",
-				@"    ### Editor: 3/24 tests failed",
-				@"    See below for logs and any callstacks",
-				@"    Context: Win64 Development EditorGame",
-				@"    FatalErrors: 0, Ensures: 0, Errors: 8, Warnings: 20",
-				@"    Result: Failed, ResultHash: 0",
-				@"    21 of 24 tests passed",
-				@"    ### The following tests failed:",
-				@"    ##### SectionFlags: SectionFlags",
-				@"    * LogAutomationController: Building static mesh SectionFlags... [D:\Build\++UE5\Sync\Engine\Source\Runtime\Core\Private\Logging\LogMacros.cpp(92)]",
-				@"    * LogAutomationController: Building static mesh SectionFlags... [D:\Build\++UE5\Sync\Engine\Source\Runtime\Core\Private\Logging\LogMacros.cpp(92)]",
-				@"    * LogAutomationController: Err0r: Screenshot 'ActorMerging_SectionFlags_LOD_0_None' test failed, Screenshots were different!  Global Difference = 0.058361, Max Local Difference = 0.821376 [D:\Build\++UE5\Sync\Engine\Source\Runtime\Core\Public\Delegates\DelegateInstancesImpl.h(546)]",
-				@"    ##### SimpleMerge: SimpleMerge",
-				@"    * LogAutomationController: Building static mesh SM_TeapotHLOD... [D:\Build\++UE5\Sync\Engine\Source\Runtime\Core\Private\Logging\LogMacros.cpp(92)]",
-				@"    * LogAutomationController: Building static mesh SM_TeapotHLOD... [D:\Build\++UE5\Sync\Engine\Source\Runtime\Core\Private\Logging\LogMacros.cpp(92)]",
-				@"    * LogAutomationController: Screenshot 'ActorMerging_SimpleMeshMerge_LOD_0_None' was similar!  Global Difference = 0.000298, Max Local Difference = 0.010725 [D:\Build\++UE5\Sync\Engine\Source\Runtime\Core\Public\Delegates\DelegateInstancesImpl.h(546)]",
-				@"    * LogAutomationController: Err0r: Screenshot 'ActorMerging_SimpleMeshMerge_LOD_1_None' test failed, Screenshots were different!  Global Difference = 0.006954, Max Local Difference = 0.129438 [D:\Build\++UE5\Sync\Engine\Source\Runtime\Core\Public\Delegates\DelegateInstancesImpl.h(546)]",
-				@"    * LogAutomationController: Err0r: Screenshot 'ActorMerging_SimpleMeshMerge_LOD_2_None' test failed, Screenshots were different!  Global Difference = 0.007732, Max Local Difference = 0.127959 [D:\Build\++UE5\Sync\Engine\Source\Runtime\Core\Public\Delegates\DelegateInstancesImpl.h(546)]",
-				@"    * LogAutomationController: Err0r: Screenshot 'ActorMerging_SimpleMeshMerge_LOD_3_None' test failed, Screenshots were different!  Global Difference = 0.009140, Max Local Difference = 0.172337 [D:\Build\++UE5\Sync\Engine\Source\Runtime\Core\Public\Delegates\DelegateInstancesImpl.h(546)]",
-				@"    * LogAutomationController: Screenshot 'ActorMerging_SimpleMeshMerge_LOD_0_BaseColor' was similar!  Global Difference = 0.000000, Max Local Difference = 0.000000 [D:\Build\++UE5\Sync\Engine\Source\Runtime\Core\Public\Delegates\DelegateInstancesImpl.h(546)]",
-				@"    * LogAutomationController: Screenshot 'ActorMerging_SimpleMeshMerge_LOD_1_BaseColor' was similar!  Global Difference = 0.002068, Max Local Difference = 0.045858 [D:\Build\++UE5\Sync\Engine\Source\Runtime\Core\Public\Delegates\DelegateInstancesImpl.h(546)]",
-				@"    * LogAutomationController: Screenshot 'ActorMerging_SimpleMeshMerge_LOD_2_BaseColor' was similar!  Global Difference = 0.002377, Max Local Difference = 0.045858 [D:\Build\++UE5\Sync\Engine\Source\Runtime\Core\Public\Delegates\DelegateInstancesImpl.h(546)]",
-				@"    * LogAutomationController: Screenshot 'ActorMerging_SimpleMeshMerge_LOD_3_BaseColor' was similar!  Global Difference = 0.002647, Max Local Difference = 0.057322 [D:\Build\++UE5\Sync\Engine\Source\Runtime\Core\Public\Delegates\DelegateInstancesImpl.h(546)]",
-				@"    ##### SingleLODMerge: SingleLODMerge",
-				@"    * LogAutomationController: Building static mesh Pencil2... [D:\Build\++UE5\Sync\Engine\Source\Runtime\Core\Private\Logging\LogMacros.cpp(92)]",
-				@"    * LogAutomationController: Building static mesh Pencil2... [D:\Build\++UE5\Sync\Engine\Source\Runtime\Core\Private\Logging\LogMacros.cpp(92)]",
-				@"    * LogAutomationController: Err0r: Screenshot 'ActorMerging_SingleLODMerge_LOD_0_BaseColor' test failed, Screenshots were different!  Global Difference = 0.013100, Max Local Difference = 0.131657 [D:\Build\++UE5\Sync\Engine\Source\Runtime\Core\Public\Delegates\DelegateInstancesImpl.h(546)]"
-			};
-
-			// #1
-			// Scenario: Job with multiple GJob step fails at CL 120 with compile & link error
-			// Expected: Creates one issue for compile error
-			{
-				IJob job = CreateJob(_mainStreamId, 120, "Test Build", _graph);
-				await ParseAsync(job.Batches[0].Steps[0].LogId!.Value, lines);
-				await UpdateCompleteStep(job, 0, 0, JobStepOutcome.Failure);
-
-				ILogFile? log = await LogFileService.GetLogFileAsync(job.Batches[0].Steps[0].LogId!.Value, CancellationToken.None);
-				List<ILogEvent> events = await LogFileService.FindEventsAsync(log!);
-				Assert.AreEqual(1, events.Count);
-				Assert.AreEqual(40, events[0].LineCount);
-
-				List<IIssue> issues = await IssueCollection.FindIssuesAsync();
-				Assert.AreEqual(1, issues.Count);
-
-				IIssue issue = issues[0];
-				List<IIssueSpan> spans = await IssueCollection.FindSpansAsync(issue.Id);
-				Assert.AreEqual(spans.Count, 1);
-				Assert.AreEqual(issue.Fingerprints.Count, 1);
-				Assert.AreEqual(issue.Fingerprints[0].Type, "Gauntlet");
-				Assert.AreEqual(issue.Summary, "HLOD test failures: SectionFlags, SimpleMerge and SingleLODMerge");
-			}
-		}
-
 		[TestMethod]
 		public async Task FixFailedTest()
 		{
