Horde: Add an option (enabled by default for local builds) which forces a config update on startup. When running locally or in a single-server setup, it's important to apply config changes immediately rather than waiting for the next scheduled update.

diff --git a/Engine/Source/Programs/Horde/Horde.Build/Configuration/ConfigUpdateService.cs b/Engine/Source/Programs/Horde/Horde.Build/Configuration/ConfigUpdateService.cs
--- a/Engine/Source/Programs/Horde/Horde.Build/Configuration/ConfigUpdateService.cs
+++ b/Engine/Source/Programs/Horde/Horde.Build/Configuration/ConfigUpdateService.cs
@@ -82,7 +82,14 @@ namespace Horde.Build.Configuration
 		}
 
 		/// <inheritdoc/>
-		public Task StartAsync(CancellationToken cancellationToken) => _ticker.StartAsync();
+		public async Task StartAsync(CancellationToken cancellationToken)
+		{
+			if (_settings.CurrentValue.ForceConfigUpdateOnStartup)
+			{
+				await TickLeaderAsync(cancellationToken);
+			}
+			await _ticker.StartAsync();
+		}
 
 		/// <inheritdoc/>
 		public Task StopAsync(CancellationToken cancellationToken) => _ticker.StopAsync();
@@ -451,15 +458,16 @@ namespace Horde.Build.Configuration
 		{
 			Uri? configUri = null;
 
-			if (Path.IsPathRooted(_settings.CurrentValue.ConfigPath) && !_settings.CurrentValue.ConfigPath.StartsWith("//", StringComparison.Ordinal))
+			ServerSettings settings = _settings.CurrentValue;
+			if (Path.IsPathRooted(settings.ConfigPath) && !settings.ConfigPath.StartsWith("//", StringComparison.Ordinal))
 			{
 				// absolute path to config
-				configUri = new Uri(_settings.CurrentValue.ConfigPath);
+				configUri = new Uri(settings.ConfigPath);
 			}
-			else if (_settings.CurrentValue.ConfigPath != null)
+			else if (settings.ConfigPath != null)
 			{
 				// relative (development) or perforce path
-				configUri = CombinePaths(new Uri(FileReference.Combine(Program.AppDir, "_").FullName), _settings.CurrentValue.ConfigPath);
+				configUri = CombinePaths(new Uri(FileReference.Combine(Program.AppDir, "_").FullName), settings.ConfigPath);
 			}
 
 			if (configUri != null)
diff --git a/Engine/Source/Programs/Horde/Horde.Build/ServerSettings.cs b/Engine/Source/Programs/Horde/Horde.Build/ServerSettings.cs
--- a/Engine/Source/Programs/Horde/Horde.Build/ServerSettings.cs
+++ b/Engine/Source/Programs/Horde/Horde.Build/ServerSettings.cs
@@ -608,6 +608,11 @@ namespace Horde.Build
 		/// </summary>
 		public int PerforceConnectionPoolSize { get; set; } = 5;
 
+		/// <summary>
+		/// Forces configuration data to be read and updated as part of appplication startup, rather than on a schedule. Useful when running locally.
+		/// </summary>
+		public bool ForceConfigUpdateOnStartup { get; set; }
+
 		/// <summary>
 		/// Lazily computed timezone value
 		/// </summary>
diff --git a/Engine/Source/Programs/Horde/Horde.Build/appsettings.Local.json b/Engine/Source/Programs/Horde/Horde.Build/appsettings.Local.json
--- a/Engine/Source/Programs/Horde/Horde.Build/appsettings.Local.json
+++ b/Engine/Source/Programs/Horde/Horde.Build/appsettings.Local.json
@@ -1,6 +1,7 @@
 {
 	"Horde":
 	{
+		"ForceConfigUpdateOnStartup": true,
 		"Commits":
 		{
 			"ReplicateMetadata": false,
