Horde: More optimizatons to TestRefs and response objects
#rnx

diff --git a/Engine/Source/Programs/Horde/Horde.Build.Tests/TestDataTests.cs b/Engine/Source/Programs/Horde/Horde.Build.Tests/TestDataTests.cs
--- a/Engine/Source/Programs/Horde/Horde.Build.Tests/TestDataTests.cs
+++ b/Engine/Source/Programs/Horde/Horde.Build.Tests/TestDataTests.cs
@@ -207,9 +207,6 @@ namespace Horde.Build.Tests
 			List<GetTestDataRefResponse> refs = refResult.Value;
 
 			Assert.AreEqual(2, refs.Count);
-
-			Assert.AreEqual(1, refs[0].TestDataIds.Count);
-			Assert.AreEqual(1, refs[1].TestDataIds.Count);			
 		}
 
 		[TestMethod]
@@ -275,9 +272,9 @@ namespace Horde.Build.Tests
 
 			Assert.AreEqual(2, refs.Count);
 
-			Assert.AreEqual(3, refs[0].TestDataIds.Count);
-			Assert.AreEqual(3, refs[1].TestDataIds.Count);			
-
+			Assert.AreEqual(1, refs[0].SuiteSkipCount);
+			Assert.AreEqual(1, refs[0].SuiteWarningCount);
+			Assert.AreEqual(1, refs[0].SuiteErrorCount);
 		}
 
 		private readonly string[] _simpleTestDataLines =
diff --git a/Engine/Source/Programs/Horde/Horde.Build/Jobs/TestData/ITestData.cs b/Engine/Source/Programs/Horde/Horde.Build/Jobs/TestData/ITestData.cs
--- a/Engine/Source/Programs/Horde/Horde.Build/Jobs/TestData/ITestData.cs
+++ b/Engine/Source/Programs/Horde/Horde.Build/Jobs/TestData/ITestData.cs
@@ -5,6 +5,7 @@ using System;
 using Horde.Build.Streams;
 using Horde.Build.Utilities;
 using MongoDB.Bson;
+using MongoDB.Bson.Serialization.Attributes;
 
 namespace Horde.Build.Jobs.TestData
 {
@@ -15,6 +16,7 @@ namespace Horde.Build.Jobs.TestData
 	using TestId = ObjectId<ITest>;
 	using TestSuiteId = ObjectId<ITestSuite>;
 	using TestMetaId = ObjectId<ITestMeta>;
+	using TestRefId = ObjectId<ITestDataRef>;
 
 	/// <summary>
 	/// Test outcome
@@ -153,20 +155,18 @@ namespace Horde.Build.Jobs.TestData
 	}
 
 	/// <summary>
-	/// Data ref 
+	/// Data ref with minimal data required for aggregate views
 	/// </summary>
 	public interface ITestDataRef
 	{
+		/// The test ref id
+		TestRefId Id { get; }
+
 		/// <summary>
 		/// The associated stream
 		/// </summary>
 		StreamId StreamId { get; }
 
-		/// <summary>
-		/// The full details test data for this ref
-		/// </summary>
-		IReadOnlyList<ObjectId> TestDataIds { get; }
-
 		/// <summary>
 		/// How long the test ran
 		/// </summary>
@@ -180,7 +180,7 @@ namespace Horde.Build.Jobs.TestData
 		/// <summary>
 		/// The environment the test ran in
 		/// </summary>
-		TestMetaId Metadata { get; }
+		TestMetaId Metadata { get; }	
 
 		/// <summary>
 		/// The ITest in stream
@@ -198,11 +198,39 @@ namespace Horde.Build.Jobs.TestData
 		TestSuiteId? SuiteId { get; }
 
 		/// <summary>
-		/// List for suite test data
+		/// Suite tests skipped
+		/// </summary>
+		int? SuiteSkipCount { get; }
+
+		/// <summary>
+		/// Suite test warnings
+		/// </summary>
+		int? SuiteWarningCount { get; }
+
+		/// <summary>
+		/// Suite test errors
 		/// </summary>
-		IReadOnlyList<ISuiteTestData>? SuiteTests { get; }
+		int? SuiteErrorCount { get; }
+
+	}
+
+	interface ITestDataDetails
+	{
+		/// The corresponding test ref		
+		TestRefId Id { get; }
+
+		/// <summary>
+		/// The full details test data for this ref
+		/// </summary>
+		IReadOnlyList<ObjectId> TestDataIds { get; }
+
+		/// <summary>
+		/// Suite test data
+		/// </summary>		
+		IReadOnlyList<ISuiteTestData>? SuiteTests { get; } 
 	}
 
+
 	/// <summary>
 	/// The tests and suites running in a given stream
 	/// </summary>
diff --git a/Engine/Source/Programs/Horde/Horde.Build/Jobs/TestData/TestDataCollection.cs b/Engine/Source/Programs/Horde/Horde.Build/Jobs/TestData/TestDataCollection.cs
--- a/Engine/Source/Programs/Horde/Horde.Build/Jobs/TestData/TestDataCollection.cs
+++ b/Engine/Source/Programs/Horde/Horde.Build/Jobs/TestData/TestDataCollection.cs
@@ -27,6 +27,7 @@
 	using TestId = ObjectId<ITest>;
 	using TestSuiteId = ObjectId<ITestSuite>;
 	using TestMetaId = ObjectId<ITestMeta>;
+	using TestRefId = ObjectId<ITestDataRef>;
 
 	/// <summary>
 	/// Collection of test data documents
@@ -138,6 +139,37 @@
 			}
 		}
 
+		/// <summary>
+		/// Stores the tests and suites running in a stream
+		/// </summary>
+		class TestStreamDocument : ITestStream
+		{
+			[BsonRequired, BsonId]
+			public ObjectId Id { get; set; }
+
+			[BsonRequired]
+			public StreamId StreamId { get; set; }
+
+			[BsonRequired]
+			public List<TestId> Tests { get; set; } = new List<TestId>();
+			IReadOnlyList<TestId> ITestStream.Tests => Tests;
+
+			[BsonRequired]
+			public List<TestSuiteId> TestSuites { get; set; } = new List<TestSuiteId>();
+			IReadOnlyList<TestSuiteId> ITestStream.TestSuites => TestSuites;
+
+			private TestStreamDocument()
+			{
+
+			}
+
+			public TestStreamDocument(StreamId streamId)
+			{
+				Id = ObjectId.GenerateNewId();
+				StreamId = streamId;
+			}
+		}
+
 		class SuiteTestData : ISuiteTestData
 		{
 			[BsonRequired, BsonElement("tid")]
@@ -164,7 +196,7 @@
 		class TestDataRefDocument : ITestDataRef
 		{
 			[BsonRequired, BsonId]
-			public ObjectId Id { get; set; }
+			public TestRefId Id { get; set; }
 
 			[BsonRequired,BsonElement("sid")]
 			public StreamId StreamId { get; set; }
@@ -175,25 +207,28 @@
 			[BsonRequired, BsonElement("bcl")]
 			public int BuildChangeList { get; set; }
 
-			[BsonRequired, BsonElement("tdid")]
-			public List<ObjectId> TestDataIds { get; set; } = new List<ObjectId>();
-			IReadOnlyList<ObjectId> ITestDataRef.TestDataIds => TestDataIds;
-
 			[BsonRequired, BsonElement("d")]
 			public TimeSpan Duration { get; set; }
 
+			// --- Single tests
 			[BsonIgnoreIfNull, BsonElement("tid")]
 			public TestId? TestId { get; set; }
 
 			[BsonIgnoreIfNull, BsonElement("o")]
 			public TestOutcome? Outcome { get; set; }
-
+			
+			// --- Suite tests			
 			[BsonIgnoreIfNull, BsonElement("suid")]
 			public TestSuiteId? SuiteId { get; set; }
 
-			[BsonIgnoreIfNull, BsonElement("sts")]
-			public List<SuiteTestData>? SuiteTests { get; set; }
-			IReadOnlyList<ISuiteTestData>? ITestDataRef.SuiteTests => SuiteTests;
+			[BsonIgnoreIfNull, BsonElement("ssc")]
+			public int? SuiteSkipCount { get; set; }
 
+			[BsonIgnoreIfNull, BsonElement("swc")]
+			public int? SuiteWarningCount { get; set; }
+
+			[BsonIgnoreIfNull, BsonElement("sec")]
+			public int? SuiteErrorCount { get; set; }
+
 			private TestDataRefDocument()
 			{
@@ -202,41 +237,44 @@
 
 			public TestDataRefDocument(StreamId streamId)
 			{
-				Id = ObjectId.GenerateNewId();
+				Id = TestRefId.GenerateNewId();
 				StreamId = streamId;
 			}
 		}
 
-		/// <summary>
-		/// Stores the tests and suites running in a stream
-		/// </summary>
-		class TestStreamDocument : ITestStream
+		class TestDataDetailsDocument : ITestDataDetails
 		{
+			/// <summary>
+			/// The corresponding test ref for these details
+			/// </summary>
 			[BsonRequired, BsonId]
-			public ObjectId Id { get; set; }
+			public TestRefId Id { get; set; }
 
-			[BsonRequired]
-			public StreamId StreamId { get; set; }
+			/// <summary>
+			/// The full (raw) data for this test, which can be consumed by various clients
+			/// </summary>
+			[BsonRequired, BsonElement("tdid")]
+			public List<ObjectId> TestDataIds { get; set; } = new List<ObjectId>();
+			IReadOnlyList<ObjectId> ITestDataDetails.TestDataIds => TestDataIds;
 
-			[BsonRequired]
-			public List<TestId> Tests { get; set; } = new List<TestId>();
-			IReadOnlyList<TestId> ITestStream.Tests => Tests;
+			/// <summary>
+			/// Suite test data
+			/// </summary>
+			[BsonIgnoreIfNull, BsonElement("sts")]
+			public List<SuiteTestData>? SuiteTests { get; set; }
+			IReadOnlyList<ISuiteTestData>? ITestDataDetails.SuiteTests => SuiteTests;
 
-			[BsonRequired]
-			public List<TestSuiteId> TestSuites { get; set; } = new List<TestSuiteId>();
-			IReadOnlyList<TestSuiteId> ITestStream.TestSuites => TestSuites;
-
-			private TestStreamDocument()
+			private TestDataDetailsDocument()
 			{
 
 			}
 
-			public TestStreamDocument(StreamId streamId)
+			public TestDataDetailsDocument(TestRefId id, List<ObjectId> testDataIds, List<SuiteTestData>? suiteTests = null)
 			{
-				Id = ObjectId.GenerateNewId();
-				StreamId = streamId;
+				Id = id;
+				TestDataIds = testDataIds;
+				SuiteTests = suiteTests;
 			}
-
 		}
 
 		/// <summary>
@@ -278,7 +316,7 @@
 		readonly IMongoCollection<TestDataDocument> _testDataDocuments;
 
 		/// <summary>
-		/// Test platform collection
+		/// Test meta collection
 		/// </summary>
 		readonly IMongoCollection<TestMetaDocument> _testMeta;
 
@@ -293,12 +331,16 @@
 		readonly IMongoCollection<TestSuiteDocument> _testSuites;
 
 		/// <summary>
-		/// Test suite collection
+		/// Test data refs collection
 		/// </summary>
 		readonly IMongoCollection<TestDataRefDocument> _testRefs;
 
 		/// <summary>
-		/// Test suite collection
+		/// Test data refs collection
 		/// </summary>
+		readonly IMongoCollection<TestDataDetailsDocument> _testDetails;
+		/// <summary>
+		/// Test streams collection
+		/// </summary>
 		readonly IMongoCollection<TestStreamDocument> _testStreams;
 
@@ -336,6 +378,8 @@
 			testRefIndexes.Add(keys => keys.Ascending(x => x.StreamId).Ascending(x => x.Metadata).Descending(x => x.BuildChangeList).Ascending(x => x.TestId).Ascending(x => x.SuiteId));
 			_testRefs = mongoService.GetCollection<TestDataRefDocument>("TestData.TestRefsV2", testRefIndexes);
 
+			_testDetails = mongoService.GetCollection<TestDataDetailsDocument>("TestData.TestDetailsV2");
+
 			List<MongoIndex<TestStreamDocument>> streamIndexes = new List<MongoIndex<TestStreamDocument>>();
 			streamIndexes.Add(keys => keys.Ascending(x => x.StreamId), unique: true);
 			_testStreams = mongoService.GetCollection<TestStreamDocument>("TestData.TestStreamsV2", streamIndexes);
@@ -364,13 +408,13 @@
 
 			if (minCreateTime != null)
 			{
-				ObjectId minTime = ObjectId.GenerateNewId(minCreateTime.Value);
+				TestRefId minTime = TestRefId.GenerateNewId(minCreateTime.Value);
 				filter &= filterBuilder.Gte(x => x.Id!, minTime);
 
 			}
 			if (maxCreateTime != null)
 			{
-				ObjectId maxTime = ObjectId.GenerateNewId(maxCreateTime.Value);
+				TestRefId maxTime = TestRefId.GenerateNewId(maxCreateTime.Value);
 				filter &= filterBuilder.Lte(x => x.Id!, maxTime);
 			}
 
@@ -1056,22 +1100,25 @@
 				testRef.BuildChangeList = testData.BuildChangeList <= 0 ? job.Change : testData.BuildChangeList;
 				testRef.Duration = TimeSpan.FromSeconds(testData.TotalDurationSeconds);
 				testRef.TestId = test.Id;
-				testRef.TestDataIds.Add(testDataId);
-				testRef.Outcome = TestOutcome.Unspecified;
+				testRef.Outcome = testData.HasSucceeded ? TestOutcome.Success : TestOutcome.Unspecified;
 
-				switch (testData.TestResult)
+				if (testRef.Outcome == TestOutcome.Unspecified)
 				{
-					case "Passed":
-						testRef.Outcome = TestOutcome.Success;
-						break;
-					case "Failed":
-					case "Cancelled":
-					case "TimedOut":
-						testRef.Outcome = TestOutcome.Failure;
-						break;
+					switch (testData.TestResult)
+					{
+						case "Passed":
+							testRef.Outcome = TestOutcome.Success;
+							break;
+						case "Failed":
+						case "Cancelled":
+						case "TimedOut":
+							testRef.Outcome = TestOutcome.Failure;
+							break;
+					}
 				}
 
 				await AddTestRef(testRef);
+				await _testDetails.InsertOneAsync(new TestDataDetailsDocument(testRef.Id, new List<ObjectId> { testDataId }));
 
 			}
 			catch (Exception ex)
@@ -1186,15 +1233,26 @@
 						testRef.Metadata = metaId.Value;
 						testRef.BuildChangeList = job.Change;
 						testRef.Duration = TimeSpan.FromSeconds(session.TestSessionInfo!.TimeElapseSec);
-						testRef.TestDataIds = documents.Select(x => x.Id).ToList();
 						testRef.Outcome = TestOutcome.Unspecified;
 						testRef.SuiteId = testSuite.Id;
-						testRef.SuiteTests = new List<SuiteTestData>();
 
+						int skipCount = 0;
+						int warningCount = 0;
+						int errorCount = 0;
+
 						// populate the suite tests
-
+						List<SuiteTestData> suiteTestData = new List<SuiteTestData>();
 						foreach (AutomatedTestSessionData.SessionTest test in suiteTests)
 						{
+							if (test.WarningCount > 0)
+							{
+								warningCount++;
+							}
 
+							if (test.ErrorCount > 0)
+							{
+								errorCount++;
+							}
+
 							TestDocument document = suiteTestDocuments[test.Name];
 							TestOutcome outcome = TestOutcome.Unspecified;
@@ -1206,16 +1264,23 @@
 									break;
 								case "Skipped":
 									outcome = TestOutcome.Skipped;
+									skipCount++;
 									break;
 								case "Fail":
 									outcome = TestOutcome.Failure;
 									break;
 							}
 
-							testRef.SuiteTests.Add(new SuiteTestData(document.Id, outcome, TimeSpan.FromSeconds(test.TimeElapseSec), test.TestUID));
+							suiteTestData.Add(new SuiteTestData(document.Id, outcome, TimeSpan.FromSeconds(test.TimeElapseSec), test.TestUID));
 						}
 
+						testRef.SuiteSkipCount = skipCount;
+						testRef.SuiteErrorCount = errorCount;
+						testRef.SuiteWarningCount =warningCount;
+
+						// and add to collections
 						await AddTestRef(testRef);
+						await _testDetails.InsertOneAsync(new TestDataDetailsDocument(testRef.Id, documents.Select(x => x.Id).ToList(), suiteTestData));
 					}
 					else
 					{
diff --git a/Engine/Source/Programs/Horde/Horde.Build/Jobs/TestData/TestDataMessages.cs b/Engine/Source/Programs/Horde/Horde.Build/Jobs/TestData/TestDataMessages.cs
--- a/Engine/Source/Programs/Horde/Horde.Build/Jobs/TestData/TestDataMessages.cs
+++ b/Engine/Source/Programs/Horde/Horde.Build/Jobs/TestData/TestDataMessages.cs
@@ -312,14 +312,14 @@ namespace Horde.Build.Jobs.TestData
 	public class GetTestDataRefResponse
 	{
 		/// <summary>
-		/// The associated stream
+		/// The test ref id
 		/// </summary>
-		public string StreamId { get; set; }
+		public string Id { get; set; }
 
 		/// <summary>
-		/// The full details test data for this ref
+		/// The associated stream
 		/// </summary>
-		public List<string> TestDataIds { get; set; }
+		public string StreamId { get; set; }
 
 		/// <summary>
 		/// How long the test ran
@@ -352,9 +352,19 @@ namespace Horde.Build.Jobs.TestData
 		public string? SuiteId { get; set; }
 
 		/// <summary>
-		/// The suite tests
+		/// Suite tests skipped
+		/// </summary>
+		public int? SuiteSkipCount { get; set; }
+
+		/// <summary>
+		/// Suite test warnings
+		/// </summary>
+		public int? SuiteWarningCount { get; set; }
+
+		/// <summary>
+		/// Suite test errors
 		/// </summary>
-		public List<GetSuiteTestDataResponse>? SuiteTests { get; set; }
+		public int? SuiteErrorCount { get; set; }
 
 		/// <summary>
 		/// Constructor
@@ -362,22 +372,17 @@ namespace Horde.Build.Jobs.TestData
 		/// <param name="testData"></param>
 		public GetTestDataRefResponse(ITestDataRef testData)
 		{
+			Id = testData.Id.ToString();
 			StreamId = testData.StreamId.ToString();
-			TestDataIds = testData.TestDataIds.Select(x => x.ToString()).ToList();
 			Duration = testData.Duration;
 			BuildChangeList = testData.BuildChangeList;
 			MetaId = testData.Metadata.ToString();
 			TestId = testData.TestId?.ToString();
 			Outcome = testData.Outcome;
 			SuiteId = testData.SuiteId?.ToString();
-			if (testData.SuiteTests != null)
-			{
-				SuiteTests = new List<GetSuiteTestDataResponse>();
-				foreach (ISuiteTestData test in testData.SuiteTests)
-				{
-					SuiteTests.Add(new GetSuiteTestDataResponse(test));
-				}
-			}			
+			SuiteSkipCount = testData.SuiteSkipCount;
+			SuiteWarningCount = testData.SuiteWarningCount;
+			SuiteErrorCount = testData.SuiteErrorCount;
 		}
 	}
 }
diff --git a/Engine/Source/Programs/Horde/Horde.Build/Utilities/ObjectId.cs b/Engine/Source/Programs/Horde/Horde.Build/Utilities/ObjectId.cs
--- a/Engine/Source/Programs/Horde/Horde.Build/Utilities/ObjectId.cs
+++ b/Engine/Source/Programs/Horde/Horde.Build/Utilities/ObjectId.cs
@@ -51,6 +51,10 @@ namespace Horde.Build.Utilities
 		[SuppressMessage("Design", "CA1000:Do not declare static members on generic types", Justification = "<Pending>")]
 		public static ObjectId<T> GenerateNewId() => new ObjectId<T>(ObjectId.GenerateNewId());
 
+		/// <inheritdoc cref="ObjectId.GenerateNewId(DateTime)"/>
+		[SuppressMessage("Design", "CA1000:Do not declare static members on generic types", Justification = "<Pending>")]
+		public static ObjectId<T> GenerateNewId(DateTime timestamp) => new ObjectId<T>(ObjectId.GenerateNewId(timestamp));
+
 		/// <inheritdoc cref="ObjectId.Parse(String)"/>
 		[SuppressMessage("Design", "CA1000:Do not declare static members on generic types", Justification = "<Pending>")]
 		public static ObjectId<T> Parse(string str) => new ObjectId<T>(ObjectId.Parse(str));
