Horde: Set the environment to "Local" when launching through Visual Studio.

diff --git a/Engine/Source/Programs/Horde/Horde.Build/Properties/launchSettings.json b/Engine/Source/Programs/Horde/Horde.Build/Properties/launchSettings.json
--- a/Engine/Source/Programs/Horde/Horde.Build/Properties/launchSettings.json
+++ b/Engine/Source/Programs/Horde/Horde.Build/Properties/launchSettings.json
@@ -2,7 +2,11 @@
   "profiles": {
     "HordeServer": {
       "commandName": "Project",
-      "commandLineArgs": ""
+      "commandLineArgs": "",
+      "environmentVariables": {
+        "ASPNETCORE_ENVIRONMENT": "Local"
+      },
+      "hotReloadEnabled": false
     }
   }
 }
\ No newline at end of file
