Horde: Allow configuring pools through global and project config files, rather than through the dashboard / rest API. Add some pool definitions for different platforms to the default config file.

diff --git a/Engine/Source/Programs/Horde/Horde.Build/Agents/Fleet/FleetService.cs b/Engine/Source/Programs/Horde/Horde.Build/Agents/Fleet/FleetService.cs
--- a/Engine/Source/Programs/Horde/Horde.Build/Agents/Fleet/FleetService.cs
+++ b/Engine/Source/Programs/Horde/Horde.Build/Agents/Fleet/FleetService.cs
@@ -360,7 +360,7 @@ namespace Horde.Build.Agents.Fleet
 			}
 
 			// These is the legacy way of creating and configuring strategies (list-based approach above is preferred)
-			switch (pool.SizeStrategy)
+			switch (pool.SizeStrategy ?? _settings.Value.DefaultAgentPoolSizeStrategy)
 			{
 				case PoolSizeStrategy.JobQueue:
 					return new JobQueueStrategy(_jobCollection, _graphCollection, _streamService, _clock, _cache, pool.JobQueueSettings);
diff --git a/Engine/Source/Programs/Horde/Horde.Build/Agents/Pools/IPool.cs b/Engine/Source/Programs/Horde/Horde.Build/Agents/Pools/IPool.cs
--- a/Engine/Source/Programs/Horde/Horde.Build/Agents/Pools/IPool.cs
+++ b/Engine/Source/Programs/Horde/Horde.Build/Agents/Pools/IPool.cs
@@ -98,112 +98,21 @@ namespace Horde.Build.Agents.Pools
 			Config = config;
 		}
 	}
-	
+
 	/// <summary>
-	/// A pool of machines
+	/// Pool of machines
 	/// </summary>
-	public interface IPool
+	public interface IPool : IPoolConfig
 	{
-		/// <summary>
-		/// Unique id for this pool
-		/// </summary>
-		public PoolId Id { get; }
-
-		/// <summary>
-		/// Name of the pool
-		/// </summary>
-		public string Name { get; }
-
-		/// <summary>
-		/// Condition for agents to automatically be included in this pool
-		/// </summary>
-		public Condition? Condition { get; }
-
 		/// <summary>
 		/// List of workspaces currently assigned to this pool
 		/// </summary>
 		public IReadOnlyList<AgentWorkspace> Workspaces { get; }
 
 		/// <summary>
-		/// Whether to sync autosdk alongside the workspaces for this pool
-		/// </summary>
-		public bool UseAutoSdk { get; }
-
-		/// <summary>
-		/// Arbitrary properties related to this pool
-		/// </summary>
-		public IReadOnlyDictionary<string, string> Properties { get; }
-
-		/// <summary>
-		/// Whether to enable autoscaling for this pool
-		/// </summary>
-		public bool EnableAutoscaling { get; }
-
-		/// <summary>
-		/// The minimum number of agents to keep in the pool
-		/// </summary>
-		public int? MinAgents { get; }
-
-		/// <summary>
-		/// The minimum number of idle agents to hold in reserve
-		/// </summary>
-		public int? NumReserveAgents { get; }
-
-		/// <summary>
-		/// Interval between conforms. If zero, the pool will not conform on a schedule.
-		/// </summary>
-		public TimeSpan? ConformInterval { get; }
-
-		/// <summary>
-		/// Last time the pool was (auto) scaled up
-		/// </summary>
-		public DateTime? LastScaleUpTime { get; }
-		
-		/// <summary>
-		/// Last time the pool was (auto) scaled down
-		/// </summary>
-		public DateTime? LastScaleDownTime { get; }
-		
-		/// <summary>
-		/// Cooldown time between scale-out events
-		/// </summary>
-		public TimeSpan? ScaleOutCooldown { get; }
-		
-		/// <summary>
-		/// Cooldown time between scale-in events
-		/// </summary>
-		public TimeSpan? ScaleInCooldown { get; }
-		
-		/// <summary>
-		/// Pool sizing strategy to be used for this pool
-		/// </summary>
-		public PoolSizeStrategy SizeStrategy { get; }
-
-		/// <summary>
-		/// List of pool sizing strategies for this pool. The first strategy with a matching condition will be picked.
-		/// </summary>
-		public IReadOnlyList<PoolSizeStrategyInfo> SizeStrategies { get; }
-
-		/// <summary>
-		/// List of fleet managers for this pool. The first strategy with a matching condition will be picked.
-		/// If empty or no conditions match, a default fleet manager will be used.
-		/// </summary>
-		public IReadOnlyList<FleetManagerInfo> FleetManagers { get; }
-		
-		/// <summary>
-		/// Settings for lease utilization pool sizing strategy (if used)
-		/// </summary>
-		public LeaseUtilizationSettings? LeaseUtilizationSettings { get; }
-		
-		/// <summary>
-		/// Settings for job queue pool sizing strategy (if used)
-		/// </summary>
-		public JobQueueSettings? JobQueueSettings { get; }
-		
-		/// <summary>
-		/// Settings for job queue pool sizing strategy (if used)
+		/// Revision string for the config file describing this pool, or null if added manually.
 		/// </summary>
-		public ComputeQueueAwsMetricSettings? ComputeQueueAwsMetricSettings { get; }
+		public string? Revision { get; }
 
 		/// <summary>
 		/// Update index for this document
diff --git a/Engine/Source/Programs/Horde/Horde.Build/Agents/Pools/IPoolCollection.cs b/Engine/Source/Programs/Horde/Horde.Build/Agents/Pools/IPoolCollection.cs
--- a/Engine/Source/Programs/Horde/Horde.Build/Agents/Pools/IPoolCollection.cs
+++ b/Engine/Source/Programs/Horde/Horde.Build/Agents/Pools/IPoolCollection.cs
@@ -3,6 +3,7 @@
 using System;
 using System.Collections.Generic;
 using System.Threading.Tasks;
+using Amazon.Util.Internal;
 using EpicGames.Horde.Common;
 using Horde.Build.Agents.Fleet;
 using Horde.Build.Utilities;
@@ -11,6 +12,38 @@ namespace Horde.Build.Agents.Pools
 {
 	using PoolId = StringId<IPool>;
 
+	/// <summary>
+	/// Exception when conflicting definitions for a pool are encountered
+	/// </summary>
+	public sealed class PoolConflictException : Exception
+	{
+		/// <summary>
+		/// Pool with conflicting definitions
+		/// </summary>
+		public PoolId PoolId { get; }
+
+		/// <summary>
+		/// The first revision string
+		/// </summary>
+		public string PrevRevision { get; }
+
+		/// <summary>
+		/// The second revision string
+		/// </summary>
+		public string NextRevision { get; }
+
+		/// <summary>
+		/// Constructor
+		/// </summary>
+		public PoolConflictException(PoolId poolId, string prevRevision, string nextRevision)
+			: base($"Duplicate definitions for {poolId} - first in {prevRevision}, now {nextRevision}")
+		{
+			PoolId = poolId;
+			PrevRevision = prevRevision;
+			NextRevision = nextRevision;
+		}
+	}
+
 	/// <summary>
 	/// Collection of pool documents
 	/// </summary>
@@ -127,5 +160,11 @@ namespace Horde.Build.Agents.Pools
 			JobQueueSettings? jobQueueSettings = null,
 			ComputeQueueAwsMetricSettings? computeQueueAwsMetricSettings = null,
 			bool? useDefaultStrategy = null);
+
+		/// <summary>
+		/// Updates the list of pools from a config file
+		/// </summary>
+		/// <param name="poolConfigs">Configuration for the pools, and revision string for the config file containing them</param>
+		Task ConfigureAsync(IReadOnlyList<(PoolConfig Config, string Revision)> poolConfigs);
 	}
 }
diff --git a/Engine/Source/Programs/Horde/Horde.Build/Agents/Pools/PoolCollection.cs b/Engine/Source/Programs/Horde/Horde.Build/Agents/Pools/PoolCollection.cs
--- a/Engine/Source/Programs/Horde/Horde.Build/Agents/Pools/PoolCollection.cs
+++ b/Engine/Source/Programs/Horde/Horde.Build/Agents/Pools/PoolCollection.cs
@@ -24,49 +24,84 @@ namespace Horde.Build.Agents.Pools
 	{
 		class PoolDocument : IPool
 		{
-			// Properties
 			[BsonId]
 			public PoolId Id { get; set; }
+
 			[BsonRequired]
 			public string Name { get; set; } = null!;
+
 			[BsonIgnoreIfNull]
 			public Condition? Condition { get; set; }
+
 			public List<AgentWorkspace> Workspaces { get; set; } = new List<AgentWorkspace>();
+
+			[BsonIgnoreIfDefault, BsonDefaultValue(true)]
 			public bool UseAutoSdk { get; set; } = true;
+
 			public Dictionary<string, string> Properties { get; set; } = new Dictionary<string, string>();
+
+			[BsonIgnoreIfDefault, BsonDefaultValue(true)]
 			public bool EnableAutoscaling { get; set; } = true;
+
+			[BsonIgnoreIfNull]
 			public int? MinAgents { get; set; }
+
+			[BsonIgnoreIfNull]
 			public int? NumReserveAgents { get; set; }
+
+			[BsonIgnoreIfNull]
 			public TimeSpan? ConformInterval { get; set; }
+
+			[BsonIgnoreIfNull]
 			public DateTime? LastScaleUpTime { get; set; }
+
+			[BsonIgnoreIfNull]
 			public DateTime? LastScaleDownTime { get; set; }
+
+			[BsonIgnoreIfNull]
 			public TimeSpan? ScaleOutCooldown { get; set; }
+
+			[BsonIgnoreIfNull]
 			public TimeSpan? ScaleInCooldown { get; set; }
+
+			[BsonIgnoreIfNull]
 			public PoolSizeStrategy? SizeStrategy { get; set; }
-			PoolSizeStrategy IPool.SizeStrategy => SizeStrategy ?? s_defaultStrategy!.Value;
-			public List<PoolSizeStrategyInfo> SizeStrategies { get; set; } = new ();
-			public List<FleetManagerInfo> FleetManagers { get; set; } = new ();
+
+			[BsonIgnoreIfNull]
+			public List<PoolSizeStrategyInfo>? SizeStrategies { get; set; }
+
+			[BsonIgnoreIfNull]
+			public List<FleetManagerInfo>? FleetManagers { get; set; }
+
+			[BsonIgnoreIfNull]
 			public LeaseUtilizationSettings? LeaseUtilizationSettings { get; set; }
+
+			[BsonIgnoreIfNull]
 			public JobQueueSettings? JobQueueSettings { get; set; }
+
+			[BsonIgnoreIfNull]
 			public ComputeQueueAwsMetricSettings? ComputeQueueAwsMetricSettings { get; set; }
+
+			[BsonIgnoreIfNull]
+			public string? Revision { get; set; }
+
 			public int UpdateIndex { get; set; }
 
 			// Read-only wrappers
 			IReadOnlyList<AgentWorkspace> IPool.Workspaces => Workspaces;
-			IReadOnlyDictionary<string, string> IPool.Properties => Properties;
-			IReadOnlyList<PoolSizeStrategyInfo> IPool.SizeStrategies => SizeStrategies;
-			IReadOnlyList<FleetManagerInfo> IPool.FleetManagers => FleetManagers;
+			IReadOnlyDictionary<string, string> IPoolConfig.Properties => Properties;
+			IReadOnlyList<PoolSizeStrategyInfo> IPoolConfig.SizeStrategies => (IReadOnlyList<PoolSizeStrategyInfo>?)SizeStrategies ?? Array.Empty<PoolSizeStrategyInfo>();
+			IReadOnlyList<FleetManagerInfo> IPoolConfig.FleetManagers => (IReadOnlyList<FleetManagerInfo>?)FleetManagers ?? Array.Empty<FleetManagerInfo>();
 
 			public PoolDocument()
 			{
 			}
 
-			public PoolDocument(IPool other)
+			public PoolDocument(IPoolConfig other, string? otherRevision)
 			{
 				Id = other.Id;
-				Name = other.Name;
+				Name = String.IsNullOrEmpty(other.Name)? other.Id.ToString() : other.Name;
 				Condition = other.Condition;
-				Workspaces.AddRange(other.Workspaces);
 				UseAutoSdk = other.UseAutoSdk;
 				Properties = new Dictionary<string, string>(other.Properties);
 				EnableAutoscaling = other.EnableAutoscaling;
@@ -78,11 +113,24 @@ namespace Horde.Build.Agents.Pools
 				ScaleOutCooldown = other.ScaleOutCooldown;
 				ScaleInCooldown = other.ScaleInCooldown;
 				SizeStrategy = other.SizeStrategy;
-				SizeStrategies.AddRange(other.SizeStrategies);
-				FleetManagers.AddRange(other.FleetManagers);
+				if (other.SizeStrategies != null && other.SizeStrategies.Count > 0)
+				{
+					SizeStrategies = new List<PoolSizeStrategyInfo>(other.SizeStrategies);
+				}
+				if (other.FleetManagers != null && other.FleetManagers.Count > 0)
+				{
+					FleetManagers = new List<FleetManagerInfo>(other.FleetManagers);
+				}
 				LeaseUtilizationSettings = other.LeaseUtilizationSettings;
 				JobQueueSettings = other.JobQueueSettings;
 				ComputeQueueAwsMetricSettings = other.ComputeQueueAwsMetricSettings;
+				Revision = otherRevision;
+			}
+
+			public PoolDocument(IPool other)
+				: this(other, other.Revision)
+			{
+				Workspaces.AddRange(other.Workspaces);
 				UpdateIndex = other.UpdateIndex;
 			}
 		}
@@ -92,17 +140,13 @@ namespace Horde.Build.Agents.Pools
 		/// </summary>
 		readonly IMongoCollection<PoolDocument> _pools;
 
-		static PoolSizeStrategy? s_defaultStrategy;
-
 		/// <summary>
 		/// Constructor
 		/// </summary>
 		/// <param name="mongoService">The database service instance</param>
-		/// <param name="serverSettings">The server settings</param>
-		public PoolCollection(MongoService mongoService, IOptions<ServerSettings> serverSettings)
+		public PoolCollection(MongoService mongoService)
 		{
 			_pools = mongoService.GetCollection<PoolDocument>("Pools");
-			s_defaultStrategy = serverSettings.Value.DefaultAgentPoolSizeStrategy;
 		}
 
 		/// <inheritdoc/>
@@ -364,5 +408,67 @@ namespace Horde.Build.Agents.Pools
 			}
 			return TryUpdateAsync(pool, transaction);
 		}
+
+		/// <inheritdoc/>
+		public async Task ConfigureAsync(IReadOnlyList<(PoolConfig Config, string Revision)> poolConfigs)
+		{
+			List<PoolDocument> pools = await _pools.Find(FilterDefinition<PoolDocument>.Empty).ToListAsync();
+			Dictionary<PoolId, PoolDocument> idToPool = pools.ToDictionary(x => x.Id, x => x);
+
+			// Make sure we don't have any duplicates in the new pool configs before setting anything
+			Dictionary<PoolId, string> poolIdToRevision = new Dictionary<PoolId, string>();
+			foreach ((PoolConfig poolConfig, string revision) in poolConfigs)
+			{
+				if (poolIdToRevision.TryGetValue(poolConfig.Id, out string? prevRevision))
+				{
+					throw new PoolConflictException(poolConfig.Id, prevRevision, revision);
+				}
+				else
+				{
+					poolIdToRevision.Add(poolConfig.Id, revision);
+				}
+			}
+
+			// Update all the pools which are still present
+			foreach ((PoolConfig poolConfig, string revision) in poolConfigs)
+			{
+				idToPool.TryGetValue(poolConfig.Id, out PoolDocument? currentPool);
+				while (currentPool == null || currentPool.Revision != revision)
+				{
+					int updateIndex = currentPool?.UpdateIndex ?? 0;
+
+					PoolDocument newPool = new PoolDocument(poolConfig, revision);
+					if (currentPool != null)
+					{
+						newPool.Workspaces.AddRange(currentPool.Workspaces);
+						newPool.UpdateIndex = currentPool.UpdateIndex;
+					}
+
+					ReplaceOneResult result = await _pools.ReplaceOneAsync(x => x.Id == poolConfig.Id && x.UpdateIndex == updateIndex, newPool, new ReplaceOptions { IsUpsert = true });
+					if (result.MatchedCount > 0)
+					{
+						break;
+					}
+
+					currentPool = await _pools.Find(x => x.Id == poolConfig.Id).FirstOrDefaultAsync();
+				}
+				idToPool.Remove(poolConfig.Id);
+			}
+
+			// Remove anything that has been deleted
+			foreach (PoolDocument remainingPool in idToPool.Values)
+			{
+				PoolDocument? remainingPoolCopy = remainingPool;
+				while (remainingPoolCopy != null && !String.IsNullOrEmpty(remainingPoolCopy.Revision))
+				{
+					DeleteResult result = await _pools.DeleteOneAsync(x => x.Id == remainingPool.Id && x.UpdateIndex == remainingPoolCopy.UpdateIndex);
+					if (result.DeletedCount > 0)
+					{
+						break;
+					}
+					remainingPoolCopy = await _pools.Find(x => x.Id == remainingPool.Id).FirstOrDefaultAsync();
+				}
+			}
+		}
 	}
 }
diff --git a/Engine/Source/Programs/Horde/Horde.Build/Agents/Pools/PoolConfig.cs b/Engine/Source/Programs/Horde/Horde.Build/Agents/Pools/PoolConfig.cs
new file mode 100644
--- /dev/null
+++ b/Engine/Source/Programs/Horde/Horde.Build/Agents/Pools/PoolConfig.cs
@@ -0,0 +1,186 @@
+// Copyright Epic Games, Inc. All Rights Reserved.
+
+using System;
+using System.Collections.Generic;
+using EpicGames.Horde.Common;
+using Horde.Build.Agents.Fleet;
+using Horde.Build.Utilities;
+
+namespace Horde.Build.Agents.Pools
+{
+	using PoolId = StringId<IPool>;
+
+	/// <summary>
+	/// Configuration for a pool of machines
+	/// </summary>
+	public interface IPoolConfig
+	{
+		/// <summary>
+		/// Unique id for this pool
+		/// </summary>
+		public PoolId Id { get; }
+
+		/// <summary>
+		/// Name of the pool
+		/// </summary>
+		public string Name { get; }
+
+		/// <summary>
+		/// Condition for agents to automatically be included in this pool
+		/// </summary>
+		public Condition? Condition { get; }
+
+		/// <summary>
+		/// Whether to sync autosdk alongside the workspaces for this pool
+		/// </summary>
+		public bool UseAutoSdk { get; }
+
+		/// <summary>
+		/// Arbitrary properties related to this pool
+		/// </summary>
+		public IReadOnlyDictionary<string, string> Properties { get; }
+
+		/// <summary>
+		/// Whether to enable autoscaling for this pool
+		/// </summary>
+		public bool EnableAutoscaling { get; }
+
+		/// <summary>
+		/// The minimum number of agents to keep in the pool
+		/// </summary>
+		public int? MinAgents { get; }
+
+		/// <summary>
+		/// The minimum number of idle agents to hold in reserve
+		/// </summary>
+		public int? NumReserveAgents { get; }
+
+		/// <summary>
+		/// Interval between conforms. If zero, the pool will not conform on a schedule.
+		/// </summary>
+		public TimeSpan? ConformInterval { get; }
+
+		/// <summary>
+		/// Last time the pool was (auto) scaled up
+		/// </summary>
+		public DateTime? LastScaleUpTime { get; }
+
+		/// <summary>
+		/// Last time the pool was (auto) scaled down
+		/// </summary>
+		public DateTime? LastScaleDownTime { get; }
+
+		/// <summary>
+		/// Cooldown time between scale-out events
+		/// </summary>
+		public TimeSpan? ScaleOutCooldown { get; }
+
+		/// <summary>
+		/// Cooldown time between scale-in events
+		/// </summary>
+		public TimeSpan? ScaleInCooldown { get; }
+
+		/// <summary>
+		/// Pool sizing strategy to be used for this pool
+		/// </summary>
+		public PoolSizeStrategy? SizeStrategy { get; }
+
+		/// <summary>
+		/// List of pool sizing strategies for this pool. The first strategy with a matching condition will be picked.
+		/// </summary>
+		public IReadOnlyList<PoolSizeStrategyInfo> SizeStrategies { get; }
+
+		/// <summary>
+		/// List of fleet managers for this pool. The first strategy with a matching condition will be picked.
+		/// If empty or no conditions match, a default fleet manager will be used.
+		/// </summary>
+		public IReadOnlyList<FleetManagerInfo> FleetManagers { get; }
+
+		/// <summary>
+		/// Settings for lease utilization pool sizing strategy (if used)
+		/// </summary>
+		public LeaseUtilizationSettings? LeaseUtilizationSettings { get; }
+
+		/// <summary>
+		/// Settings for job queue pool sizing strategy (if used)
+		/// </summary>
+		public JobQueueSettings? JobQueueSettings { get; }
+
+		/// <summary>
+		/// Settings for job queue pool sizing strategy (if used)
+		/// </summary>
+		public ComputeQueueAwsMetricSettings? ComputeQueueAwsMetricSettings { get; }
+	}
+
+	/// <summary>
+	/// Mutable configuration for a pool
+	/// </summary>
+	public class PoolConfig : IPoolConfig
+	{
+		/// <inheritdoc/>
+		public PoolId Id { get; set; }
+
+		/// <inheritdoc/>
+		public string Name { get; set; } = String.Empty;
+
+		/// <inheritdoc/>
+		public Condition? Condition { get; set; }
+
+		/// <inheritdoc/>
+		public bool UseAutoSdk { get; set; } = true;
+
+		/// <inheritdoc cref="IPoolConfig.Properties"/>
+		public Dictionary<string, string> Properties { get; set; } = new Dictionary<string, string>();
+
+		/// <inheritdoc/>
+		IReadOnlyDictionary<string, string> IPoolConfig.Properties => Properties;
+
+		/// <inheritdoc/>
+		public bool EnableAutoscaling { get; set; } = true;
+
+		/// <inheritdoc/>
+		public int? MinAgents { get; set; }
+
+		/// <inheritdoc/>
+		public int? NumReserveAgents { get; set; }
+
+		/// <inheritdoc/>
+		public TimeSpan? ConformInterval { get; set; }
+
+		/// <inheritdoc/>
+		public DateTime? LastScaleUpTime { get; set; }
+
+		/// <inheritdoc/>
+		public DateTime? LastScaleDownTime { get; set; }
+
+		/// <inheritdoc/>
+		public TimeSpan? ScaleOutCooldown { get; set; }
+
+		/// <inheritdoc/>
+		public TimeSpan? ScaleInCooldown { get; set; }
+
+		/// <inheritdoc/>
+		public PoolSizeStrategy? SizeStrategy { get; set; }
+
+		/// <inheritdoc/>
+		public List<PoolSizeStrategyInfo> SizeStrategies { get; } = new List<PoolSizeStrategyInfo>();
+
+		/// <inheritdoc cref="IPoolConfig.SizeStrategies"/>
+		IReadOnlyList<PoolSizeStrategyInfo> IPoolConfig.SizeStrategies => SizeStrategies;
+
+		/// <inheritdoc/>
+		public List<FleetManagerInfo> FleetManagers { get; } = new List<FleetManagerInfo>();
+
+		/// <inheritdoc cref="IPoolConfig.FleetManagers"/>
+		IReadOnlyList<FleetManagerInfo> IPoolConfig.FleetManagers => FleetManagers;
+
+		/// <inheritdoc/>
+		public LeaseUtilizationSettings? LeaseUtilizationSettings { get; set; }
+
+		/// <inheritdoc/>
+		public JobQueueSettings? JobQueueSettings { get; set; }
+
+		/// <inheritdoc/>
+		public ComputeQueueAwsMetricSettings? ComputeQueueAwsMetricSettings { get; set; }
+	}
+}
diff --git a/Engine/Source/Programs/Horde/Horde.Build/Configuration/ConfigUpdateService.cs b/Engine/Source/Programs/Horde/Horde.Build/Configuration/ConfigUpdateService.cs
--- a/Engine/Source/Programs/Horde/Horde.Build/Configuration/ConfigUpdateService.cs
+++ b/Engine/Source/Programs/Horde/Horde.Build/Configuration/ConfigUpdateService.cs
@@ -10,6 +10,7 @@ using EpicGames.Core;
 using EpicGames.Perforce;
 using Horde.Build.Acls;
 using Horde.Build.Agents;
+using Horde.Build.Agents.Pools;
 using Horde.Build.Notifications;
 using Horde.Build.Perforce;
 using Horde.Build.Projects;
@@ -44,6 +45,7 @@ namespace Horde.Build.Configuration
 
 		readonly GlobalsService _globalsService;
 		readonly ConfigCollection _configCollection;
+		readonly IPoolCollection _poolCollection;
 		readonly ToolCollection _toolCollection;
 		readonly ProjectService _projectService;
 		readonly StreamService _streamService;
@@ -57,11 +59,12 @@ namespace Horde.Build.Configuration
 		/// <summary>
 		/// Constructor
 		/// </summary>
-		public ConfigUpdateService(GlobalsService globalsService, MongoService mongoService, ConfigCollection configCollection, IPerforceService perforceService, ToolCollection toolCollection, ProjectService projectService, StreamService streamService, INotificationService notificationService, AgentService agentService, IClock clock, IOptionsMonitor<ServerSettings> settings, ILogger<ConfigUpdateService> logger)
+		public ConfigUpdateService(GlobalsService globalsService, MongoService mongoService, ConfigCollection configCollection, IPerforceService perforceService, IPoolCollection poolCollection, ToolCollection toolCollection, ProjectService projectService, StreamService streamService, INotificationService notificationService, AgentService agentService, IClock clock, IOptionsMonitor<ServerSettings> settings, ILogger<ConfigUpdateService> logger)
 		{
 			_globalsService = globalsService;
 			_configCollection = configCollection;
 			_perforceService = perforceService;
+			_poolCollection = poolCollection;
 			_toolCollection = toolCollection;
 			_projectService = projectService;
 			_streamService = streamService;
@@ -317,6 +320,19 @@ namespace Horde.Build.Configuration
 				_logger.LogInformation("Removing stream {StreamId}", removeStreamId);
 				await _streamService.DeleteStreamAsync(removeStreamId);
 			}
+
+			// Configure all the pools
+			List<(PoolConfig, string)> poolConfigs = new();
+			poolConfigs.AddRange(globalConfig.Pools.Select(x => (x, _cachedGlobalConfigRevision!)));
+			poolConfigs.AddRange(_cachedProjectConfigs.Values.SelectMany(x => x.Config.Pools.Select(y => (y, x.Revision))));
+			try
+			{
+				await _poolCollection.ConfigureAsync(poolConfigs);
+			}
+			catch (PoolConflictException ex)
+			{
+				await SendFailureNotificationAsync(ex, configPath, cancellationToken);
+			}
 		}
 
 		static readonly FileExtensionContentTypeProvider s_contentTypeProvider = new FileExtensionContentTypeProvider();
diff --git a/Engine/Source/Programs/Horde/Horde.Build/Defaults/globals.json b/Engine/Source/Programs/Horde/Horde.Build/Defaults/globals.json
--- a/Engine/Source/Programs/Horde/Horde.Build/Defaults/globals.json
+++ b/Engine/Source/Programs/Horde/Horde.Build/Defaults/globals.json
@@ -17,5 +17,27 @@
 				"id": "default"
 			}
 		]
-	}
+	},
+	"pools": [
+		{
+			"id": "any",
+			"name": "Any",
+			"condition": "true"
+		},
+		{
+			"id": "win64",
+			"name": "Win64",
+			"condition": "Platform == 'Win64'"
+		},
+		{
+			"id": "mac",
+			"name": "Mac",
+			"condition": "Platform == 'Mac'"
+		},
+		{
+			"id": "linux",
+			"name": "Linux",
+			"condition": "Platform == 'Linux'"
+		}
+	]
 }
diff --git a/Engine/Source/Programs/Horde/Horde.Build/Projects/ProjectConfig.cs b/Engine/Source/Programs/Horde/Horde.Build/Projects/ProjectConfig.cs
--- a/Engine/Source/Programs/Horde/Horde.Build/Projects/ProjectConfig.cs
+++ b/Engine/Source/Programs/Horde/Horde.Build/Projects/ProjectConfig.cs
@@ -6,6 +6,7 @@ using System.ComponentModel.DataAnnotations;
 using System.Text.RegularExpressions;
 using EpicGames.Core;
 using Horde.Build.Acls;
+using Horde.Build.Agents.Pools;
 using Horde.Build.Streams;
 using Horde.Build.Utilities;
 
@@ -31,6 +32,11 @@ namespace Horde.Build.Projects
 		/// </summary>
 		public string? Logo { get; set; }
 
+		/// <summary>
+		/// List of pools for this project
+		/// </summary>
+		public List<PoolConfig> Pools { get; set; } = new List<PoolConfig>();
+
 		/// <summary>
 		/// Categories to include in this project
 		/// </summary>
diff --git a/Engine/Source/Programs/Horde/Horde.Build/Server/GlobalConfig.cs b/Engine/Source/Programs/Horde/Horde.Build/Server/GlobalConfig.cs
--- a/Engine/Source/Programs/Horde/Horde.Build/Server/GlobalConfig.cs
+++ b/Engine/Source/Programs/Horde/Horde.Build/Server/GlobalConfig.cs
@@ -9,6 +9,7 @@ using EpicGames.Horde.Common;
 using EpicGames.Perforce;
 using EpicGames.Serialization;
 using Horde.Build.Acls;
+using Horde.Build.Agents.Pools;
 using Horde.Build.Agents.Software;
 using Horde.Build.Projects;
 using Horde.Build.Storage;
@@ -32,6 +33,11 @@ namespace Horde.Build.Server
 		/// </summary>
 		public List<ProjectConfigRef> Projects { get; set; } = new List<ProjectConfigRef>();
 
+		/// <summary>
+		/// List of pools
+		/// </summary>
+		public List<PoolConfig> Pools { get; set; } = new List<PoolConfig>();
+
 		/// <summary>
 		/// List of scheduled downtime
 		/// </summary>
