Horde: Use pre-built binary of Redis inside Docker image during build

Redis is not packaged in final image, only used for testing to provide a real-world implementation of Redis.

diff --git a/Engine/Source/Programs/Horde/Horde.Build/Dockerfile b/Engine/Source/Programs/Horde/Horde.Build/Dockerfile
--- a/Engine/Source/Programs/Horde/Horde.Build/Dockerfile
+++ b/Engine/Source/Programs/Horde/Horde.Build/Dockerfile
@@ -7,19 +7,15 @@
 #       Running directly from the source tree will cause a lot of unnecessary
 #       files to be uploaded to Docker before building.
 
+FROM redis:6.2.7 AS redis
+
 # Build image
 FROM mcr.microsoft.com/dotnet/sdk:6.0 AS build-env
 WORKDIR /app
 COPY Source/Programs/Horde/nuget.config ./Source/Programs/Horde/nuget.config
-RUN apt-get update && apt-get install -y build-essential && apt-get install -y curl
+RUN DEBIAN_FRONTEND=noninteractive apt-get -qq update && apt-get -qq install curl
 
-# Pre-built binaries for Redis are hard to come by. But it's small and fast enough to build it here.
-RUN cd /tmp && \
-	wget --no-verbose http://download.redis.io/releases/redis-5.0.9.tar.gz && \
-	tar xzf redis-5.0.9.tar.gz && \
-	cd redis-5.0.9 && \
-	make && \
-	mv src/redis-server /usr/local/bin/
+COPY --from=redis /usr/local/bin/redis-server /usr/local/bin/redis-server
 
 # Since the .deb does not install in this image, just download it and extract the static binary
 RUN wget https://repo.mongodb.org/apt/debian/dists/buster/mongodb-org/4.4/main/binary-amd64/mongodb-org-server_4.4.0_amd64.deb && \
