Horde: Add settings for whether to create conform and upgrade tasks, and disable them in local builds.

diff --git a/Engine/Source/Programs/Horde/Horde.Build/ServerSettings.cs b/Engine/Source/Programs/Horde/Horde.Build/ServerSettings.cs
--- a/Engine/Source/Programs/Horde/Horde.Build/ServerSettings.cs
+++ b/Engine/Source/Programs/Horde/Horde.Build/ServerSettings.cs
@@ -608,6 +608,16 @@ namespace Horde.Build
 		/// </summary>
 		public int PerforceConnectionPoolSize { get; set; } = 5;
 
+		/// <summary>
+		/// Whether to enable the upgrade task source.
+		/// </summary>
+		public bool EnableUpgradeTasks { get; set; } = true;
+
+		/// <summary>
+		/// Whether to enable the conform task source.
+		/// </summary>
+		public bool EnableConformTasks { get; set; } = true;
+
 		/// <summary>
 		/// Forces configuration data to be read and updated as part of appplication startup, rather than on a schedule. Useful when running locally.
 		/// </summary>
diff --git a/Engine/Source/Programs/Horde/Horde.Build/Tasks/ConformTaskSource.cs b/Engine/Source/Programs/Horde/Horde.Build/Tasks/ConformTaskSource.cs
--- a/Engine/Source/Programs/Horde/Horde.Build/Tasks/ConformTaskSource.cs
+++ b/Engine/Source/Programs/Horde/Horde.Build/Tasks/ConformTaskSource.cs
@@ -19,6 +19,7 @@ using HordeCommon;
 using HordeCommon.Rpc.Tasks;
 using Microsoft.Extensions.Hosting;
 using Microsoft.Extensions.Logging;
+using Microsoft.Extensions.Options;
 
 namespace Horde.Build.Tasks
 {
@@ -45,12 +46,13 @@ namespace Horde.Build.Tasks
 		readonly PerforceLoadBalancer _perforceLoadBalancer;
 		readonly ILogFileService _logService;
 		readonly ILogger _logger;
+		readonly IOptionsMonitor<ServerSettings> _settings;
 		readonly ITicker _tickConformList;
 
 		/// <summary>
 		/// Constructor
 		/// </summary>
-		public ConformTaskSource(MongoService mongoService, GlobalsService globalsService, IAgentCollection agentCollection, PoolService poolService, ILogFileService logService, PerforceLoadBalancer perforceLoadBalancer, IClock clock, ILogger<ConformTaskSource> logger)
+		public ConformTaskSource(MongoService mongoService, GlobalsService globalsService, IAgentCollection agentCollection, PoolService poolService, ILogFileService logService, PerforceLoadBalancer perforceLoadBalancer, IClock clock, IOptionsMonitor<ServerSettings> settings, ILogger<ConformTaskSource> logger)
 		{
 			_mongoService = mongoService;
 			_globalsService = globalsService;
@@ -59,6 +61,7 @@ namespace Horde.Build.Tasks
 			_conformList = new SingletonDocument<ConformList>(mongoService);
 			_perforceLoadBalancer = perforceLoadBalancer;
 			_logService = logService;
+			_settings = settings;
 			_logger = logger;
 			_tickConformList = clock.AddSharedTicker<ConformTaskSource>(TimeSpan.FromMinutes(1.0), CleanConformListAsync, logger);
 
@@ -156,6 +159,11 @@ namespace Horde.Build.Tasks
 		/// <inheritdoc/>
 		public override async Task<Task<AgentLease?>> AssignLeaseAsync(IAgent agent, CancellationToken cancellationToken)
 		{
+			if (!_settings.CurrentValue.EnableConformTasks)
+			{
+				return Skip(cancellationToken);
+			}
+
 			DateTime utcNow = DateTime.UtcNow;
 			if (!await IsConformPendingAsync(agent, utcNow))
 			{
diff --git a/Engine/Source/Programs/Horde/Horde.Build/Tasks/UpgradeTaskSource.cs b/Engine/Source/Programs/Horde/Horde.Build/Tasks/UpgradeTaskSource.cs
--- a/Engine/Source/Programs/Horde/Horde.Build/Tasks/UpgradeTaskSource.cs
+++ b/Engine/Source/Programs/Horde/Horde.Build/Tasks/UpgradeTaskSource.cs
@@ -14,6 +14,7 @@ using Horde.Build.Server;
 using Horde.Build.Utilities;
 using HordeCommon;
 using HordeCommon.Rpc.Tasks;
+using Microsoft.Extensions.Options;
 
 namespace Horde.Build.Tasks
 {
@@ -31,13 +32,15 @@ namespace Horde.Build.Tasks
 		readonly AgentSoftwareService _agentSoftwareService;
 		readonly GlobalsService _globalsService;
 		readonly ILogFileService _logService;
+		readonly IOptionsMonitor<ServerSettings> _settings;
 		readonly IClock _clock;
 
-		public UpgradeTaskSource(AgentSoftwareService agentSoftwareService, GlobalsService globalsService, ILogFileService logService, IClock clock)
+		public UpgradeTaskSource(AgentSoftwareService agentSoftwareService, GlobalsService globalsService, ILogFileService logService, IOptionsMonitor<ServerSettings> settings, IClock clock)
 		{
 			_agentSoftwareService = agentSoftwareService;
 			_globalsService = globalsService;
 			_logService = logService;
+			_settings = settings;
 			_clock = clock;
 
 			OnLeaseStartedProperties.Add(nameof(UpgradeTask.LogId), x => new LogId(x.LogId));
@@ -45,6 +48,11 @@ namespace Horde.Build.Tasks
 
 		public override async Task<Task<AgentLease?>> AssignLeaseAsync(IAgent agent, CancellationToken cancellationToken)
 		{
+			if (!_settings.CurrentValue.EnableUpgradeTasks)
+			{
+				return Skip(cancellationToken);
+			}
+
 			string? requiredVersion = await GetRequiredSoftwareVersion(agent);
 			if (requiredVersion == null || agent.Version == requiredVersion)
 			{
diff --git a/Engine/Source/Programs/Horde/Horde.Build/appsettings.Local.json b/Engine/Source/Programs/Horde/Horde.Build/appsettings.Local.json
--- a/Engine/Source/Programs/Horde/Horde.Build/appsettings.Local.json
+++ b/Engine/Source/Programs/Horde/Horde.Build/appsettings.Local.json
@@ -1,6 +1,8 @@
 {
 	"Horde":
 	{
+		"EnableUpgradeTasks": false,
+		"EnableConformTasks": false,
 		"ForceConfigUpdateOnStartup": true,
 		"Commits":
 		{
