Horde: Don't include suite tests in stream response
#rnx

diff --git a/Engine/Source/Programs/Horde/Horde.Build/Jobs/TestData/TestDataController.cs b/Engine/Source/Programs/Horde/Horde.Build/Jobs/TestData/TestDataController.cs
--- a/Engine/Source/Programs/Horde/Horde.Build/Jobs/TestData/TestDataController.cs
+++ b/Engine/Source/Programs/Horde/Horde.Build/Jobs/TestData/TestDataController.cs
@@ -141,14 +141,6 @@ namespace Horde.Build.Jobs.TestData
 				suites = await _testDataService.FindTestSuites(suiteIds.ToArray());
 			}
 
-			foreach(ITestSuite suite in suites)
-			{
-				foreach(TestId id in suite.Tests)
-				{
-					testIds.Add(id);
-				}
-			}
-
 			List<ITest> tests = new List<ITest>();
 			if (testIds.Count > 0)
 			{
@@ -181,14 +173,6 @@ namespace Horde.Build.Jobs.TestData
 					if (suite != null)
 					{
 						streamSuites.Add(suite);
-						foreach (TestId testId in suite.Tests)
-						{
-							ITest? test = tests.FirstOrDefault(x => x.Id == testId);
-							if (test != null)
-							{
-								streamTests.Add(test);
-							}
-						}
 					}
 				}
 
