Horde: Include metadata in test suite responses, also include id's as these are useful now
#rnx

diff --git a/Engine/Source/Programs/Horde/Horde.Build.Tests/TestDataTests.cs b/Engine/Source/Programs/Horde/Horde.Build.Tests/TestDataTests.cs
--- a/Engine/Source/Programs/Horde/Horde.Build.Tests/TestDataTests.cs
+++ b/Engine/Source/Programs/Horde/Horde.Build.Tests/TestDataTests.cs
@@ -256,12 +256,10 @@ namespace Horde.Build.Tests
 			Assert.AreEqual(4, streams[0].Tests.Count);				
 			Assert.AreEqual(1, streams[0].TestMetadata.Count);
 			Assert.AreEqual(1, streams[0].TestSuites.Count);
-			Assert.AreEqual(4, streams[0].TestSuites[0].Tests.Count);
 
 			Assert.AreEqual(4, streams[1].Tests.Count);
 			Assert.AreEqual(1, streams[1].TestMetadata.Count);
-			Assert.AreEqual(1, streams[1].TestSuites.Count);
-			Assert.AreEqual(4, streams[1].TestSuites[0].Tests.Count);
+			Assert.AreEqual(1, streams[1].TestSuites.Count);			
 
 			Assert.AreEqual(streams[0].TestMetadata[0].Id, streams[1].TestMetadata[0].Id);
 
diff --git a/Engine/Source/Programs/Horde/Horde.Build/Jobs/TestData/TestDataMessages.cs b/Engine/Source/Programs/Horde/Horde.Build/Jobs/TestData/TestDataMessages.cs
--- a/Engine/Source/Programs/Horde/Horde.Build/Jobs/TestData/TestDataMessages.cs
+++ b/Engine/Source/Programs/Horde/Horde.Build/Jobs/TestData/TestDataMessages.cs
@@ -175,6 +175,11 @@ namespace Horde.Build.Jobs.TestData
 	/// </summary>
 	public class GetTestResponse
 	{
+		/// <summary>
+		/// The id of the test
+		/// </summary>
+		public string Id { get; set; }
+
 		/// <summary>
 		/// The name of the test 
 		/// </summary>
@@ -191,12 +196,13 @@ namespace Horde.Build.Jobs.TestData
 		public string? SuiteName { get; set; }
 
 		/// <summary>
-		/// The platforms the test runs on
+		/// The meta data the test runs on
 		/// </summary>
 		public List<string> Metadata { get; set; }
 
 		internal GetTestResponse(ITest test)
 		{
+			Id = test.Id.ToString();
 			Name = test.Name;
 			DisplayName = test.DisplayName;
 			SuiteName = test.SuiteName?.ToString();
@@ -209,20 +215,26 @@ namespace Horde.Build.Jobs.TestData
 	/// </summary>
 	public class GetTestSuiteResponse
 	{
+		/// <summary>
+		/// The id of the suite
+		/// </summary>
+		public string Id { get; set; }
+
 		/// <summary>
 		/// The name of the test suite
 		/// </summary>
 		public string Name { get; set; }
 
 		/// <summary>
-		/// The tests in the suite
+		/// The meta data the test suite runs on
 		/// </summary>
-		public List<TestId> Tests { get; set; }
+		public List<string> Metadata { get; set; }
 
 		internal GetTestSuiteResponse(ITestSuite suite)
 		{
+			Id = suite.Id.ToString();
 			Name = suite.Name;
-			Tests = suite.Tests.ToList();
+			Metadata = suite.Tests.Select(x => x.ToString()).ToList();
 		}
 	}
 
