Horde: Disable Slack notification sink when a Slack token is not provided. Prevents errors when attempting to query user information.

diff --git a/Engine/Source/Programs/Horde/Horde.Build/Startup.cs b/Engine/Source/Programs/Horde/Horde.Build/Startup.cs
--- a/Engine/Source/Programs/Horde/Horde.Build/Startup.cs
+++ b/Engine/Source/Programs/Horde/Horde.Build/Startup.cs
@@ -417,9 +417,14 @@ namespace Horde.Build
 			services.AddSingleton<ProjectService>();
 			services.AddSingleton<ReplicationService>();
 			services.AddSingleton<ScheduleService>();
-			services.AddSingleton<SlackNotificationSink>();
-			services.AddSingleton<IAvatarService, SlackNotificationSink>(sp => sp.GetRequiredService<SlackNotificationSink>());
-			services.AddSingleton<INotificationSink, SlackNotificationSink>(sp => sp.GetRequiredService<SlackNotificationSink>());
+
+			if (settings.SlackToken != null)
+			{
+				services.AddSingleton<SlackNotificationSink>();
+				services.AddSingleton<IAvatarService, SlackNotificationSink>(sp => sp.GetRequiredService<SlackNotificationSink>());
+				services.AddSingleton<INotificationSink, SlackNotificationSink>(sp => sp.GetRequiredService<SlackNotificationSink>());
+			}
+
 			services.AddSingleton<StreamService>();
 			services.AddSingleton<DeviceService>();						
 			services.AddSingleton<NoticeService>();
@@ -617,7 +622,10 @@ namespace Horde.Build
 				services.AddHostedService<MetricService>();
 				services.AddHostedService(provider => provider.GetRequiredService<PerforceLoadBalancer>());
 				services.AddHostedService<PoolUpdateService>();
-				services.AddHostedService(provider => provider.GetRequiredService<SlackNotificationSink>());
+				if (settings.SlackToken != null)
+				{
+					services.AddHostedService(provider => provider.GetRequiredService<SlackNotificationSink>());
+				}
 				services.AddHostedService<ConfigUpdateService>();
 				services.AddHostedService<TelemetryService>();
 				services.AddHostedService(provider => provider.GetRequiredService<DeviceService>());				
