Horde: Fix service provider not being disposed.

diff --git a/Engine/Source/Programs/Horde/Horde.Build/Program.cs b/Engine/Source/Programs/Horde/Horde.Build/Program.cs
--- a/Engine/Source/Programs/Horde/Horde.Build/Program.cs
+++ b/Engine/Source/Programs/Horde/Horde.Build/Program.cs
@@ -179,14 +179,14 @@ namespace Horde.Build
 				Serilog.Log.Logger.Information("Enabling datadog tracing (OpenTrace)");
 			}
 
-			IServiceCollection services = new ServiceCollection();
+			ServiceCollection services = new ServiceCollection();
 			services.AddCommandsFromAssembly(Assembly.GetExecutingAssembly());
 			services.AddLogging(builder => builder.AddSerilog());
 			services.AddSingleton<IConfiguration>(config);
 			services.AddSingleton<ServerSettings>(hordeSettings);
 
 #pragma warning disable ASP0000 // Do not call 'IServiceCollection.BuildServiceProvider' in 'ConfigureServices'
-			IServiceProvider serviceProvider = services.BuildServiceProvider();
+			await using ServiceProvider serviceProvider = services.BuildServiceProvider();
 			return await CommandHost.RunAsync(arguments, serviceProvider, typeof(ServerCommand));
 #pragma warning restore ASP0000 // Do not call 'IServiceCollection.BuildServiceProvider' in 'ConfigureServices'
 		}
