Horde: Store metadata on stream suites to improve queries
#rnx
#skipci

diff --git a/Engine/Source/Programs/Horde/Horde.Build.Tests/TestDataTests.cs b/Engine/Source/Programs/Horde/Horde.Build.Tests/TestDataTests.cs
--- a/Engine/Source/Programs/Horde/Horde.Build.Tests/TestDataTests.cs
+++ b/Engine/Source/Programs/Horde/Horde.Build.Tests/TestDataTests.cs
@@ -254,11 +254,11 @@ namespace Horde.Build.Tests
 			Assert.AreEqual(2, streams.Count);
 
 			Assert.AreEqual(0, streams[0].Tests.Count);				
-			Assert.AreEqual(0, streams[0].TestMetadata.Count);
+			Assert.AreEqual(1, streams[0].TestMetadata.Count);
 			Assert.AreEqual(1, streams[0].TestSuites.Count);
 
 			Assert.AreEqual(0, streams[1].Tests.Count);
-			Assert.AreEqual(0, streams[1].TestMetadata.Count);
+			Assert.AreEqual(1, streams[1].TestMetadata.Count);
 			Assert.AreEqual(1, streams[1].TestSuites.Count);			
 
 			//Assert.AreEqual(streams[0].TestMetadata[0].Id, streams[1].TestMetadata[0].Id);
diff --git a/Engine/Source/Programs/Horde/Horde.Build/Jobs/TestData/ITestData.cs b/Engine/Source/Programs/Horde/Horde.Build/Jobs/TestData/ITestData.cs
--- a/Engine/Source/Programs/Horde/Horde.Build/Jobs/TestData/ITestData.cs
+++ b/Engine/Source/Programs/Horde/Horde.Build/Jobs/TestData/ITestData.cs
@@ -126,6 +126,11 @@ namespace Horde.Build.Jobs.TestData
 		/// The tests that compose the suite
 		/// </summary>
 		IReadOnlyList<TestId> Tests { get; }
+
+		/// <summary>
+		/// The meta data for the test suite
+		/// </summary>
+		IReadOnlyList<TestMetaId> Metadata { get; }
 	}
 
 	/// <summary>
diff --git a/Engine/Source/Programs/Horde/Horde.Build/Jobs/TestData/TestDataCollection.cs b/Engine/Source/Programs/Horde/Horde.Build/Jobs/TestData/TestDataCollection.cs
--- a/Engine/Source/Programs/Horde/Horde.Build/Jobs/TestData/TestDataCollection.cs
+++ b/Engine/Source/Programs/Horde/Horde.Build/Jobs/TestData/TestDataCollection.cs
@@ -121,6 +121,10 @@ namespace Horde.Build.Jobs.TestData
 			[BsonRequired, BsonElement("n")]
 			public string Name { get; set; }
 
+			[BsonRequired, BsonElement("m")]
+			public List<TestMetaId> Metadata { get; set; }
+			IReadOnlyList<TestMetaId> ITestSuite.Metadata => Metadata;
+
 			[BsonRequired, BsonElement("t")]
 			public List<TestId> Tests { get; set; }
 			IReadOnlyList<TestId> ITestSuite.Tests => Tests;
@@ -129,12 +133,14 @@ namespace Horde.Build.Jobs.TestData
 			{
 				Name = String.Empty;
 				Tests = new List<TestId>();
+				Metadata = new List<TestMetaId>();
 			}
 
-			public TestSuiteDocument(string name, List<TestId> tests)
+			public TestSuiteDocument(string name, List<TestMetaId> metaIds, List<TestId> tests)
 			{
 				Id = TestSuiteId.GenerateNewId();
 				Name = name;
+				Metadata = metaIds;
 				Tests = tests;
 			}
 		}
@@ -692,30 +698,31 @@ namespace Horde.Build.Jobs.TestData
 			return await _tests.Find(filter).ToListAsync();
 		}
 
-		private async Task<List<TestSuiteDocument>> FindTestSuites(string[]? projectNames = null, string[]? suiteNames = null, TestMetaId[]? metaIds = null, TestId[]? testIds = null)
+		private async Task<List<TestSuiteDocument>> FindTestSuites(string[] suiteNames, TestMetaId[] metaIds)
 		{
-			// get the tests which satsify the environment
-			List<TestDocument> tests = await FindTests(true, projectNames, null, suiteNames, metaIds, testIds);
 
-			TestId[] testIdQuery = tests.Select(t => t.Id).Distinct().ToArray();
+			FilterDefinitionBuilder<TestSuiteDocument> filterBuilder = Builders<TestSuiteDocument>.Filter;
+			FilterDefinition<TestSuiteDocument> filter = FilterDefinition<TestSuiteDocument>.Empty;
 
-			if (testIdQuery.Length == 0)
+			if (suiteNames.Length == 1)
 			{
-				return new List<TestSuiteDocument>();
+				filter &= filterBuilder.Eq(x => x.Name, suiteNames[0]);
+			}
+			else
+			{
+				filter &= filterBuilder.In(x => x.Name, suiteNames);
 			}
 
-			FilterDefinitionBuilder<TestSuiteDocument> filterBuilder = Builders<TestSuiteDocument>.Filter;
-			FilterDefinition<TestSuiteDocument> filter = FilterDefinition<TestSuiteDocument>.Empty;
-
-			if (testIdQuery.Length == 1)
+			if (metaIds.Length == 1)
 			{
-				filter &= filterBuilder.AnyEq(x => x.Tests, testIdQuery[0]);
+				filter &= filterBuilder.AnyEq(x => x.Metadata, metaIds[0]);
 			}
 			else
 			{
-				filter &= filterBuilder.AnyIn(x => x.Tests, testIdQuery);
+				filter &= filterBuilder.AnyIn(x => x.Metadata, metaIds);
 			}
 
+
 			return await _testSuites.Find(filter).ToListAsync();
 		}
 
@@ -912,7 +919,7 @@ namespace Horde.Build.Jobs.TestData
 
 			if (suiteTests.Count == 0)
 			{
-				suite = new TestSuiteDocument(suiteName, testIds);
+				suite = new TestSuiteDocument(suiteName, new List<TestMetaId> { metaId }, testIds);
 				await _testSuites.InsertOneAsync(suite);
 			}
 			else
@@ -921,16 +928,27 @@ namespace Horde.Build.Jobs.TestData
 
 				List<TestId> newIds = suite.Tests.Union(testIds).ToList();
 
-				if (newIds.Count != suite.Tests.Count)
+				if (newIds.Count != suite.Tests.Count || !suite.Metadata.Contains(metaId))
 				{
-					suite.Tests = newIds;
-
+					List<UpdateDefinition<TestSuiteDocument>> updates = new List<UpdateDefinition<TestSuiteDocument>>();
 					FilterDefinitionBuilder<TestSuiteDocument> ebuilder = Builders<TestSuiteDocument>.Filter;
 					FilterDefinition<TestSuiteDocument> efilter = ebuilder.Eq(x => x.Id, suite.Id);
 
 					UpdateDefinitionBuilder<TestSuiteDocument> updateBuilder = Builders<TestSuiteDocument>.Update;
-					UpdateDefinition<TestSuiteDocument> update = updateBuilder.Set(x => x.Tests, suite.Tests);
-					await _testSuites.UpdateOneAsync(efilter, update);
+
+					if (newIds.Count != suite.Tests.Count)
+					{
+						suite.Tests = newIds;
+						updates.Add(updateBuilder.Set(x => x.Tests, suite.Tests));
+					}
+
+					if (!suite.Metadata.Contains(metaId))
+					{
+						suite.Metadata.Add(metaId);
+						updates.Add(updateBuilder.Set(x => x.Metadata, suite.Metadata));
+					}
+
+					await _testSuites.UpdateOneAsync(efilter, updateBuilder.Combine(updates));
 				}
 
 			}
diff --git a/Engine/Source/Programs/Horde/Horde.Build/Jobs/TestData/TestDataController.cs b/Engine/Source/Programs/Horde/Horde.Build/Jobs/TestData/TestDataController.cs
--- a/Engine/Source/Programs/Horde/Horde.Build/Jobs/TestData/TestDataController.cs
+++ b/Engine/Source/Programs/Horde/Horde.Build/Jobs/TestData/TestDataController.cs
@@ -119,8 +119,7 @@ namespace Horde.Build.Jobs.TestData
 
 			List<ITestStream> streams = await _testDataService.FindTestStreams(queryStreams.ToArray());
 
-			// flatten requested streams to single service queries
-		
+			// flatten requested streams to single service queries		
 			HashSet<TestSuiteId> suiteIds = new HashSet<TestSuiteId>();
 			for (int i = 0; i < streams.Count; i++)
 			{
@@ -147,6 +146,7 @@ namespace Horde.Build.Jobs.TestData
 				tests = await _testDataService.FindTests(testIds.ToArray());
 			}
 
+			// gather all meta data
 			List<ITestMeta> metaData = new List<ITestMeta>();
 			foreach (ITest test in tests)
 			{
@@ -156,6 +156,14 @@ namespace Horde.Build.Jobs.TestData
 				}
 			}
 
+			foreach (ITestSuite suite in suites)
+			{
+				foreach (TestMetaId metaId in suite.Metadata)
+				{
+					metaIds.Add(metaId);
+				}
+			}
+
 			if (metaIds.Count > 0)
 			{
 				metaData = await _testDataService.FindTestMeta(metaIds: metaIds.ToArray());
@@ -185,6 +193,14 @@ namespace Horde.Build.Jobs.TestData
 					}						
 				}
 
+				foreach (ITestSuite suite in streamSuites)
+				{
+					foreach (TestMetaId id in suite.Metadata)
+					{
+						streamMetaIds.Add(id);
+					}
+				}
+
 				List<ITestMeta> streamMetaData = metaData.Where(x => streamMetaIds.Contains(x.Id)).ToList();
 
 				responses.Add(new GetTestStreamResponse(s.StreamId, streamTests, streamSuites, streamMetaData));				
