Horde: Fix copy/paste error in suite response
#rnx
#skipci

diff --git a/Engine/Source/Programs/Horde/Horde.Build/Jobs/TestData/TestDataMessages.cs b/Engine/Source/Programs/Horde/Horde.Build/Jobs/TestData/TestDataMessages.cs
--- a/Engine/Source/Programs/Horde/Horde.Build/Jobs/TestData/TestDataMessages.cs
+++ b/Engine/Source/Programs/Horde/Horde.Build/Jobs/TestData/TestDataMessages.cs
@@ -234,7 +234,7 @@ namespace Horde.Build.Jobs.TestData
 		{
 			Id = suite.Id.ToString();
 			Name = suite.Name;
-			Metadata = suite.Tests.Select(x => x.ToString()).ToList();
+			Metadata = suite.Metadata.Select(x => x.ToString()).ToList();
 		}
 	}
 
