Horde: Fix agent tests.

diff --git a/Engine/Source/Programs/Horde/Horde.Agent.Tests/Helper.cs b/Engine/Source/Programs/Horde/Horde.Agent.Tests/Helper.cs
--- a/Engine/Source/Programs/Horde/Horde.Agent.Tests/Helper.cs
+++ b/Engine/Source/Programs/Horde/Horde.Agent.Tests/Helper.cs
@@ -183,6 +183,8 @@ namespace Horde.Agent.Tests
 
 	class SimpleTestExecutor : JobExecutor
 	{
+		public const string Name = "Simple";
+
 		private readonly Func<BeginStepResponse, ILogger, CancellationToken, Task<JobStepOutcome>> _func;
 
 		public SimpleTestExecutor(Func<BeginStepResponse, ILogger, CancellationToken, Task<JobStepOutcome>> func)
@@ -225,7 +227,7 @@ namespace Horde.Agent.Tests
 	{
 		readonly JobExecutor _executor;
 
-		public override string Name => "Simple";
+		public override string Name => SimpleTestExecutor.Name;
 
 		public SimpleTestExecutorFactory(JobExecutor executor)
 		{
diff --git a/Engine/Source/Programs/Horde/Horde.Agent.Tests/WorkerServiceTest.cs b/Engine/Source/Programs/Horde/Horde.Agent.Tests/WorkerServiceTest.cs
--- a/Engine/Source/Programs/Horde/Horde.Agent.Tests/WorkerServiceTest.cs
+++ b/Engine/Source/Programs/Horde/Horde.Agent.Tests/WorkerServiceTest.cs
@@ -114,6 +114,7 @@ namespace Horde.Agent.Tests
 			executeJobTask.BatchId = "batchId1";
 			executeJobTask.LogId = "logId1";
 			executeJobTask.JobName = "jobName1";
+			executeJobTask.Executor = SimpleTestExecutor.Name;
 			executeJobTask.AutoSdkWorkspace = new AgentWorkspace();
 			executeJobTask.Workspace = new AgentWorkspace();
 
