Backout change to agent tests.

diff --git a/Engine/Source/Programs/Horde/Horde.Agent.Tests/Helper.cs b/Engine/Source/Programs/Horde/Horde.Agent.Tests/Helper.cs
--- a/Engine/Source/Programs/Horde/Horde.Agent.Tests/Helper.cs
+++ b/Engine/Source/Programs/Horde/Horde.Agent.Tests/Helper.cs
@@ -14,6 +14,7 @@ using HordeCommon;
 using HordeCommon.Rpc;
 using HordeCommon.Rpc.Tasks;
 using Microsoft.Extensions.Logging;
+using Microsoft.Extensions.Logging.Abstractions;
 using Microsoft.Extensions.Options;
 
 namespace Horde.Agent.Tests
@@ -31,13 +32,13 @@ namespace Horde.Agent.Tests
 		public T Value { get; }
 	}
 
-	class RpcClientRefStub : IRpcClientRef
+	class RpcClientRefStub<TClient> : IRpcClientRef<TClient> where TClient : ClientBase<TClient>
 	{
 		public GrpcChannel Channel { get; }
-		public HordeRpc.HordeRpcClient Client { get; }
+		public TClient Client { get; }
 		public Task DisposingTask { get; }
 
-		public RpcClientRefStub(GrpcChannel channel, HordeRpc.HordeRpcClient client)
+		public RpcClientRefStub(GrpcChannel channel, TClient client)
 		{
 			Channel = channel;
 			Client = client;
@@ -54,47 +55,25 @@ namespace Horde.Agent.Tests
 		private readonly GrpcChannel _grpcChannel;
 		private readonly HordeRpc.HordeRpcClient _hordeRpcClient;
 
+		public ILogger Logger => NullLogger.Instance;
+
 		public RpcConnectionStub(GrpcChannel grpcChannel, HordeRpc.HordeRpcClient hordeRpcClient)
 		{
 			_grpcChannel = grpcChannel;
 			_hordeRpcClient = hordeRpcClient;
 		}
 
-		public IRpcClientRef? TryGetClientRef(RpcContext context)
+		public IRpcClientRef<TClient>? TryGetClientRef<TClient>() where TClient : ClientBase<TClient>
 		{
-			return new RpcClientRefStub(_grpcChannel, _hordeRpcClient);
+			return (IRpcClientRef<TClient>)(object)new RpcClientRefStub<HordeRpc.HordeRpcClient>(_grpcChannel, _hordeRpcClient);
 		}
 
-		public Task<IRpcClientRef> GetClientRef(RpcContext context, CancellationToken cancellationToken)
+		public Task<IRpcClientRef<TClient>> GetClientRefAsync<TClient>(CancellationToken cancellationToken) where TClient : ClientBase<TClient>
 		{
-			IRpcClientRef rpcClientRefStub = new RpcClientRefStub(_grpcChannel, _hordeRpcClient);
+			IRpcClientRef<TClient> rpcClientRefStub = (IRpcClientRef<TClient>)(object)new RpcClientRefStub<HordeRpc.HordeRpcClient>(_grpcChannel, _hordeRpcClient);
 			return Task.FromResult(rpcClientRefStub);
 		}
 
-		public Task<T> InvokeOnceAsync<T>(Func<HordeRpc.HordeRpcClient, Task<T>> func, RpcContext context,
-			CancellationToken cancellationToken)
-		{
-			return func(_hordeRpcClient);
-		}
-
-		public Task<T> InvokeOnceAsync<T>(Func<HordeRpc.HordeRpcClient, AsyncUnaryCall<T>> func, RpcContext context,
-			CancellationToken cancellationToken)
-		{
-			return func(_hordeRpcClient).ResponseAsync;
-		}
-
-		public Task<T> InvokeAsync<T>(Func<HordeRpc.HordeRpcClient, Task<T>> func, RpcContext context,
-			CancellationToken cancellationToken)
-		{
-			return func(_hordeRpcClient);
-		}
-
-		public Task<T> InvokeAsync<T>(Func<HordeRpc.HordeRpcClient, AsyncUnaryCall<T>> func, RpcContext context,
-			CancellationToken cancellationToken)
-		{
-			return func(_hordeRpcClient).ResponseAsync;
-		}
-
 		public ValueTask DisposeAsync()
 		{
 			return new ValueTask();
