Horde: Allow sending notifications to people in an allow list from the dev environment.

diff --git a/Engine/Source/Programs/Horde/Horde.Build/Notifications/Sinks/SlackNotificationSink.cs b/Engine/Source/Programs/Horde/Horde.Build/Notifications/Sinks/SlackNotificationSink.cs
--- a/Engine/Source/Programs/Horde/Horde.Build/Notifications/Sinks/SlackNotificationSink.cs
+++ b/Engine/Source/Programs/Horde/Horde.Build/Notifications/Sinks/SlackNotificationSink.cs
@@ -1449,9 +1449,12 @@ namespace Horde.Build.Notifications.Sinks
 				return user.Login;
 			}
 
-			if (!_environment.IsProduction() || !allowMentions)
+			if (_allowUsers == null || !_allowUsers.Contains(slackUserId))
 			{
-				return $"{user.Name} [{slackUserId}]";
+				if (!_environment.IsProduction() || !allowMentions)
+				{
+					return $"{user.Name} [{slackUserId}]";
+				}
 			}
 
 			return $"<@{slackUserId}>";
