Horde: Allow local executor to find a working directory from the executable directory upwards.

diff --git a/Engine/Source/Programs/Horde/Horde.Agent/Execution/LocalExecutor.cs b/Engine/Source/Programs/Horde/Horde.Agent/Execution/LocalExecutor.cs
--- a/Engine/Source/Programs/Horde/Horde.Agent/Execution/LocalExecutor.cs
+++ b/Engine/Source/Programs/Horde/Horde.Agent/Execution/LocalExecutor.cs
@@ -2,6 +2,7 @@
 
 using System;
 using System.Net.Http;
+using System.Reflection;
 using System.Threading;
 using System.Threading.Tasks;
 using EpicGames.Core;
@@ -23,11 +24,31 @@ namespace Horde.Agent.Execution
 			: base(session, jobId, batchId, agentTypeName, httpClientFactory)
 		{
 			_settings = settings;
-			if(settings.WorkspaceDir == null)
+			if (settings.WorkspaceDir == null)
 			{
-				throw new Exception("Missing LocalWorkspaceDir from settings");
+				_localWorkspaceDir = FindWorkspaceRoot();
 			}
-			_localWorkspaceDir = new DirectoryReference(settings.WorkspaceDir);
+			else
+			{
+				_localWorkspaceDir = new DirectoryReference(settings.WorkspaceDir);
+			}
+		}
+
+		static DirectoryReference FindWorkspaceRoot()
+		{
+			const string HordeSlnRelativePath = "Engine/Source/Programs/Horde/Horde.sln";
+
+			FileReference executableFile = new FileReference(Assembly.GetExecutingAssembly().Location);
+			for (DirectoryReference? directory = executableFile.Directory; directory != null; directory = directory.ParentDirectory)
+			{
+				FileReference HordeSln = FileReference.Combine(directory, HordeSlnRelativePath);
+				if (FileReference.Exists(HordeSln))
+				{
+					return directory;
+				}
+			}
+
+			throw new Exception($"Unable to find workspace root directory (looking for '{HordeSlnRelativePath}' in a parent directory of '{executableFile.Directory}'");
 		}
 
 		protected override Task<bool> SetupAsync(BeginStepResponse step, ILogger logger, CancellationToken cancellationToken)
