Horde: Add timezone property to IClock, and use it for all schedule triggers. Note: this changes issue reports to be sent on the configured timezone rather than UTC.

diff --git a/Engine/Source/Programs/Horde/Horde.Build.Tests/SchedulerTests.cs b/Engine/Source/Programs/Horde/Horde.Build.Tests/SchedulerTests.cs
--- a/Engine/Source/Programs/Horde/Horde.Build.Tests/SchedulerTests.cs
+++ b/Engine/Source/Programs/Horde/Horde.Build.Tests/SchedulerTests.cs
@@ -200,6 +200,7 @@ namespace Horde.Build.Tests
 		{
 			DateTime startTime = new DateTime(2021, 1, 1, 12, 0, 0, DateTimeKind.Local); // Friday Jan 1, 2021 
 			Clock.UtcNow = startTime;
+			Clock.TimeZone = TimeZoneInfo.Local;
 
 			ScheduleConfig schedule = new ScheduleConfig();
 			schedule.Enabled = true;
@@ -228,6 +229,7 @@ namespace Horde.Build.Tests
 		{
 			DateTime startTime = new DateTime(2021, 1, 1, 12, 0, 0, DateTimeKind.Local); // Friday Jan 1, 2021 
 			Clock.UtcNow = startTime;
+			Clock.TimeZone = TimeZoneInfo.Local;
 
 			ScheduleConfig schedule = new ScheduleConfig();
 			schedule.Enabled = true;
@@ -267,6 +269,7 @@ namespace Horde.Build.Tests
 		{
 			DateTime startTime = new DateTime(2021, 1, 1, 12, 0, 0, DateTimeKind.Local); // Friday Jan 1, 2021 
 			Clock.UtcNow = startTime;
+			Clock.TimeZone = TimeZoneInfo.Local;
 
 			ScheduleConfig schedule = new ScheduleConfig();
 			schedule.Enabled = true;
@@ -304,6 +307,7 @@ namespace Horde.Build.Tests
 		{
 			DateTime startTime = new DateTime(2021, 1, 1, 12, 0, 0, DateTimeKind.Local); // Friday Jan 1, 2021 
 			Clock.UtcNow = startTime;
+			Clock.TimeZone = TimeZoneInfo.Local;
 
 			ScheduleConfig schedule = new ScheduleConfig();
 			schedule.Enabled = true;
@@ -341,6 +345,7 @@ namespace Horde.Build.Tests
 		{
 			DateTime startTime = new DateTime(2021, 1, 1, 12, 0, 0, DateTimeKind.Local); // Friday Jan 1, 2021 
 			Clock.UtcNow = startTime;
+			Clock.TimeZone = TimeZoneInfo.Local;
 
 			ScheduleConfig schedule = new ScheduleConfig();
 			schedule.Enabled = true;
@@ -383,6 +388,7 @@ namespace Horde.Build.Tests
 		{
 			DateTime startTime = new DateTime(2021, 1, 1, 12, 0, 0, DateTimeKind.Local); // Friday Jan 1, 2021 
 			Clock.UtcNow = startTime;
+			Clock.TimeZone = TimeZoneInfo.Local;
 
 			ScheduleConfig schedule = new ScheduleConfig();
 			schedule.Enabled = true;
@@ -571,6 +577,7 @@ namespace Horde.Build.Tests
 		{
 			DateTime startTime = new DateTime(2021, 1, 1, 12, 0, 0, DateTimeKind.Local); // Friday Jan 1, 2021 
 			Clock.UtcNow = startTime;
+			Clock.TimeZone = TimeZoneInfo.Local;
 
 			ScheduleConfig schedule = new ScheduleConfig();
 			schedule.Enabled = true;
diff --git a/Engine/Source/Programs/Horde/Horde.Build/Issues/IssueReportService.cs b/Engine/Source/Programs/Horde/Horde.Build/Issues/IssueReportService.cs
--- a/Engine/Source/Programs/Horde/Horde.Build/Issues/IssueReportService.cs
+++ b/Engine/Source/Programs/Horde/Horde.Build/Issues/IssueReportService.cs
@@ -154,13 +154,13 @@ namespace Horde.Build.Issues
 							continue;
 						}
 
-						DateTime lastScheduledReportTime = GetLastScheduledReportTime(workflow, currentTime);
+						DateTime lastScheduledReportTime = GetLastScheduledReportTime(workflow, currentTime, _clock.TimeZone);
 						if (lastReportTime > lastScheduledReportTime)
 						{
 							continue;
 						}
 
-						DateTime prevScheduledReportTime = GetLastScheduledReportTime(workflow, lastScheduledReportTime - TimeSpan.FromMinutes(1.0));
+						DateTime prevScheduledReportTime = GetLastScheduledReportTime(workflow, lastScheduledReportTime - TimeSpan.FromMinutes(1.0), _clock.TimeZone);
 
 						_logger.LogInformation("Creating report for {StreamId} workflow {WorkflowId}", stream.Id, workflow.Id);
 
@@ -273,22 +273,22 @@ namespace Horde.Build.Issues
 			return workflowIdToStats;
 		}
 
-		static DateTime GetLastScheduledReportTime(WorkflowConfig workflow, DateTime currentTime)
+		static DateTime GetLastScheduledReportTime(WorkflowConfig workflow, DateTime currentTimeUtc, TimeZoneInfo timeZone)
 		{
 			if (workflow.ReportTimes.Count > 0)
 			{
-				DateTime startOfDay = currentTime - currentTime.TimeOfDay;
+				DateTime startOfDayUtc = timeZone.GetStartOfDayUtc(currentTimeUtc);
 				for (; ; )
 				{
 					for (int idx = workflow.ReportTimes.Count - 1; idx >= 0; idx--)
 					{
-						DateTime reportTime = startOfDay + workflow.ReportTimes[idx];
-						if (reportTime < currentTime)
+						DateTime reportTime = startOfDayUtc + workflow.ReportTimes[idx];
+						if (reportTime < currentTimeUtc)
 						{
 							return reportTime;
 						}
 					}
-					startOfDay -= TimeSpan.FromDays(1.0);
+					startOfDayUtc -= TimeSpan.FromDays(1.0);
 				}
 			}
 			return DateTime.MinValue;
diff --git a/Engine/Source/Programs/Horde/Horde.Build/Jobs/Schedules/ScheduleService.cs b/Engine/Source/Programs/Horde/Horde.Build/Jobs/Schedules/ScheduleService.cs
--- a/Engine/Source/Programs/Horde/Horde.Build/Jobs/Schedules/ScheduleService.cs
+++ b/Engine/Source/Programs/Horde/Horde.Build/Jobs/Schedules/ScheduleService.cs
@@ -67,7 +67,6 @@ namespace Horde.Build.Jobs.Schedules
 		readonly JobService _jobService;
 		readonly StreamService _streamService;
 		readonly ITemplateCollection _templateCollection;
-		readonly TimeZoneInfo _timeZone;
 		readonly IClock _clock;
 		readonly ILogger _logger;
 		readonly RedisService _redis;
@@ -88,10 +87,6 @@ namespace Horde.Build.Jobs.Schedules
 			_streamService = streamService;
 			_templateCollection = templateCollection;
 			_clock = clock;
-
-			string? scheduleTimeZone = settings.CurrentValue.ScheduleTimeZone;
-			_timeZone = (scheduleTimeZone == null) ? TimeZoneInfo.Local : TZConvert.GetTimeZoneInfo(scheduleTimeZone);
-
 			_logger = logger;
 
 			_redis = redis;
@@ -214,7 +209,7 @@ namespace Horde.Build.Jobs.Schedules
 				{
 					if (templateRef.Schedule != null)
 					{
-						DateTime? nextTriggerTimeUtc = templateRef.Schedule.GetNextTriggerTimeUtc(_timeZone);
+						DateTime? nextTriggerTimeUtc = templateRef.Schedule.GetNextTriggerTimeUtc(_clock.TimeZone);
 						if (nextTriggerTimeUtc != null)
 						{
 							if (utcNow > nextTriggerTimeUtc.Value)
diff --git a/Engine/Source/Programs/Horde/Horde.Build/Server/DowntimeService.cs b/Engine/Source/Programs/Horde/Horde.Build/Server/DowntimeService.cs
--- a/Engine/Source/Programs/Horde/Horde.Build/Server/DowntimeService.cs
+++ b/Engine/Source/Programs/Horde/Horde.Build/Server/DowntimeService.cs
@@ -40,8 +40,8 @@ namespace Horde.Build.Server
 		}
 
 		readonly GlobalsService _globalsService;
+		readonly IClock _clock;
 		readonly ITicker _ticker;
-		readonly IOptionsMonitor<ServerSettings> _settings;
 		readonly ILogger _logger;
 
 		/// <summary>
@@ -49,12 +49,11 @@ namespace Horde.Build.Server
 		/// </summary>
 		/// <param name="globalsService">The database service instance</param>
 		/// <param name="clock"></param>
-		/// <param name="settings">The server settings</param>
 		/// <param name="logger">Logger instance</param>
-		public DowntimeService(GlobalsService globalsService, IClock clock, IOptionsMonitor<ServerSettings> settings, ILogger<DowntimeService> logger)
+		public DowntimeService(GlobalsService globalsService, IClock clock, ILogger<DowntimeService> logger)
 		{
 			_globalsService = globalsService;
-			_settings = settings;
+			_clock = clock;
 			_logger = logger;
 
 			// Ensure the initial value to be correct
@@ -81,7 +80,7 @@ namespace Horde.Build.Server
 		{
 			IGlobals globals = await _globalsService.GetAsync();
 
-			DateTimeOffset now = TimeZoneInfo.ConvertTime(DateTimeOffset.Now, _settings.CurrentValue.TimeZoneInfo);
+			DateTimeOffset now = TimeZoneInfo.ConvertTime(new DateTimeOffset(_clock.UtcNow), _clock.TimeZone);
 			bool isActive = globals.Config.Downtime.Any(x => x.IsActive(now));
 
 			DateTimeOffset? next = null;
diff --git a/Engine/Source/Programs/Horde/Horde.Build/Server/Notices/NoticesController.cs b/Engine/Source/Programs/Horde/Horde.Build/Server/Notices/NoticesController.cs
--- a/Engine/Source/Programs/Horde/Horde.Build/Server/Notices/NoticesController.cs
+++ b/Engine/Source/Programs/Horde/Horde.Build/Server/Notices/NoticesController.cs
@@ -6,6 +6,7 @@ using System.Threading.Tasks;
 using Horde.Build.Acls;
 using Horde.Build.Users;
 using Horde.Build.Utilities;
+using HordeCommon;
 using Microsoft.AspNetCore.Authorization;
 using Microsoft.AspNetCore.Mvc;
 using Microsoft.Extensions.Options;
@@ -47,26 +48,26 @@ namespace Horde.Build.Server.Notices
 		readonly LazyCachedValue<Task<List<INotice>>> _cachedNotices;
 
 		/// <summary>
-		/// Server settings
+		/// Time information
 		/// </summary>
-		readonly IOptionsMonitor<ServerSettings> _settings;
+		readonly IClock _clock;
 
 		/// <summary>
 		/// Constructor
 		/// </summary>
-		/// <param name="settings">The server settings</param>
 		/// <param name="globalsService">The globals service singleton</param>
 		/// <param name="noticeService">The notice service singleton</param>
 		/// <param name="aclService">The acl service singleton</param>
 		/// <param name="userCollection">The user collection singleton</param>
-		public NoticesController(IOptionsMonitor<ServerSettings> settings, GlobalsService globalsService, NoticeService noticeService, AclService aclService, IUserCollection userCollection)
+		/// <param name="clock"></param>
+		public NoticesController(GlobalsService globalsService, NoticeService noticeService, AclService aclService, IUserCollection userCollection, IClock clock)
 		{			
-			_settings = settings;
 			_aclService = aclService;
 			_userCollection = userCollection;
 			_noticeService = noticeService;
 			_cachedGlobals = new LazyCachedValue<Task<IGlobals>>(async () => await globalsService.GetAsync(), TimeSpan.FromSeconds(30.0));
 			_cachedNotices = new LazyCachedValue<Task<List<INotice>>>(() => noticeService.GetNoticesAsync(), TimeSpan.FromMinutes(1));
+			_clock = clock;
 		}
 
 		/// <summary>
@@ -132,7 +133,7 @@ namespace Horde.Build.Server.Notices
 
 			IGlobals globals = await _cachedGlobals.GetCached();
 
-			DateTimeOffset now = TimeZoneInfo.ConvertTime(DateTimeOffset.Now, _settings.CurrentValue.TimeZoneInfo);
+			DateTimeOffset now = TimeZoneInfo.ConvertTime(new DateTimeOffset(_clock.UtcNow), _clock.TimeZone);
 						
 			foreach (ScheduledDowntime schedule in globals.Config.Downtime)
 			{
diff --git a/Engine/Source/Programs/Horde/Horde.Build/ServerSettings.cs b/Engine/Source/Programs/Horde/Horde.Build/ServerSettings.cs
--- a/Engine/Source/Programs/Horde/Horde.Build/ServerSettings.cs
+++ b/Engine/Source/Programs/Horde/Horde.Build/ServerSettings.cs
@@ -623,20 +623,6 @@ namespace Horde.Build
 		/// </summary>
 		public bool ForceConfigUpdateOnStartup { get; set; }
 
-		/// <summary>
-		/// Lazily computed timezone value
-		/// </summary>
-		public TimeZoneInfo TimeZoneInfo
-		{
-			get
-			{
-				_cachedTimeZoneInfo ??= (ScheduleTimeZone == null) ? TimeZoneInfo.Local : TZConvert.GetTimeZoneInfo(ScheduleTimeZone);
-				return _cachedTimeZoneInfo;
-			}
-		}
-
-		private TimeZoneInfo? _cachedTimeZoneInfo;
-
 		/// <summary>
 		/// Whether to open a browser on startup
 		/// </summary>
diff --git a/Engine/Source/Programs/Horde/Horde.Build/Utilities/Clock.cs b/Engine/Source/Programs/Horde/Horde.Build/Utilities/Clock.cs
--- a/Engine/Source/Programs/Horde/Horde.Build/Utilities/Clock.cs
+++ b/Engine/Source/Programs/Horde/Horde.Build/Utilities/Clock.cs
@@ -6,9 +6,12 @@ using System.Diagnostics;
 using System.Threading;
 using System.Threading.Tasks;
 using EpicGames.Redis.Utility;
+using Horde.Build;
 using Horde.Build.Server;
 using Microsoft.Extensions.Logging;
+using Microsoft.Extensions.Options;
 using StackExchange.Redis;
+using TimeZoneConverter;
 
 namespace HordeCommon
 {
@@ -51,9 +54,13 @@ namespace HordeCommon
 		/// <summary>
 		/// Return time expressed as the Coordinated Universal Time (UTC)
 		/// </summary>
-		/// <returns></returns>
 		DateTime UtcNow { get; }
 
+		/// <summary>
+		/// Time zone for schedules etc...
+		/// </summary>
+		TimeZoneInfo TimeZone { get; }
+
 		/// <summary>
 		/// Create an event that will trigger after the given time
 		/// </summary>
@@ -130,6 +137,19 @@ namespace HordeCommon
 		/// <param name="logger">Logger for error messages</param>
 		/// <returns>New ticker instance</returns>
 		public static ITicker AddSharedTicker<T>(this IClock clock, TimeSpan interval, Func<CancellationToken, ValueTask> tickAsync, ILogger logger) => clock.AddSharedTicker(typeof(T).Name, interval, tickAsync, logger);
+
+		/// <summary>
+		/// Gets the start of the day for the given datetime in UTC, respecting the configured timezone.
+		/// </summary>
+		/// <param name="timeZone">Time zone to adjust for</param>
+		/// <param name="time">Time to convert</param>
+		/// <returns>UTC datetime for the start of the day</returns>
+		public static DateTime GetStartOfDayUtc(this TimeZoneInfo timeZone, DateTime time)
+		{
+			DateTime currentTimeLocal = TimeZoneInfo.ConvertTime(time, timeZone);
+			DateTime startOfDayLocal = currentTimeLocal - currentTimeLocal.TimeOfDay;
+			return TimeZoneInfo.ConvertTime(startOfDayLocal, timeZone, TimeZoneInfo.Utc);
+		}
 	}
 
 	/// <summary>
@@ -210,16 +230,23 @@ namespace HordeCommon
 		}
 
 		readonly RedisService _redis;
+		readonly TimeZoneInfo _timeZone;
 
 		/// <inheritdoc/>
 		public DateTime UtcNow => DateTime.UtcNow;
 
+		/// <inheritdoc/>
+		public TimeZoneInfo TimeZone => _timeZone;
+
 		/// <summary>
 		/// Constructor
 		/// </summary>
-		public Clock(RedisService redis)
+		public Clock(RedisService redis, IOptions<ServerSettings> settings)
 		{
 			_redis = redis;
+
+			string? timeZoneName = settings.Value.ScheduleTimeZone;
+			_timeZone = (timeZoneName == null) ? TimeZoneInfo.Local : TZConvert.GetTimeZoneInfo(timeZoneName);
 		}
 
 		/// <inheritdoc/>
@@ -322,6 +349,7 @@ namespace HordeCommon
 		public FakeClock()
 		{
 			_utcNowPrivate = DateTime.UtcNow;
+			TimeZone = TimeZoneInfo.Utc;
 		}
 
 		/// <summary>
@@ -356,6 +384,9 @@ namespace HordeCommon
 			set => _utcNowPrivate = value.ToUniversalTime(); 
 		}
 
+		/// <inheritdoc/>
+		public TimeZoneInfo TimeZone { get; set; }
+
 		/// <inheritdoc/>
 		public ITicker AddTicker(string name, TimeSpan interval, Func<CancellationToken, ValueTask<TimeSpan?>> tickAsync, ILogger logger)
 		{
