Horde: Fix mongodb serialization exception on nullable
#rnx

diff --git a/Engine/Source/Programs/Horde/Horde.Build.Tests/TestDataTests.cs b/Engine/Source/Programs/Horde/Horde.Build.Tests/TestDataTests.cs
--- a/Engine/Source/Programs/Horde/Horde.Build.Tests/TestDataTests.cs
+++ b/Engine/Source/Programs/Horde/Horde.Build.Tests/TestDataTests.cs
@@ -201,7 +201,7 @@ namespace Horde.Build.Tests
 
 			Assert.AreEqual(streams[0].TestMetadata[0].Id, streams[1].TestMetadata[0].Id);
 
-			ActionResult<List<GetTestDataRefResponse>> refResult = await TestDataController.GetTestDataRefAsync(streamIds, meta.Select(x => x.Id).ToArray());
+			ActionResult<List<GetTestDataRefResponse>> refResult = await TestDataController.GetTestDataRefAsync(streamIds, meta.Select(x => x.Id).ToArray(), streams[0].Tests.Select(t => t.Id).ToArray());
 			Assert.IsNotNull(refResult);
 			Assert.IsNotNull(refResult.Value);
 			List<GetTestDataRefResponse> refs = refResult.Value;
diff --git a/Engine/Source/Programs/Horde/Horde.Build/Jobs/TestData/TestDataCollection.cs b/Engine/Source/Programs/Horde/Horde.Build/Jobs/TestData/TestDataCollection.cs
--- a/Engine/Source/Programs/Horde/Horde.Build/Jobs/TestData/TestDataCollection.cs
+++ b/Engine/Source/Programs/Horde/Horde.Build/Jobs/TestData/TestDataCollection.cs
@@ -431,14 +431,12 @@ namespace Horde.Build.Jobs.TestData
 
 			if (testIds != null && testIds.Length > 0)
 			{
-				filter &= filterBuilder.Ne(x => x.TestId, null);
-				filter &= filterBuilder.In(x => x.TestId!.Value, testIds);
+				filter &= filterBuilder.And(filterBuilder.Ne(x => x.TestId, null), filterBuilder.In(x => (TestId) x.TestId!, testIds));
 			}
 
 			if (suiteIds != null && suiteIds.Length > 0)
 			{
-				filter &= filterBuilder.Ne(x => x.SuiteId, null);
-				filter &= filterBuilder.In(x => x.SuiteId!.Value, suiteIds);
+				filter &= filterBuilder.And(filterBuilder.Ne(x => x.SuiteId, null), filterBuilder.In(x => (TestSuiteId)x.SuiteId!, suiteIds));
 			}
 
 			if (minChange != null)
