Horde: Support querying both suite and test ids for refs
#rnx

diff --git a/Engine/Source/Programs/Horde/Horde.Build/Jobs/TestData/TestDataCollection.cs b/Engine/Source/Programs/Horde/Horde.Build/Jobs/TestData/TestDataCollection.cs
--- a/Engine/Source/Programs/Horde/Horde.Build/Jobs/TestData/TestDataCollection.cs
+++ b/Engine/Source/Programs/Horde/Horde.Build/Jobs/TestData/TestDataCollection.cs
@@ -429,14 +429,22 @@ namespace Horde.Build.Jobs.TestData
 				filter &= filterBuilder.In(x => x.Metadata, metaIds);
 			}
 
-			if (testIds != null && testIds.Length > 0)
+			if ((testIds != null && testIds.Length> 0) && (suiteIds != null && suiteIds.Length > 0))
 			{
-				filter &= filterBuilder.And(filterBuilder.Ne(x => x.TestId, null), filterBuilder.In(x => (TestId) x.TestId!, testIds));
+				filter &= filterBuilder.Or(filterBuilder.And(filterBuilder.Ne(x => x.TestId, null), filterBuilder.In(x => (TestId)x.TestId!, testIds),
+										   filterBuilder.And(filterBuilder.Ne(x => x.SuiteId, null), filterBuilder.In(x => (TestSuiteId)x.SuiteId!, suiteIds))));
 			}
-
-			if (suiteIds != null && suiteIds.Length > 0)
+			else
 			{
-				filter &= filterBuilder.And(filterBuilder.Ne(x => x.SuiteId, null), filterBuilder.In(x => (TestSuiteId)x.SuiteId!, suiteIds));
+				if (testIds != null && testIds.Length > 0)
+				{
+					filter &= filterBuilder.And(filterBuilder.Ne(x => x.TestId, null), filterBuilder.In(x => (TestId)x.TestId!, testIds));
+				}
+
+				if (suiteIds != null && suiteIds.Length > 0)
+				{
+					filter &= filterBuilder.And(filterBuilder.Ne(x => x.SuiteId, null), filterBuilder.In(x => (TestSuiteId)x.SuiteId!, suiteIds));
+				}
 			}
 
 			if (minChange != null)
