Horde: Fix compound filter grouping
#rnx

diff --git a/Engine/Source/Programs/Horde/Horde.Build/Jobs/TestData/TestDataCollection.cs b/Engine/Source/Programs/Horde/Horde.Build/Jobs/TestData/TestDataCollection.cs
--- a/Engine/Source/Programs/Horde/Horde.Build/Jobs/TestData/TestDataCollection.cs
+++ b/Engine/Source/Programs/Horde/Horde.Build/Jobs/TestData/TestDataCollection.cs
@@ -431,8 +431,8 @@ namespace Horde.Build.Jobs.TestData
 
 			if ((testIds != null && testIds.Length> 0) && (suiteIds != null && suiteIds.Length > 0))
 			{
-				filter &= filterBuilder.Or(filterBuilder.And(filterBuilder.Ne(x => x.TestId, null), filterBuilder.In(x => (TestId)x.TestId!, testIds),
-										   filterBuilder.And(filterBuilder.Ne(x => x.SuiteId, null), filterBuilder.In(x => (TestSuiteId)x.SuiteId!, suiteIds))));
+				filter &= filterBuilder.Or(filterBuilder.And(filterBuilder.Ne(x => x.TestId, null), filterBuilder.In(x => (TestId)x.TestId!, testIds)),
+										   filterBuilder.And(filterBuilder.Ne(x => x.SuiteId, null), filterBuilder.In(x => (TestSuiteId)x.SuiteId!, suiteIds)));
 			}
 			else
 			{
