Horde: Explict test endpoint and adding test details
#rnx

diff --git a/Engine/Source/Programs/Horde/Horde.Build.Tests/TestDataTests.cs b/Engine/Source/Programs/Horde/Horde.Build.Tests/TestDataTests.cs
--- a/Engine/Source/Programs/Horde/Horde.Build.Tests/TestDataTests.cs
+++ b/Engine/Source/Programs/Horde/Horde.Build.Tests/TestDataTests.cs
@@ -18,6 +18,7 @@ using Microsoft.AspNetCore.Mvc;
 using System.Security.Claims;
 using Horde.Build.Logs;
 using System.Linq;
+using System.Text.Json;
 
 namespace Horde.Build.Tests
 {
@@ -207,8 +208,22 @@ namespace Horde.Build.Tests
 			List<GetTestDataRefResponse> refs = refResult.Value;
 
 			Assert.AreEqual(2, refs.Count);
+
+			// Get tests
+			TestRequest request = new TestRequest() { testIds = streams[0].Tests.Select(x => x.Id.ToString()).ToList() };
+			JsonSerializerOptions options = new JsonSerializerOptions();
+			ActionResult<List<GetTestResponse>> testResults = await TestDataController.GetTestsAsync(Convert.ToBase64String(System.Text.Encoding.UTF8.GetBytes(JsonSerializer.Serialize(request, options))));
+			Assert.IsNotNull(testResults);
+			Assert.IsNotNull(testResults.Value);
+			Assert.IsNotNull(testResults.Value[0].Id, streams[0].Tests[0].Id);
 		}
 
+		internal class TestRequest
+		{
+			public List<string>? testIds { get; set; }
+		}
+
+
 		[TestMethod]
 		public async Task SessionReportTest()
 		{
diff --git a/Engine/Source/Programs/Horde/Horde.Build/Jobs/TestData/ITestData.cs b/Engine/Source/Programs/Horde/Horde.Build/Jobs/TestData/ITestData.cs
--- a/Engine/Source/Programs/Horde/Horde.Build/Jobs/TestData/ITestData.cs
+++ b/Engine/Source/Programs/Horde/Horde.Build/Jobs/TestData/ITestData.cs
@@ -218,8 +218,11 @@ namespace Horde.Build.Jobs.TestData
 		int? SuiteErrorCount { get; }
 
 	}
-
-	interface ITestDataDetails
+	
+	/// <summary>
+	/// Test data details
+	/// </summary>
+	public interface ITestDataDetails
 	{
 		/// The corresponding test ref		
 		TestRefId Id { get; }
diff --git a/Engine/Source/Programs/Horde/Horde.Build/Jobs/TestData/ITestDataCollection.cs b/Engine/Source/Programs/Horde/Horde.Build/Jobs/TestData/ITestDataCollection.cs
--- a/Engine/Source/Programs/Horde/Horde.Build/Jobs/TestData/ITestDataCollection.cs
+++ b/Engine/Source/Programs/Horde/Horde.Build/Jobs/TestData/ITestDataCollection.cs
@@ -14,6 +14,7 @@ namespace Horde.Build.Jobs.TestData
 	using TestId = ObjectId<ITest>;
 	using TestSuiteId = ObjectId<ITestSuite>;
 	using TestMetaId = ObjectId<ITestMeta>;
+	using TestRefId = ObjectId<ITestDataRef>;
 
 	/// <summary>
 	/// Collection of test data documents
@@ -41,6 +42,14 @@ namespace Horde.Build.Jobs.TestData
 		/// <returns></returns>
 		Task<List<ITestDataRef>> FindTestRefs(StreamId[] streamIds, TestMetaId[]? metaIds = null, TestId[]? testIds = null, TestSuiteId[]? suiteIds = null, DateTime? minCreateTime = null, DateTime? maxCreateTime = null, int? minChange = null, int? maxChange = null);
 
+
+		/// <summary>
+		/// Find test details
+		/// </summary>
+		/// <param name="ids"></param>
+		/// <returns></returns>
+		Task<List<ITestDataDetails>> FindTestDetails(TestRefId[] ids);
+
 		/// <summary>
 		/// Find tests
 		/// </summary>
diff --git a/Engine/Source/Programs/Horde/Horde.Build/Jobs/TestData/TestDataCollection.cs b/Engine/Source/Programs/Horde/Horde.Build/Jobs/TestData/TestDataCollection.cs
--- a/Engine/Source/Programs/Horde/Horde.Build/Jobs/TestData/TestDataCollection.cs
+++ b/Engine/Source/Programs/Horde/Horde.Build/Jobs/TestData/TestDataCollection.cs
@@ -539,6 +539,25 @@ namespace Horde.Build.Jobs.TestData
 			await _testDataDocuments.DeleteOneAsync<TestDataDocument>(x => x.Id == id);
 		}
 
+		/// <inheritdoc/>
+		public async Task<List<ITestDataDetails>> FindTestDetails(TestRefId[] ids)
+		{
+			FilterDefinitionBuilder<TestDataDetailsDocument> filterBuilder = Builders<TestDataDetailsDocument>.Filter;
+			FilterDefinition<TestDataDetailsDocument> filter = FilterDefinition<TestDataDetailsDocument>.Empty;
+
+			if (ids.Length == 1)
+			{
+				filter &= filterBuilder.Eq(x => x.Id, ids[0]);
+			}
+			else
+			{
+				filter &= filterBuilder.In(x => x.Id, ids);
+			}
+
+			List<TestDataDetailsDocument> results = await _testDetails.Find(filter).ToListAsync();
+			return results.ConvertAll<ITestDataDetails>(x => x);
+		}
+
 		/// <inheritdoc/>
 		public async Task<List<ITestMeta>> FindTestMeta(string[]? projectNames = null, string[]? platforms = null, string[]? configurations = null, string[]? buildTargets = null, string? rhi = null, TestMetaId[]? metaIds = null)
 		{
diff --git a/Engine/Source/Programs/Horde/Horde.Build/Jobs/TestData/TestDataController.cs b/Engine/Source/Programs/Horde/Horde.Build/Jobs/TestData/TestDataController.cs
--- a/Engine/Source/Programs/Horde/Horde.Build/Jobs/TestData/TestDataController.cs
+++ b/Engine/Source/Programs/Horde/Horde.Build/Jobs/TestData/TestDataController.cs
@@ -11,6 +11,10 @@ using Microsoft.AspNetCore.Authorization;
 using Microsoft.AspNetCore.Mvc;
 using MongoDB.Bson;
 using EpicGames.Core;
+using System.Text;
+using System.Text.Json.Serialization;
+using System.Text.Json;
+using System.Threading;
 
 namespace Horde.Build.Jobs.TestData
 {
@@ -20,6 +24,7 @@ namespace Horde.Build.Jobs.TestData
 	using TestId = ObjectId<ITest>;
 	using TestSuiteId = ObjectId<ITestSuite>;
 	using TestMetaId = ObjectId<ITestMeta>;
+	using TestRefId = ObjectId<ITestDataRef>;
 
 	/// <summary>
 	/// Controller for the /api/v1/testdata endpoint
@@ -81,6 +86,51 @@ namespace Horde.Build.Jobs.TestData
 			return responses;
 		}
 
+		/// <summary>
+		/// Get test details from provided refs
+		/// </summary>
+		/// <param name="ids"></param>
+		/// <returns></returns>
+		[HttpGet]
+		[Route("/api/v2/testdata/details")]
+		[ProducesResponseType(typeof(List<GetTestDataDetailsResponse>), 200)]
+		public async Task<ActionResult<List<GetTestDataDetailsResponse>>> GetTestDetailsAsync([FromQuery(Name = "id")] string[] ids)
+		{
+			TestRefId[] idValues = Array.ConvertAll(ids, x => new TestRefId(x));
+			List<ITestDataDetails> details = await _testDataService.FindTestDetails(idValues);
+			return details.Select(d => new GetTestDataDetailsResponse(d)).ToList();
+		}
+
+		internal class TestRequest
+		{
+			public List<string>? testIds { get; set; }
+		}
+
+		/// <summary>
+		/// Get test details from provided refs
+		/// </summary>
+		/// <param name="tests">Base64 encoded tests request (there may be thousands of requested test ids)</param>
+		/// <returns></returns>
+		[HttpGet]
+		[Route("/api/v2/testdata/tests")]
+		[ProducesResponseType(typeof(List<GetTestResponse>), 200)]
+		public async Task<ActionResult<List<GetTestResponse>>> GetTestsAsync([FromQuery] string tests)
+		{
+			JsonSerializerOptions options = new JsonSerializerOptions();
+			TestRequest request = JsonSerializer.Deserialize<TestRequest>(Convert.FromBase64String(tests), options)!;
+
+			if (request.testIds == null || request.testIds.Count == 0)
+			{
+				return new List<GetTestResponse>();
+			}
+
+			HashSet<string> testIds = new HashSet<string>(request.testIds);
+
+			List<ITest> testValues = await _testDataService.FindTests(testIds.Select(x => new TestId(x)).ToArray());			
+
+			return testValues.Select(x => new GetTestResponse(x)).ToList();
+		}
+
 		/// <summary>
 		/// Get stream test data for the provided ids
 		/// </summary>
diff --git a/Engine/Source/Programs/Horde/Horde.Build/Jobs/TestData/TestDataMessages.cs b/Engine/Source/Programs/Horde/Horde.Build/Jobs/TestData/TestDataMessages.cs
--- a/Engine/Source/Programs/Horde/Horde.Build/Jobs/TestData/TestDataMessages.cs
+++ b/Engine/Source/Programs/Horde/Horde.Build/Jobs/TestData/TestDataMessages.cs
@@ -6,12 +6,12 @@ using System.ComponentModel.DataAnnotations;
 using System.Linq;
 using Horde.Build.Streams;
 using Horde.Build.Utilities;
+using MongoDB.Bson;
 using MongoDB.Bson.Serialization;
 
 namespace Horde.Build.Jobs.TestData
 {
-	using StreamId = StringId<IStream>;
-	using TestId = ObjectId<ITest>;
+	using StreamId = StringId<IStream>;	
 
 	/// <summary>
 	/// Response object describing test data to store
@@ -318,6 +318,39 @@ namespace Horde.Build.Jobs.TestData
 		}
 	}
 
+	/// <summary>
+	/// Test details
+	/// </summary>
+	public class GetTestDataDetailsResponse
+	{
+		/// <summary>
+		/// The corresponding test ref
+		/// </summary>
+		public string Id { get; set; }
+
+		/// <summary>
+		/// The test documents for this ref
+		/// </summary>
+		public List<string> TestDataIds { get; set; }
+
+		/// <summary>
+		/// Suite test data
+		/// </summary>		
+		public List<GetSuiteTestDataResponse>? SuiteTests { get; set; }
+
+		/// <summary>
+		/// Constructor
+		/// </summary>
+		/// <param name="details"></param>
+		public GetTestDataDetailsResponse(ITestDataDetails details)
+		{
+			Id = details.Id.ToString();
+			TestDataIds = details.TestDataIds.Select(x => x.ToString()).ToList();
+			SuiteTests = details.SuiteTests?.Select(x => new GetSuiteTestDataResponse(x)).ToList();
+		}
+	}
+
+
 	/// <summary>
 	/// Data ref 
 	/// </summary>
@@ -390,7 +423,10 @@ namespace Horde.Build.Jobs.TestData
 			BuildChangeList = testData.BuildChangeList;
 			MetaId = testData.Metadata.ToString();
 			TestId = testData.TestId?.ToString();
-			Outcome = testData.Outcome;
+			if (testData.TestId != null)
+			{
+				Outcome = testData.Outcome;
+			}			
 			SuiteId = testData.SuiteId?.ToString();
 			SuiteSkipCount = testData.SuiteSkipCount;
 			SuiteWarningCount = testData.SuiteWarningCount;
diff --git a/Engine/Source/Programs/Horde/Horde.Build/Jobs/TestData/TestDataService.cs b/Engine/Source/Programs/Horde/Horde.Build/Jobs/TestData/TestDataService.cs
--- a/Engine/Source/Programs/Horde/Horde.Build/Jobs/TestData/TestDataService.cs
+++ b/Engine/Source/Programs/Horde/Horde.Build/Jobs/TestData/TestDataService.cs
@@ -6,7 +6,6 @@ using System.Threading.Tasks;
 using EpicGames.Core;
 using Horde.Build.Streams;
 using Horde.Build.Utilities;
-using Microsoft.Extensions.Options;
 
 namespace Horde.Build.Jobs.TestData
 {
@@ -15,6 +14,7 @@ namespace Horde.Build.Jobs.TestData
 	using TestId = ObjectId<ITest>;
 	using TestSuiteId = ObjectId<ITestSuite>;
 	using TestMetaId = ObjectId<ITestMeta>;
+	using TestRefId = ObjectId<ITestDataRef>;
 
 	/// <summary>
 	/// Device management service
@@ -103,5 +103,15 @@ namespace Horde.Build.Jobs.TestData
 
 			return await _testData.FindTestRefs(streamIds, metaIds, tids, sids, minCreateTime, maxCreateTime, minChange, maxChange);
 		}
+
+		/// <summary>
+		/// Find test details
+		/// </summary>
+		/// <param name="ids"></param>
+		/// <returns></returns>
+		public async Task<List<ITestDataDetails>> FindTestDetails(TestRefId[] ids)
+		{
+			return await _testData.FindTestDetails(ids);
+		}
 	}
 }
\ No newline at end of file
