Horde: Fix issue with long uri when requests test responses
#rnx

diff --git a/Engine/Source/Programs/Horde/Horde.Build.Tests/TestDataTests.cs b/Engine/Source/Programs/Horde/Horde.Build.Tests/TestDataTests.cs
--- a/Engine/Source/Programs/Horde/Horde.Build.Tests/TestDataTests.cs
+++ b/Engine/Source/Programs/Horde/Horde.Build.Tests/TestDataTests.cs
@@ -18,7 +18,6 @@ using Microsoft.AspNetCore.Mvc;
 using System.Security.Claims;
 using Horde.Build.Logs;
 using System.Linq;
-using System.Text.Json;
 
 namespace Horde.Build.Tests
 {
@@ -209,10 +208,8 @@ namespace Horde.Build.Tests
 
 			Assert.AreEqual(2, refs.Count);
 
-			// Get tests
-			TestRequest request = new TestRequest() { testIds = streams[0].Tests.Select(x => x.Id.ToString()).ToList() };
-			JsonSerializerOptions options = new JsonSerializerOptions();
-			ActionResult<List<GetTestResponse>> testResults = await TestDataController.GetTestsAsync(Convert.ToBase64String(System.Text.Encoding.UTF8.GetBytes(JsonSerializer.Serialize(request, options))));
+			GetTestsRequest request = new GetTestsRequest() { testIds = streams[0].Tests.Select(x => x.Id.ToString()).ToList() };
+			ActionResult<List<GetTestResponse>> testResults = await TestDataController.GetTestsAsync(request);
 			Assert.IsNotNull(testResults);
 			Assert.IsNotNull(testResults.Value);
 			Assert.IsNotNull(testResults.Value[0].Id, streams[0].Tests[0].Id);
diff --git a/Engine/Source/Programs/Horde/Horde.Build/Jobs/TestData/TestDataController.cs b/Engine/Source/Programs/Horde/Horde.Build/Jobs/TestData/TestDataController.cs
--- a/Engine/Source/Programs/Horde/Horde.Build/Jobs/TestData/TestDataController.cs
+++ b/Engine/Source/Programs/Horde/Horde.Build/Jobs/TestData/TestDataController.cs
@@ -101,29 +101,16 @@ namespace Horde.Build.Jobs.TestData
 			return details.Select(d => new GetTestDataDetailsResponse(d)).ToList();
 		}
 
-		internal class TestRequest
-		{
-			public List<string>? testIds { get; set; }
-		}
-
 		/// <summary>
 		/// Get test details from provided refs
 		/// </summary>
-		/// <param name="tests">Base64 encoded tests request (there may be thousands of requested test ids)</param>
+		/// <param name="request"></param>
 		/// <returns></returns>
-		[HttpGet]
+		[HttpPost]
 		[Route("/api/v2/testdata/tests")]
 		[ProducesResponseType(typeof(List<GetTestResponse>), 200)]
-		public async Task<ActionResult<List<GetTestResponse>>> GetTestsAsync([FromQuery] string tests)
+		public async Task<ActionResult<List<GetTestResponse>>> GetTestsAsync([FromBody] GetTestsRequest request)
 		{
-			JsonSerializerOptions options = new JsonSerializerOptions();
-			TestRequest request = JsonSerializer.Deserialize<TestRequest>(Convert.FromBase64String(tests), options)!;
-
-			if (request.testIds == null || request.testIds.Count == 0)
-			{
-				return new List<GetTestResponse>();
-			}
-
 			HashSet<string> testIds = new HashSet<string>(request.testIds);
 
 			List<ITest> testValues = await _testDataService.FindTests(testIds.Select(x => new TestId(x)).ToArray());			
diff --git a/Engine/Source/Programs/Horde/Horde.Build/Jobs/TestData/TestDataMessages.cs b/Engine/Source/Programs/Horde/Horde.Build/Jobs/TestData/TestDataMessages.cs
--- a/Engine/Source/Programs/Horde/Horde.Build/Jobs/TestData/TestDataMessages.cs
+++ b/Engine/Source/Programs/Horde/Horde.Build/Jobs/TestData/TestDataMessages.cs
@@ -210,6 +210,17 @@ namespace Horde.Build.Jobs.TestData
 		}
 	}
 
+	/// <summary>
+	/// Get tests request
+	/// </summary>
+	public class GetTestsRequest
+	{
+		/// <summary>
+		/// Test ids to get
+		/// </summary>
+		public List<string> testIds { get; set; }
+	}
+
 	/// <summary>
 	/// A test suite that runs in a stream, contain subtests
 	/// </summary>
@@ -350,7 +361,6 @@ namespace Horde.Build.Jobs.TestData
 		}
 	}
 
-
 	/// <summary>
 	/// Data ref 
 	/// </summary>
