Horde: Fix warning with test ids
#rnx
#skipci

diff --git a/Engine/Source/Programs/Horde/Horde.Build/Jobs/TestData/TestDataMessages.cs b/Engine/Source/Programs/Horde/Horde.Build/Jobs/TestData/TestDataMessages.cs
--- a/Engine/Source/Programs/Horde/Horde.Build/Jobs/TestData/TestDataMessages.cs
+++ b/Engine/Source/Programs/Horde/Horde.Build/Jobs/TestData/TestDataMessages.cs
@@ -218,7 +218,7 @@ namespace Horde.Build.Jobs.TestData
 		/// <summary>
 		/// Test ids to get
 		/// </summary>
-		public List<string> testIds { get; set; }
+		public List<string> testIds { get; set; } = new List<string>();
 	}
 
 	/// <summary>
