Horde: Add a clickable history icon to automation device rows to facilitate being able to select various text in the row, such as ips
#rnx
#skipci

diff --git a/Engine/Source/Programs/Horde/HordeDashboard/src/components/DeviceView.tsx b/Engine/Source/Programs/Horde/HordeDashboard/src/components/DeviceView.tsx
--- a/Engine/Source/Programs/Horde/HordeDashboard/src/components/DeviceView.tsx
+++ b/Engine/Source/Programs/Horde/HordeDashboard/src/components/DeviceView.tsx
@@ -79,7 +79,6 @@ const customStyles = mergeStyleSets({
             cursor: "default"
          },
          '.ms-GroupHeader-expand,.ms-GroupHeader-expand:hover': {
-            cursor: "pointer",
             background: "#DFDEDD"
          },
          '.ms-DetailsRow': {
@@ -87,7 +86,6 @@ const customStyles = mergeStyleSets({
             background: "unset"
          },
          '.ms-DetailsRow:hover': {
-            cursor: "pointer",
             background: "#F3F2F1"
          }
       },
@@ -303,25 +301,36 @@ const DevicePanel: React.FC = observer(() => {
       // Automation Kit
 
       const r = handler.getReservation(device);
-
-      if (!r) {
-         return null;
-      }
-
+      
       let url = "";
 
-      if (r.jobId) {
+      if (r?.jobId) {
          url = `/job/${r.jobId}`;
          if (r.stepId) {
             url += `?step=${r.stepId}`;
          }
       }
       return <Stack verticalFill={true} verticalAlign="center">
-         <Stack>
-            {!!url && <Link to={url}><Text variant="small">Job Details</Text></Link>}
-            {!url && !!r.reservationDetails && r.reservationDetails.startsWith("https://horde") && <FluentLink href={r.reservationDetails} target="_blank" underline={false}><Text variant="small" onClick={(ev) => { ev.stopPropagation(); }}>Job Details</Text></FluentLink>}
-            <Text variant="tiny">Host: {r.hostname}</Text>
-            <Text variant="tiny">Duration: {Math.floor(moment.duration(moment.utc().diff(moment(r.createTimeUtc))).asMinutes())}m</Text>
+         <Stack horizontal tokens={{childrenGap: 12}}>
+            {automationTab && <Stack verticalFill={true} verticalAlign="center">
+               <div style={{ cursor: "pointer" }} onClick={(ev) => {
+                  ev.stopPropagation(); ev.preventDefault();
+                  setEditState({
+                     infoShown: true, device: item
+                  })
+               }}>
+                  <Stack horizontal tokens={{ childrenGap: 18 }}>
+                     <Icon style={{ paddingTop: 2 }} iconName="History" />
+                  </Stack>
+               </div>
+
+            </Stack>}
+            {!!r && <Stack>
+               {!!url && <Link to={url}><Text variant="small">Job Details</Text></Link>}
+               {!url && !!r.reservationDetails && r.reservationDetails.startsWith("https://horde") && <FluentLink href={r.reservationDetails} target="_blank" underline={false}><Text variant="small" onClick={(ev) => { ev.stopPropagation(); }}>Job Details</Text></FluentLink>}
+               <Text variant="tiny">Host: {r.hostname}</Text>
+               <Text variant="tiny">Duration: {Math.floor(moment.duration(moment.utc().diff(moment(r.createTimeUtc))).asMinutes())}m</Text>
+            </Stack>}
          </Stack>
       </Stack>
 
@@ -435,21 +444,6 @@ const DevicePanel: React.FC = observer(() => {
 
    groups = newGroups;
 
-   const onRenderRow: IDetailsListProps['onRenderRow'] = (props) => {
-
-      if (props) {
-
-         return <div style={{ cursor: "pointer" }} onClick={(ev) => {
-            ev.stopPropagation(); ev.preventDefault();
-            setEditState({ infoShown: true, device: props.item })
-         }}>
-            <DetailsRow {...props} />
-         </div>
-      }
-
-      return null;
-   };
-
    const pivotItems: JSX.Element[] = [];
 
    pivotItems.push(<PivotItem headerText="Shared" itemKey={pivotKeyShared} key={pivotKeyShared} />);
@@ -557,7 +551,6 @@ const DevicePanel: React.FC = observer(() => {
                         layoutMode={DetailsListLayoutMode.justified}
                         compact={true}
                         onShouldVirtualize={() => false}
-                        onRenderRow={onRenderRow}
                      />
                   </ScrollablePane>
                </div>
@@ -664,7 +657,7 @@ export const CheckoutConfirmModal: React.FC<{ check: "in" | "out" | "error", dev
 
 
 export const DeviceView: React.FC = () => {
-   
+
    const windowSize = useWindowSize();
    const vw = Math.max(document.documentElement.clientWidth, window.innerWidth || 0);
 
@@ -676,7 +669,7 @@ export const DeviceView: React.FC = () => {
          <Stack tokens={{ childrenGap: 0 }} styles={{ root: { backgroundColor: modeColors.background, width: "100%" } }}>
             <Stack style={{ maxWidth: 1800, paddingTop: 6, marginLeft: 4, height: 'calc(100vh - 8px)' }}>
                <Stack horizontal className={hordeClasses.raised}>
-                  <Stack style={{ width: "100%", height: 'calc(100vh - 228px)'}} tokens={{ childrenGap: 18 }}>
+                  <Stack style={{ width: "100%", height: 'calc(100vh - 228px)' }} tokens={{ childrenGap: 18 }}>
                      <DevicePanel />
                   </Stack>
                </Stack>
