Horde: TestCollection fixes
#rnx

diff --git a/Engine/Source/Programs/Horde/Horde.Build.Tests/TestDataTests.cs b/Engine/Source/Programs/Horde/Horde.Build.Tests/TestDataTests.cs
--- a/Engine/Source/Programs/Horde/Horde.Build.Tests/TestDataTests.cs
+++ b/Engine/Source/Programs/Horde/Horde.Build.Tests/TestDataTests.cs
@@ -166,25 +166,38 @@ namespace Horde.Build.Tests
 			IJobStep step = job.Batches[0].Steps[0];
 			IJob job2 = CreateJob(_releaseStreamId, 105, "Test Build", _graph);
 			IJobStep step2 = job2.Batches[0].Steps[0];
+			IJob job3 = CreateJob(_releaseStreamId, 105, "Test Build", _graph);
+			IJobStep step3 = job3.Batches[0].Steps[0];
 
-			BsonDocument testData = BsonDocument.Parse(String.Join('\n', _simpleTestDataLines));
+			BsonDocument testData1 = BsonDocument.Parse(String.Join('\n', _simpleTestDataLines));
 
-			List<(string key, BsonDocument)> data = new List<(string key, BsonDocument)>();
-			BsonArray items = testData.GetValue("Items").AsBsonArray;
+			List<(string key, BsonDocument)> data1 = new List<(string key, BsonDocument)>();
+			BsonArray items = testData1.GetValue("Items").AsBsonArray;
 			foreach (BsonValue item in items.ToList())
 			{
 				BsonDocument value = item.AsBsonDocument.GetValue("Data").AsBsonDocument;
-				data.Add(("Simple Report Key", value));
+				data1.Add(("Simple Report Key", value));
 			}
 
-			await TestDataCollection.AddAsync(job, step, data.ToArray());
-			await TestDataCollection.AddAsync(job2, step2, data.ToArray());
+			BsonDocument testData2 = BsonDocument.Parse(String.Join('\n', _simpleTestDataLines2));
+
+			List<(string key, BsonDocument)> data2 = new List<(string key, BsonDocument)>();
+			items = testData2.GetValue("Items").AsBsonArray;
+			foreach (BsonValue item in items.ToList())
+			{
+				BsonDocument value = item.AsBsonDocument.GetValue("Data").AsBsonDocument;
+				data2.Add(("Simple Report Key", value));
+			}
+
+			await TestDataCollection.AddAsync(job, step, data1.ToArray());
+			await TestDataCollection.AddAsync(job2, step2, data1.ToArray());
+			await TestDataCollection.AddAsync(job3, step3, data2.ToArray());
 
 			ActionResult<List<GetTestMetaResponse>> metaResult = await TestDataController.GetTestMetaAsync();
 			Assert.IsNotNull(metaResult);
 			Assert.IsNotNull(metaResult.Value);
 			List<GetTestMetaResponse> meta = metaResult.Value;
-			Assert.AreEqual(1, meta.Count);
+			Assert.AreEqual(2, meta.Count);
 
 			ActionResult<List<GetTestStreamResponse>> streamResult = await TestDataController.GetTestStreamsAsync(streamIds);
 			Assert.IsNotNull(streamResult);
@@ -193,10 +206,10 @@ namespace Horde.Build.Tests
 
 			Assert.AreEqual(2, streams.Count);
 			Assert.AreEqual(1, streams[0].Tests.Count);
-			Assert.AreEqual(1, streams[0].TestMetadata.Count);
+			Assert.AreEqual(2, streams[0].TestMetadata.Count);
 			Assert.AreEqual(0, streams[0].TestSuites.Count);
 			Assert.AreEqual(1, streams[1].Tests.Count);
-			Assert.AreEqual(1, streams[1].TestMetadata.Count);
+			Assert.AreEqual(2, streams[1].TestMetadata.Count);
 			Assert.AreEqual(0, streams[1].TestSuites.Count);
 
 			Assert.AreEqual(streams[0].TestMetadata[0].Id, streams[1].TestMetadata[0].Id);
@@ -206,7 +219,7 @@ namespace Horde.Build.Tests
 			Assert.IsNotNull(refResult.Value);
 			List<GetTestDataRefResponse> refs = refResult.Value;
 
-			Assert.AreEqual(2, refs.Count);
+			Assert.AreEqual(3, refs.Count);
 
 			GetTestsRequest request = new GetTestsRequest() { testIds = streams[0].Tests.Select(x => x.Id.ToString()).ToList() };
 			ActionResult<List<GetTestResponse>> testResults = await TestDataController.GetTestsAsync(request);
@@ -230,6 +243,8 @@ namespace Horde.Build.Tests
 			IJobStep step = job.Batches[0].Steps[0];
 			IJob job2 = CreateJob(_releaseStreamId, 105, "Test Build", _graph);
 			IJobStep step2 = job2.Batches[0].Steps[0];
+			IJob job3 = CreateJob(_releaseStreamId, 105, "Test Build", _graph);
+			IJobStep step3 = job3.Batches[0].Steps[0];
 
 			BsonDocument testData = BsonDocument.Parse(String.Join('\n', _testSessionDataLines));
 
@@ -241,18 +256,36 @@ namespace Horde.Build.Tests
 				data.Add(("Session Report Key", value));
 			}
 
+			List<string> dataLines2 = new List<string>();
+			for (int i = 0; i < _testSessionDataLines.Length; i++)
+			{
+				dataLines2.Add(_testSessionDataLines[i].Replace("Win64", "Android", StringComparison.OrdinalIgnoreCase));
+			}
+
+			BsonDocument testData2 = BsonDocument.Parse(String.Join('\n', dataLines2));
+
+			List<(string key, BsonDocument)> data2 = new List<(string key, BsonDocument)>();
+			items = testData2.GetValue("Items").AsBsonArray;
+			foreach (BsonValue item in items.ToList())
+			{
+				BsonDocument value = item.AsBsonDocument.GetValue("Data").AsBsonDocument;
+				data2.Add(("Session Report Key", value));
+			}
+
 			await TestDataCollection.AddAsync(job, step, data.ToArray());
 			await TestDataCollection.AddAsync(job2, step2, data.ToArray());
+			await TestDataCollection.AddAsync(job3, step3, data2.ToArray());
 
 			ActionResult<List<GetTestMetaResponse>> metaResult = await TestDataController.GetTestMetaAsync();
 			Assert.IsNotNull(metaResult);
 			Assert.IsNotNull(metaResult.Value);
 			List<GetTestMetaResponse> meta = metaResult.Value;
-			Assert.AreEqual(1, meta.Count);
+			Assert.AreEqual(2, meta.Count);
 			Assert.AreEqual(1, meta[0].BuildTargets.Count);
 			Assert.AreEqual("Client", meta[0].BuildTargets[0]);
 			Assert.AreEqual(1, meta[0].Platforms.Count);
 			Assert.AreEqual("Win64", meta[0].Platforms[0]);
+			Assert.AreEqual("Android", meta[1].Platforms[0]);
 			Assert.AreEqual(1, meta[0].Configurations.Count);
 			Assert.AreEqual("Development", meta[0].Configurations[0]);
 			Assert.AreEqual("EngineTest", meta[0].ProjectName);
@@ -266,11 +299,11 @@ namespace Horde.Build.Tests
 			Assert.AreEqual(2, streams.Count);
 
 			Assert.AreEqual(0, streams[0].Tests.Count);				
-			Assert.AreEqual(1, streams[0].TestMetadata.Count);
+			Assert.AreEqual(2, streams[0].TestMetadata.Count);
 			Assert.AreEqual(1, streams[0].TestSuites.Count);
 
 			Assert.AreEqual(0, streams[1].Tests.Count);
-			Assert.AreEqual(1, streams[1].TestMetadata.Count);
+			Assert.AreEqual(2, streams[1].TestMetadata.Count);
 			Assert.AreEqual(1, streams[1].TestSuites.Count);			
 
 			//Assert.AreEqual(streams[0].TestMetadata[0].Id, streams[1].TestMetadata[0].Id);
@@ -280,7 +313,7 @@ namespace Horde.Build.Tests
 			Assert.IsNotNull(refResult.Value);
 			List<GetTestDataRefResponse> refs = refResult.Value;
 
-			Assert.AreEqual(2, refs.Count);
+			Assert.AreEqual(3, refs.Count);
 
 			Assert.AreEqual(1, refs[0].SuiteSkipCount);
 			Assert.AreEqual(1, refs[0].SuiteWarningCount);
@@ -332,6 +365,51 @@ namespace Horde.Build.Tests
 			@"}"
 		};
 
+		private readonly string[] _simpleTestDataLines2 =
+		{
+			@"{",
+			@"  ""Items"": [",
+			@"    {",
+			@"      ""Key"": ""Simple Report::UE.BootTest EngineTest Editor Win64"",",
+			@"      ""Data"": {",
+			@"  	  ""Version"" : 1,",
+			@"        ""Type"": ""Simple Report"",",
+			@"        ""TestName"": ""EditorBootTest"",",
+			@"        ""Description"": ""Win64 Development EditorGame"",",
+			@"        ""ReportCreatedOn"": ""10/31/2022 11:44:56 AM"",",
+			@"        ""TotalDurationSeconds"": 35.33694,",
+			@"        ""HasSucceeded"": true,",
+			@"        ""Status"": ""Passed"",",
+			@"        ""URLLink"": """",",
+			@"        ""BuildChangeList"": 22815797,",
+			@"        ""MainRole"": {",
+			@"          ""Type"": ""Client"",",
+			@"          ""Platform"": ""Android"",",
+			@"          ""Configuration"": ""Development""",
+			@"        },",
+			@"        ""Roles"": [",
+			@"          {",
+			@"            ""Type"": ""Client"",",
+			@"            ""Platform"": ""Android"",",
+			@"            ""Configuration"": ""Development""",
+			@"          }",
+			@"        ],",
+			@"        ""TestResult"": ""Passed"",",
+			@"        ""Logs"": [],",
+			@"        ""Errors"": [],",
+			@"        ""Warnings"": [],",
+			@"        ""Metadata"": {",
+			@"          ""Platform"": ""Android"",",
+			@"          ""BuildTarget"": ""Client"",",
+			@"          ""Configuration"": ""Development"",",
+			@"          ""Project"": ""EngineTest""",
+			@"        }",
+			@"      }",
+			@"    }",
+			@"  ],",
+			@"}"
+		};
+
 		private readonly string[] _testSessionDataLines =
 		{
 			@"{",			
diff --git a/Engine/Source/Programs/Horde/Horde.Build/Jobs/TestData/TestDataCollection.cs b/Engine/Source/Programs/Horde/Horde.Build/Jobs/TestData/TestDataCollection.cs
--- a/Engine/Source/Programs/Horde/Horde.Build/Jobs/TestData/TestDataCollection.cs
+++ b/Engine/Source/Programs/Horde/Horde.Build/Jobs/TestData/TestDataCollection.cs
@@ -723,31 +723,41 @@ namespace Horde.Build.Jobs.TestData
 			return await _tests.Find(filter).ToListAsync();
 		}
 
-		private async Task<List<TestSuiteDocument>> FindTestSuites(string[] suiteNames, TestMetaId[] metaIds)
+		private async Task<List<TestSuiteDocument>> FindTestSuites(string[] suiteNames, TestMetaId[]? metaIds = null, string[]? projectNames = null)
 		{
+			List<ITestMeta> metaData = new List<ITestMeta>();
+
+			if ((projectNames != null && projectNames.Length > 0) || (metaIds != null && metaIds.Length > 0))
+			{
+				metaData = await FindTestMeta(projectNames, metaIds: metaIds);
+			}
+
+			TestMetaId[] queryIds = metaData.Select(x => x.Id).ToArray();
 
 			FilterDefinitionBuilder<TestSuiteDocument> filterBuilder = Builders<TestSuiteDocument>.Filter;
 			FilterDefinition<TestSuiteDocument> filter = FilterDefinition<TestSuiteDocument>.Empty;
 
-			if (suiteNames.Length == 1)
-			{
-				filter &= filterBuilder.Eq(x => x.Name, suiteNames[0]);
-			}
-			else
+			if (queryIds.Length > 0)
 			{
-				filter &= filterBuilder.In(x => x.Name, suiteNames);
+				if (queryIds.Length == 1)
+				{
+					filter &= filterBuilder.AnyEq(x => x.Metadata, queryIds[0]);
+				}
+				else
+				{
+					filter &= filterBuilder.AnyIn(x => x.Metadata, queryIds);
+				}
 			}
 
-			if (metaIds.Length == 1)
+			if (suiteNames.Length == 1)
 			{
-				filter &= filterBuilder.AnyEq(x => x.Metadata, metaIds[0]);
+				filter &= filterBuilder.Eq(x => x.Name, suiteNames[0]);
 			}
 			else
 			{
-				filter &= filterBuilder.AnyIn(x => x.Metadata, metaIds);
+				filter &= filterBuilder.In(x => x.Name, suiteNames);
 			}
 
-
 			return await _testSuites.Find(filter).ToListAsync();
 		}
 
@@ -822,7 +832,7 @@ namespace Horde.Build.Jobs.TestData
 		/// <param name="step"></param>
 		/// <param name="testName"></param>
 		/// <returns></returns>
-		private async Task<TestMetaId?> AddTestMeta(Dictionary<string, string> metaData, IJob job, IJobStep step, string testName)
+		private async Task<ITestMeta?> AddTestMeta(Dictionary<string, string> metaData, IJob job, IJobStep step, string testName)
 		{
 			string? platform;
 			if (!metaData.TryGetValue("Platform", out platform))
@@ -879,43 +889,43 @@ namespace Horde.Build.Jobs.TestData
 				await _testMeta.InsertOneAsync(result);
 			}
 			
-			return result.Id;
+			return result;
 
 		}
 
-		async Task<TestDocument?> AddOrUpdateTest(TestMetaId metaId, string testName, string? suiteName = null,  string? displayName = null)
+		async Task<TestDocument?> AddOrUpdateTest(ITestMeta metaData, string testName, string? suiteName = null,  string? displayName = null)
 		{
 
 			List<TestDocument> tests;
 
 			if (suiteName != null)
 			{
-				tests = await FindTests(true, metaIds: new TestMetaId[] { metaId }, testNames: new string[] { testName }, suiteNames: new string[] {suiteName});
+				tests = await FindTests(true, projectNames: new string[] { metaData.ProjectName }, testNames: new string[] { testName }, suiteNames: new string[] {suiteName});
 			}
 			else
 			{
-				tests = await FindTests(false, metaIds: new TestMetaId[] { metaId }, testNames: new string[] { testName });
+				tests = await FindTests(false, projectNames: new string[] { metaData.ProjectName }, testNames: new string[] { testName });
 			}
 
 			if (tests.Count > 1)
 			{
-				throw new Exception($"Duplicate tests found for MetaId: {metaId}, TestName: {testName}, SuiteName: {suiteName}, Count: {tests.Count}");
+				throw new Exception($"Duplicate tests found for MetaId: {metaData.Id}, TestName: {testName}, SuiteName: {suiteName}, Count: {tests.Count}");
 			}
 
 			TestDocument? test;
 
 			if (tests.Count == 0)
 			{
-				test = new TestDocument( testName, new List<TestMetaId>() { metaId }, suiteName, displayName);
+				test = new TestDocument( testName, new List<TestMetaId>() { metaData.Id }, suiteName, displayName);
 				await _tests.InsertOneAsync(test);
 			}
 			else
 			{
 				test = tests[0];
 
-				if (!test.Metadata.Contains(metaId))
+				if (!test.Metadata.Contains(metaData.Id))
 				{
-					test.Metadata.Add(metaId);
+					test.Metadata.Add(metaData.Id);
 
 					FilterDefinitionBuilder<TestDocument> ebuilder = Builders<TestDocument>.Filter;
 					FilterDefinition<TestDocument> efilter = ebuilder.Eq(x => x.Id, test.Id);
@@ -929,22 +939,21 @@ namespace Horde.Build.Jobs.TestData
 			return test;
 		}
 
-		private async Task<TestSuiteDocument?> AddOrUpdateTestSuite(string suiteName, TestMetaId metaId, List<TestId> testIds)
+		private async Task<TestSuiteDocument?> AddOrUpdateTestSuite(string suiteName, ITestMeta metaData, List<TestId> testIds)
 		{
-
 			TestSuiteDocument? suite;
 
-			List<TestSuiteDocument> suiteTests = await FindTestSuites(suiteNames: new string[] { suiteName }, metaIds: new TestMetaId[] {metaId});
+			List<TestSuiteDocument> suiteTests = await FindTestSuites(suiteNames: new string[] { suiteName }, projectNames: new string[] {metaData.ProjectName});
 
 			if (suiteTests.Count > 1)
 			{
-				throw new Exception($"Duplicate tests found for MetaId: {metaId}, SuiteName: {suiteName}, Count: {suiteTests.Count}");
+				throw new Exception($"Duplicate tests found for MetaId: {metaData.Id}, SuiteName: {suiteName}, Count: {suiteTests.Count}");
 			}
 
 
 			if (suiteTests.Count == 0)
 			{
-				suite = new TestSuiteDocument(suiteName, new List<TestMetaId> { metaId }, testIds);
+				suite = new TestSuiteDocument(suiteName, new List<TestMetaId> { metaData.Id }, testIds);
 				await _testSuites.InsertOneAsync(suite);
 			}
 			else
@@ -953,7 +962,7 @@ namespace Horde.Build.Jobs.TestData
 
 				List<TestId> newIds = suite.Tests.Union(testIds).ToList();
 
-				if (newIds.Count != suite.Tests.Count || !suite.Metadata.Contains(metaId))
+				if (newIds.Count != suite.Tests.Count || !suite.Metadata.Contains(metaData.Id))
 				{
 					List<UpdateDefinition<TestSuiteDocument>> updates = new List<UpdateDefinition<TestSuiteDocument>>();
 					FilterDefinitionBuilder<TestSuiteDocument> ebuilder = Builders<TestSuiteDocument>.Filter;
@@ -967,9 +976,9 @@ namespace Horde.Build.Jobs.TestData
 						updates.Add(updateBuilder.Set(x => x.Tests, suite.Tests));
 					}
 
-					if (!suite.Metadata.Contains(metaId))
+					if (!suite.Metadata.Contains(metaData.Id))
 					{
-						suite.Metadata.Add(metaId);
+						suite.Metadata.Add(metaData.Id);
 						updates.Add(updateBuilder.Set(x => x.Metadata, suite.Metadata));
 					}
 
@@ -1123,14 +1132,14 @@ namespace Horde.Build.Jobs.TestData
 			{
 				SimpleTestData testData = BsonSerializer.Deserialize<SimpleTestData>(testDoc);
 
-				TestMetaId? metaId = await AddTestMeta(testData.Metadata, job, step, testData.TestName);
+				ITestMeta? metaData = await AddTestMeta(testData.Metadata, job, step, testData.TestName);
 
-				if (metaId == null)
+				if (metaData == null)
 				{
 					throw new Exception($"Unable to add or update simple test report for {testData.TestName}, unable to generate meta data");
 				}
 
-				TestDocument? test = await AddOrUpdateTest(metaId.Value, testData.TestName);
+				TestDocument? test = await AddOrUpdateTest(metaData!, testData.TestName);
 
 				if (test == null)
 				{
@@ -1139,7 +1148,7 @@ namespace Horde.Build.Jobs.TestData
 
 				TestDataRefDocument testRef = new TestDataRefDocument(job.StreamId);
 				testRef.StreamId = job.StreamId;
-				testRef.Metadata = metaId.Value;
+				testRef.Metadata = metaData.Id;
 				testRef.BuildChangeList = testData.BuildChangeList <= 0 ? job.Change : testData.BuildChangeList;
 				testRef.Duration = TimeSpan.FromSeconds(testData.TotalDurationSeconds);
 				testRef.TestId = test.Id;
@@ -1214,9 +1223,9 @@ namespace Horde.Build.Jobs.TestData
 						continue;
 					}
 
-					TestMetaId? metaId = await AddTestMeta(session.Metadata, job, step, "Automated Test Session");
+					ITestMeta? metaData = await AddTestMeta(session.Metadata, job, step, "Automated Test Session");
 
-					if (metaId == null)
+					if (metaData == null)
 					{
 						throw new Exception($"Unable to add or update simple test report for automated test session, unable to generate meta data");
 					}
@@ -1248,11 +1257,11 @@ namespace Horde.Build.Jobs.TestData
 								displayName = utest.TestDisplayName;
 							}
 
-							TestDocument? testDoc = await AddOrUpdateTest(metaId!.Value, test.Name, suite, displayName);
+							TestDocument? testDoc = await AddOrUpdateTest(metaData, test.Name, suite, displayName);
 
 							if (testDoc == null)
 							{
-								_logger.LogWarning("Skipping adding suite test {TestName} metaId: {MetaId} jobId: {JobId} stepId: {StepId}", test.Name, metaId!.Value, job.Id, step.Id);
+								_logger.LogWarning("Skipping adding suite test {TestName} metaId: {MetaId} jobId: {JobId} stepId: {StepId}", test.Name, metaData.Id, job.Id, step.Id);
 								continue;
 							}
 
@@ -1264,7 +1273,7 @@ namespace Horde.Build.Jobs.TestData
 					// register test suite
 					if (suiteTestIds.Count > 0)
 					{
-						TestSuiteDocument? testSuite = await AddOrUpdateTestSuite(suite, metaId.Value, suiteTestIds.ToList());
+						TestSuiteDocument? testSuite = await AddOrUpdateTestSuite(suite, metaData, suiteTestIds.ToList());
 
 						if (testSuite == null)
 						{
@@ -1273,7 +1282,7 @@ namespace Horde.Build.Jobs.TestData
 
 						TestDataRefDocument testRef = new TestDataRefDocument(job.StreamId);
 						testRef.StreamId = job.StreamId;
-						testRef.Metadata = metaId.Value;
+						testRef.Metadata = metaData.Id;
 						testRef.BuildChangeList = job.Change;
 						testRef.Duration = TimeSpan.FromSeconds(session.TestSessionInfo!.TimeElapseSec);
 						testRef.Outcome = TestOutcome.Unspecified;
