Horde: Allow comments in JSON config files.

diff --git a/Engine/Source/Programs/Horde/Horde.Build/Startup.cs b/Engine/Source/Programs/Horde/Horde.Build/Startup.cs
--- a/Engine/Source/Programs/Horde/Horde.Build/Startup.cs
+++ b/Engine/Source/Programs/Horde/Horde.Build/Startup.cs
@@ -725,6 +725,7 @@ namespace Horde.Build
 
 		public static void ConfigureJsonSerializer(JsonSerializerOptions options)
 		{
+			options.ReadCommentHandling = JsonCommentHandling.Skip;
 			options.DefaultIgnoreCondition = JsonIgnoreCondition.WhenWritingNull;
 			options.PropertyNamingPolicy = JsonNamingPolicy.CamelCase;
 			options.PropertyNameCaseInsensitive = true;
