Horde: Configure a default storage backend, which writes to C:\ProgramData\HordeServer\Storage.

diff --git a/Engine/Source/Programs/Horde/Horde.Build/Defaults/globals.json b/Engine/Source/Programs/Horde/Horde.Build/Defaults/globals.json
--- a/Engine/Source/Programs/Horde/Horde.Build/Defaults/globals.json
+++ b/Engine/Source/Programs/Horde/Horde.Build/Defaults/globals.json
@@ -12,9 +12,17 @@
 		}
 	],
 	"storage": {
+		"backends": [
+			{
+				"id": "default-backend",
+				"type":"FileSystem",
+				"baseDir": "Storage" // Default base directory is C:\ProgramData\HordeServer
+			}
+		],
 		"namespaces": [
 			{
-				"id": "default"
+				"id": "default",
+				"backend": "default-backend"
 			}
 		]
 	},
