Horde: Abstract the construction of logger instances behind a factory, making it easier to adjust implementation details without having to modify all call sites.

diff --git a/Engine/Source/Programs/Horde/Horde.Agent.Tests/JsonRpcLoggerTests.cs b/Engine/Source/Programs/Horde/Horde.Agent.Tests/JsonRpcLoggerTests.cs
--- a/Engine/Source/Programs/Horde/Horde.Agent.Tests/JsonRpcLoggerTests.cs
+++ b/Engine/Source/Programs/Horde/Horde.Agent.Tests/JsonRpcLoggerTests.cs
@@ -126,6 +126,8 @@ namespace Horde.Agent.Tests
 
 		class FakeLogSink : IJsonRpcLogSink
 		{
+			public ValueTask DisposeAsync() => new ValueTask();
+
 			public Task SetOutcomeAsync(JobStepOutcome outcome, CancellationToken cancellationToken) => Task.CompletedTask;
 
 			public Task WriteEventsAsync(List<CreateEventRequest> events, CancellationToken cancellationToken) => Task.CompletedTask;
@@ -141,7 +143,7 @@ namespace Horde.Agent.Tests
 
 			TreeReader reader = new TreeReader(store, cache, NullLogger.Instance);
 
-			FakeLogSink innerSink = new FakeLogSink();
+			await using FakeLogSink innerSink = new FakeLogSink();
 
 			const int Count = 20000;
 
diff --git a/Engine/Source/Programs/Horde/Horde.Agent.Tests/WorkerServiceTest.cs b/Engine/Source/Programs/Horde/Horde.Agent.Tests/WorkerServiceTest.cs
--- a/Engine/Source/Programs/Horde/Horde.Agent.Tests/WorkerServiceTest.cs
+++ b/Engine/Source/Programs/Horde/Horde.Agent.Tests/WorkerServiceTest.cs
@@ -16,6 +16,7 @@ using Grpc.Net.Client;
 using Horde.Agent.Execution;
 using Horde.Agent.Leases;
 using Horde.Agent.Leases.Handlers;
+using Horde.Agent.Parser;
 using Horde.Agent.Services;
 using Horde.Agent.Utility;
 using HordeCommon;
@@ -26,7 +27,6 @@ using Microsoft.Extensions.DependencyInjection;
 using Microsoft.Extensions.Logging;
 using Microsoft.Extensions.Logging.Abstractions;
 using Microsoft.Extensions.Options;
-using Microsoft.VisualStudio.TestPlatform.ObjectModel.DataCollection;
 using Microsoft.VisualStudio.TestTools.UnitTesting;
 using Moq;
 
@@ -37,6 +37,26 @@ namespace Horde.Agent.Tests
 	{
 		private readonly ServiceCollection _serviceCollection;
 
+		class FakeServerLogger : IServerLogger
+		{
+			public JobStepOutcome Outcome => JobStepOutcome.Success;
+
+			public IDisposable BeginScope<TState>(TState state) => NullLogger.Instance.BeginScope<TState>(state);
+
+			public ValueTask DisposeAsync() => new ValueTask();
+
+			public bool IsEnabled(LogLevel logLevel) => true;
+
+			public void Log<TState>(LogLevel logLevel, EventId eventId, TState state, Exception? exception, Func<TState, Exception?, string> formatter) { }
+
+			public Task StopAsync() => Task.CompletedTask;
+		}
+
+		class FakeServerLoggerFactory : IServerLoggerFactory
+		{
+			public IServerLogger CreateLogger(ISession session, string logId, string? jobId, string? batchId, string? stepId, bool? warnings = null) => new FakeServerLogger();
+		}
+
 		internal static JobExecutor NullExecutor = new SimpleTestExecutor(async (step, logger, cancellationToken) =>
 		{
 			await Task.Delay(1, cancellationToken);
@@ -47,6 +67,7 @@ namespace Horde.Agent.Tests
 		{
 			_serviceCollection = new ServiceCollection();
 			_serviceCollection.AddLogging();
+			_serviceCollection.AddSingleton<IServerLoggerFactory, FakeServerLoggerFactory>();
 
 			_serviceCollection.Configure<AgentSettings>(settings =>
 			{
@@ -147,17 +168,13 @@ namespace Horde.Agent.Tests
 				token)).Outcome;
 
 			Assert.AreEqual(LeaseOutcome.Success, outcome);
-			Assert.AreEqual(4, client.UpdateStepRequests.Count);
-			// An extra UpdateStep request is sent by JsonRpcLogger on failures which is why four requests
-			// are returned instead of three.
+			Assert.AreEqual(3, client.UpdateStepRequests.Count);
 			Assert.AreEqual(JobStepOutcome.Success, client.UpdateStepRequests[0].Outcome);
 			Assert.AreEqual(JobStepState.Completed, client.UpdateStepRequests[0].State);
 			Assert.AreEqual(JobStepOutcome.Failure, client.UpdateStepRequests[1].Outcome);
-			Assert.AreEqual(JobStepState.Unspecified, client.UpdateStepRequests[1].State);
-			Assert.AreEqual(JobStepOutcome.Failure, client.UpdateStepRequests[2].Outcome);
-			Assert.AreEqual(JobStepState.Aborted, client.UpdateStepRequests[2].State);
-			Assert.AreEqual(JobStepOutcome.Success, client.UpdateStepRequests[3].Outcome);
-			Assert.AreEqual(JobStepState.Completed, client.UpdateStepRequests[3].State);
+			Assert.AreEqual(JobStepState.Aborted, client.UpdateStepRequests[1].State);
+			Assert.AreEqual(JobStepOutcome.Success, client.UpdateStepRequests[2].Outcome);
+			Assert.AreEqual(JobStepState.Completed, client.UpdateStepRequests[2].State);
 		}
 		
 		[TestMethod]
@@ -239,7 +256,7 @@ namespace Horde.Agent.Tests
 
 		public static ISession CreateSession(IRpcConnection rpcConnection)
 		{
-			Mock<ISession> fakeSession = new Mock<ISession>();
+			Mock<ISession> fakeSession = new Mock<ISession>(MockBehavior.Strict);
 			fakeSession.Setup(x => x.AgentId).Returns("LocalAgent");
 			fakeSession.Setup(x => x.SessionId).Returns("Session");
 			fakeSession.Setup(x => x.RpcConnection).Returns(rpcConnection);
diff --git a/Engine/Source/Programs/Horde/Horde.Agent/Commands/Utilities/BuildGraphCommand.cs b/Engine/Source/Programs/Horde/Horde.Agent/Commands/Utilities/BuildGraphCommand.cs
--- a/Engine/Source/Programs/Horde/Horde.Agent/Commands/Utilities/BuildGraphCommand.cs
+++ b/Engine/Source/Programs/Horde/Horde.Agent/Commands/Utilities/BuildGraphCommand.cs
@@ -23,7 +23,7 @@ namespace Horde.Agent.Commands.Utilities
 	[Command("BuildGraph", "Executes a BuildGraph script with the given arguments using a build of UAT within the current branch, and runs the output through the log processor")]
 	class BuildGraphCommand : Command
 	{
-		sealed class LogSink : IJsonRpcLogSink, IDisposable
+		sealed class LogSink : IJsonRpcLogSink, IAsyncDisposable, IDisposable
 		{
 			readonly ILogger _inner;
 			readonly FileStream _eventStream;
@@ -42,6 +42,12 @@ namespace Horde.Agent.Commands.Utilities
 				inner.LogInformation("Writing output to {File}", outputFile);
 			}
 
+			public ValueTask DisposeAsync()
+			{
+				Dispose();
+				return new ValueTask();
+			}
+
 			public void Dispose()
 			{
 				_eventStream.Dispose();
diff --git a/Engine/Source/Programs/Horde/Horde.Agent/Execution/JobExecutor.cs b/Engine/Source/Programs/Horde/Horde.Agent/Execution/JobExecutor.cs
--- a/Engine/Source/Programs/Horde/Horde.Agent/Execution/JobExecutor.cs
+++ b/Engine/Source/Programs/Horde/Horde.Agent/Execution/JobExecutor.cs
@@ -36,8 +36,6 @@ namespace Horde.Agent.Execution
 {
 	abstract class JobExecutor
 	{
-		public const string StorageHttpClientName = "Storage";
-
 		protected class ExportedNode
 		{
 			public string Name { get; set; } = String.Empty;
@@ -666,11 +664,11 @@ namespace Horde.Agent.Execution
 			// Create the storage client
 			using MemoryCache cache = new MemoryCache(new MemoryCacheOptions { });
 
-			using HttpClient httpClient = _httpClientFactory.CreateClient(StorageHttpClientName);
+			using HttpClient httpClient = _httpClientFactory.CreateClient(HttpStorageClient.HttpClientName);
 			httpClient.BaseAddress = new Uri($"{_session.ServerUrl}/api/v1/storage/default");
 			httpClient.DefaultRequestHeaders.Authorization = new AuthenticationHeaderValue("Bearer", _session.Token);
 
-			using HttpClient redirectHttpClient = _httpClientFactory.CreateClient(StorageHttpClientName);
+			using HttpClient redirectHttpClient = _httpClientFactory.CreateClient(HttpStorageClient.HttpClientName);
 
 			HttpStorageClient storage = new HttpStorageClient(CreateHttpClient, CreateHttpRedirectClient, logger);
 			//			FileStorageClient storage = new FileStorageClient(DirectoryReference.Combine(sharedStorageDir, "bundles"), cache, logger);
@@ -914,7 +912,7 @@ namespace Horde.Agent.Execution
 
 		HttpClient CreateHttpClient()
 		{
-			HttpClient httpClient = _httpClientFactory.CreateClient(StorageHttpClientName);
+			HttpClient httpClient = _httpClientFactory.CreateClient(HttpStorageClient.HttpClientName);
 			httpClient.BaseAddress = new Uri($"{_session.ServerUrl}/api/v1/storage/default");
 			httpClient.DefaultRequestHeaders.Authorization = new AuthenticationHeaderValue("Bearer", _session.Token);
 			return httpClient;
@@ -922,7 +920,7 @@ namespace Horde.Agent.Execution
 
 		HttpClient CreateHttpRedirectClient()
 		{
-			return _httpClientFactory.CreateClient(StorageHttpClientName);
+			return _httpClientFactory.CreateClient(HttpStorageClient.HttpClientName);
 		}
 
 		protected async Task<int> ExecuteAutomationToolAsync(BeginStepResponse step, DirectoryReference workspaceDir, string arguments, IStorageClient? store, ILogger logger, CancellationToken cancellationToken)
diff --git a/Engine/Source/Programs/Horde/Horde.Agent/Leases/Handlers/ConformHandler.cs b/Engine/Source/Programs/Horde/Horde.Agent/Leases/Handlers/ConformHandler.cs
--- a/Engine/Source/Programs/Horde/Horde.Agent/Leases/Handlers/ConformHandler.cs
+++ b/Engine/Source/Programs/Horde/Horde.Agent/Leases/Handlers/ConformHandler.cs
@@ -2,6 +2,7 @@
 
 using System;
 using System.Collections.Generic;
+using System.Net.Http;
 using System.Threading;
 using System.Threading.Tasks;
 using Horde.Agent.Execution;
@@ -19,21 +20,21 @@ namespace Horde.Agent.Leases.Handlers
 	class ConformHandler : LeaseHandler<ConformTask>
 	{
 		readonly AgentSettings _settings;
-		readonly ILogger _defaultLogger;
+		readonly IServerLoggerFactory _serverLoggerFactory;
 
 		/// <summary>
 		/// Constructor
 		/// </summary>
-		public ConformHandler(IOptions<AgentSettings> settings, ILogger<ConformHandler> defaultLogger)
+		public ConformHandler(IOptions<AgentSettings> settings, IServerLoggerFactory serverLoggerFactory)
 		{
 			_settings = settings.Value;
-			_defaultLogger = defaultLogger;
+			_serverLoggerFactory = serverLoggerFactory;
 		}
 
 		/// <inheritdoc/>
 		public override async Task<LeaseResult> ExecuteAsync(ISession session, string leaseId, ConformTask conformTask, CancellationToken cancellationToken)
 		{
-			await using JsonRpcLogger conformLogger = new JsonRpcLogger(session.RpcConnection, conformTask.LogId, _defaultLogger);
+			await using IServerLogger conformLogger = _serverLoggerFactory.CreateLogger(session, conformTask.LogId, null);
 
 			conformLogger.LogInformation("Conforming, lease {LeaseId}", leaseId);
 			await session.TerminateProcessesAsync(conformLogger, cancellationToken);
diff --git a/Engine/Source/Programs/Horde/Horde.Agent/Leases/Handlers/JobHandler.cs b/Engine/Source/Programs/Horde/Horde.Agent/Leases/Handlers/JobHandler.cs
--- a/Engine/Source/Programs/Horde/Horde.Agent/Leases/Handlers/JobHandler.cs
+++ b/Engine/Source/Programs/Horde/Horde.Agent/Leases/Handlers/JobHandler.cs
@@ -5,6 +5,7 @@ using System.Collections.Generic;
 using System.Diagnostics;
 using System.IO;
 using System.Linq;
+using System.Net.Http;
 using System.Runtime.InteropServices;
 using System.Threading;
 using System.Threading.Tasks;
@@ -34,22 +35,24 @@ namespace Horde.Agent.Leases.Handlers
 
 		readonly IEnumerable<JobExecutorFactory> _executorFactories;
 		readonly AgentSettings _settings;
+		readonly IServerLoggerFactory _serverLoggerFactory;
 		readonly ILogger _defaultLogger;
 
 		/// <summary>
 		/// Constructor
 		/// </summary>
-		public JobHandler(IEnumerable<JobExecutorFactory> executorFactories, IOptions<AgentSettings> settings, ILogger<JobHandler> defaultLogger)
+		public JobHandler(IEnumerable<JobExecutorFactory> executorFactories, IOptions<AgentSettings> settings, IServerLoggerFactory serverLoggerFactory, ILogger<JobHandler> defaultLogger)
 		{
 			_executorFactories = executorFactories;
 			_settings = settings.Value;
+			_serverLoggerFactory = serverLoggerFactory;
 			_defaultLogger = defaultLogger;
 		}
 
 		/// <inheritdoc/>
 		public override async Task<LeaseResult> ExecuteAsync(ISession session, string leaseId, ExecuteJobTask executeTask, CancellationToken cancellationToken)
 		{
-			await using JsonRpcLogger logger = new JsonRpcLogger(session.RpcConnection, executeTask.LogId, executeTask.JobId, executeTask.BatchId, null, null, _defaultLogger);
+			await using IServerLogger logger = _serverLoggerFactory.CreateLogger(session, executeTask.LogId, executeTask.JobId, executeTask.BatchId, null);
 
 			logger.LogInformation("Executing job \"{JobName}\", jobId {JobId}, batchId {BatchId}, leaseId {LeaseId}", executeTask.JobName, executeTask.JobId, executeTask.BatchId, leaseId);
 			GlobalTracer.Instance.ActiveSpan?.SetTag("jobId", executeTask.JobId.ToString());
@@ -187,7 +190,7 @@ namespace Horde.Agent.Leases.Handlers
 					{
 						// Start writing to the log file
 #pragma warning disable CA2000 // Dispose objects before losing scope
-						await using (JsonRpcLogger stepLogger = new JsonRpcLogger(rpcClient, step.LogId, executeTask.JobId, executeTask.BatchId, step.StepId, step.Warnings, _defaultLogger))
+						await using (IServerLogger stepLogger = _serverLoggerFactory.CreateLogger(session, step.LogId, executeTask.JobId, executeTask.BatchId, step.StepId, step.Warnings))
 						{
 							// Execute the task
 							ILogger forwardingLogger = new DefaultLoggerIndentHandler(stepLogger);
diff --git a/Engine/Source/Programs/Horde/Horde.Agent/Leases/Handlers/UpgradeHandler.cs b/Engine/Source/Programs/Horde/Horde.Agent/Leases/Handlers/UpgradeHandler.cs
--- a/Engine/Source/Programs/Horde/Horde.Agent/Leases/Handlers/UpgradeHandler.cs
+++ b/Engine/Source/Programs/Horde/Horde.Agent/Leases/Handlers/UpgradeHandler.cs
@@ -5,6 +5,7 @@ using System.Collections.Generic;
 using System.Diagnostics;
 using System.IO;
 using System.IO.Compression;
+using System.Net.Http;
 using System.Reflection;
 using System.Text;
 using System.Threading;
@@ -22,17 +23,17 @@ namespace Horde.Agent.Leases.Handlers
 {
 	class UpgradeHandler : LeaseHandler<UpgradeTask>
 	{
-		readonly ILogger _defaultLogger;
+		readonly IServerLoggerFactory _serverLoggerFactory;
 
-		public UpgradeHandler(ILogger<UpgradeHandler> defaultLogger)
+		public UpgradeHandler(IServerLoggerFactory serverLoggerFactory)
 		{
-			_defaultLogger = defaultLogger;
+			_serverLoggerFactory = serverLoggerFactory;
 		}
 
 		/// <inheritdoc/>
 		public override async Task<LeaseResult> ExecuteAsync(ISession session, string leaseId, UpgradeTask task, CancellationToken cancellationToken)
 		{
-			await using JsonRpcLogger logger = new JsonRpcLogger(session.RpcConnection, task.LogId, _defaultLogger);
+			await using IServerLogger logger = _serverLoggerFactory.CreateLogger(session, task.LogId);
 
 			string requiredVersion = task.SoftwareId;
 
diff --git a/Engine/Source/Programs/Horde/Horde.Agent/Program.cs b/Engine/Source/Programs/Horde/Horde.Agent/Program.cs
--- a/Engine/Source/Programs/Horde/Horde.Agent/Program.cs
+++ b/Engine/Source/Programs/Horde/Horde.Agent/Program.cs
@@ -13,6 +13,7 @@ using Datadog.Trace.Configuration;
 using Datadog.Trace.OpenTracing;
 using EpicGames.Core;
 using EpicGames.Horde.Storage;
+using EpicGames.Horde.Storage.Backends;
 using Horde.Agent.Execution;
 using Horde.Agent.Leases;
 using Horde.Agent.Leases.Handlers;
@@ -170,7 +171,7 @@ namespace Horde.Agent
 				return builder.WaitAndRetryAsync(new[] { TimeSpan.FromSeconds(1), TimeSpan.FromSeconds(5), TimeSpan.FromSeconds(10) });
 			});
 
-			services.AddHttpClient(JobExecutor.StorageHttpClientName)
+			services.AddHttpClient(HttpStorageClient.HttpClientName)
 				.AddTransientHttpErrorPolicy(builder =>
 				{
 					return builder.WaitAndRetryAsync(new[] { TimeSpan.FromSeconds(1), TimeSpan.FromSeconds(5), TimeSpan.FromSeconds(10) });
@@ -193,6 +194,7 @@ namespace Horde.Agent
 
 			services.AddSingleton<CapabilitiesService>();
 			services.AddSingleton<SessionFactoryService>();
+			services.AddSingleton<IServerLoggerFactory, ServerLoggerFactory>();
 			services.AddHostedService<WorkerService>();
 
 			// Allow commands to augment the service collection for their own DI service providers
diff --git a/Engine/Source/Programs/Horde/Horde.Agent/Utility/JsonRpcLogger.cs b/Engine/Source/Programs/Horde/Horde.Agent/Utility/JsonRpcLogger.cs
--- a/Engine/Source/Programs/Horde/Horde.Agent/Utility/JsonRpcLogger.cs
+++ b/Engine/Source/Programs/Horde/Horde.Agent/Utility/JsonRpcLogger.cs
@@ -5,18 +5,20 @@ using System.Buffers;
 using System.Collections.Generic;
 using System.Diagnostics;
 using System.Linq;
+using System.Net.Http;
 using System.Text;
 using System.Text.Json;
 using System.Text.Json.Nodes;
 using System.Threading;
 using System.Threading.Channels;
 using System.Threading.Tasks;
-using System.Xml.Linq;
 using EpicGames.Core;
 using EpicGames.Horde.Logs;
 using EpicGames.Horde.Storage;
+using EpicGames.Horde.Storage.Backends;
 using Google.Protobuf;
 using Grpc.Core;
+using Horde.Agent.Services;
 using Horde.Agent.Utility;
 using Horde.Common.Rpc;
 using HordeCommon;
@@ -29,7 +31,7 @@ namespace Horde.Agent.Parser
 {
 	using JsonObject = System.Text.Json.Nodes.JsonObject;
 
-	interface IJsonRpcLogSink
+	interface IJsonRpcLogSink : IAsyncDisposable
 	{
 		Task WriteEventsAsync(List<CreateEventRequest> events, CancellationToken cancellationToken);
 		Task WriteOutputAsync(WriteOutputRequest request, CancellationToken cancellationToken);
@@ -53,6 +55,8 @@ namespace Horde.Agent.Parser
 			_logger = logger;
 		}
 
+		public ValueTask DisposeAsync() => new ValueTask();
+
 		/// <inheritdoc/>
 		public async Task WriteEventsAsync(List<CreateEventRequest> events, CancellationToken cancellationToken)
 		{
@@ -90,7 +94,7 @@ namespace Horde.Agent.Parser
 		readonly IRpcConnection _connection;
 		readonly string _logId;
 		readonly LogBuilder _builder;
-		readonly IJsonRpcLogSink _inner;
+		readonly IJsonRpcLogSink? _inner;
 		readonly TreeWriter _writer;
 		readonly ILogger _logger;
 
@@ -100,11 +104,11 @@ namespace Horde.Agent.Parser
 		readonly object _lockObject = new object();
 
 		// Tailing task
-		readonly CancellationTokenSource _tailCancellationSource;
+		CancellationTokenSource _tailCancellationSource;
 		readonly Task _tailTask;
 		AsyncEvent _newTailDataEvent = new AsyncEvent();
 
-		public JsonRpcAndStorageLogSink(IRpcConnection connection, string logId, IJsonRpcLogSink inner, IStorageClient store, ILogger logger)
+		public JsonRpcAndStorageLogSink(IRpcConnection connection, string logId, IJsonRpcLogSink? inner, IStorageClient store, ILogger logger)
 		{
 			_connection = connection;
 			_logId = logId;
@@ -119,19 +123,28 @@ namespace Horde.Agent.Parser
 
 		public async ValueTask DisposeAsync()
 		{
-			_tailCancellationSource.Cancel();
-			try
-			{
-				await _tailTask;
-			}
-			catch (OperationCanceledException)
+			if(_tailCancellationSource != null)
 			{
+				_tailCancellationSource.Cancel();
+				try
+				{
+					await _tailTask;
+				}
+				catch (OperationCanceledException)
+				{
+				}
+				catch (Exception ex)
+				{
+					_logger.LogError(ex, "Exception on log tailing task ({LogId}): {Message}", _logId, ex.Message);
+				}
+				_tailCancellationSource.Dispose();
+				_tailCancellationSource = null!;
 			}
-			catch (Exception ex)
+
+			if (_inner != null)
 			{
-				_logger.LogError(ex, "Exception on log tailing task ({LogId}): {Message}", _logId, ex.Message);
+				await _inner.DisposeAsync();
 			}
-			_tailCancellationSource.Dispose();
 		}
 
 		async Task TickTailAsync(CancellationToken cancellationToken)
@@ -173,18 +186,30 @@ namespace Horde.Agent.Parser
 		}
 
 		/// <inheritdoc/>
-		public Task SetOutcomeAsync(JobStepOutcome outcome, CancellationToken cancellationToken) => _inner.SetOutcomeAsync(outcome, cancellationToken);
+		public async Task SetOutcomeAsync(JobStepOutcome outcome, CancellationToken cancellationToken)
+		{
+			if (_inner != null)
+			{
+				await _inner.SetOutcomeAsync(outcome, cancellationToken);
+			}
+		}
 
 		/// <inheritdoc/>
 		public async Task WriteEventsAsync(List<CreateEventRequest> events, CancellationToken cancellationToken)
 		{
-			await _inner.WriteEventsAsync(events, cancellationToken);
+			if (_inner != null)
+			{
+				await _inner.WriteEventsAsync(events, cancellationToken);
+			}
 		}
 
 		/// <inheritdoc/>
 		public async Task WriteOutputAsync(WriteOutputRequest request, CancellationToken cancellationToken)
 		{
-			await _inner.WriteOutputAsync(request, cancellationToken);
+			if (_inner != null)
+			{
+				await _inner.WriteOutputAsync(request, cancellationToken);
+			}
 
 			_builder.WriteData(request.Data.Memory);
 			_bufferLength += request.Data.Length;
@@ -206,6 +231,7 @@ namespace Horde.Agent.Parser
 			_logger.LogDebug("Updating log {LogId} to line {LineCount}, blob {Locator}", _logId, lineCount, locator);
 
 			UpdateLogRequest request = new UpdateLogRequest();
+			request.LogId = _logId;
 			request.LineCount = lineCount;
 			request.BlobLocator = locator.ToString();
 			await _connection.InvokeAsync((LogRpcClient client) => client.UpdateLogAsync(request, cancellationToken: cancellationToken), cancellationToken);
@@ -251,10 +277,26 @@ namespace Horde.Agent.Parser
 		#endregion
 	}
 
+	/// <summary>
+	/// Interface for a log device
+	/// </summary>
+	public interface IServerLogger : ILogger, IAsyncDisposable
+	{
+		/// <summary>
+		/// Outcome of the job step, including any warnings/errors
+		/// </summary>
+		JobStepOutcome Outcome { get; }
+
+		/// <summary>
+		/// Flushes the logger with the server and stops the background work
+		/// </summary>
+		Task StopAsync();
+	}
+
 	/// <summary>
 	/// Class to handle uploading log data to the server in the background
 	/// </summary>
-	sealed class JsonRpcLogger : ILogger, IAsyncDisposable
+	sealed class JsonRpcLogger : IServerLogger
 	{
 		class QueueItem
 		{
@@ -308,44 +350,6 @@ namespace Horde.Agent.Parser
 			Outcome = JobStepOutcome.Success;
 		}
 
-		/// <summary>
-		/// Constructor
-		/// </summary>
-		/// <param name="rpcClient">RPC client to use for server requests</param>
-		/// <param name="logId">The log id to write to</param>
-		/// <param name="inner">Additional logger to write to</param>
-		public JsonRpcLogger(IRpcConnection rpcClient, string logId, ILogger inner)
-			: this(rpcClient, logId, null, null, null, null, inner)
-		{
-		}
-
-		/// <summary>
-		/// Constructor
-		/// </summary>
-		/// <param name="rpcClient">RPC client to use for server requests</param>
-		/// <param name="logId">The log id to write to</param>
-		/// <param name="warnings">Whether to include warnings in the output</param>
-		/// <param name="inner">Additional logger to write to</param>
-		public JsonRpcLogger(IRpcConnection rpcClient, string logId, bool? warnings, ILogger inner)
-			: this(rpcClient, logId, null, null, null, warnings, inner)
-		{
-		}
-
-		/// <summary>
-		/// Constructor
-		/// </summary>
-		/// <param name="rpcClient">RPC client to use for server requests</param>
-		/// <param name="logId">The log id to write to</param>
-		/// <param name="jobId">Id of the job being executed</param>
-		/// <param name="jobBatchId">Batch being executed</param>
-		/// <param name="jobStepId">Id of the step being executed</param>
-		/// <param name="warnings">Whether to include warnings in the output</param>
-		/// <param name="inner">Additional logger to write to</param>
-		public JsonRpcLogger(IRpcConnection rpcClient, string logId, string? jobId, string? jobBatchId, string? jobStepId, bool? warnings, ILogger inner)
-			: this(new JsonRpcLogSink(rpcClient, jobId, jobBatchId, jobStepId, inner), logId, warnings, inner)
-		{
-		}
-
 		/// <inheritdoc/>
 		public void Log<TState>(LogLevel logLevel, EventId eventId, TState state, Exception? exception, Func<TState, Exception?, string> formatter)
 		{
@@ -405,6 +409,7 @@ namespace Horde.Agent.Parser
 		public async ValueTask DisposeAsync()
 		{
 			await StopAsync();
+			await _sink.DisposeAsync();
 		}
 
 		/// <summary>
diff --git a/Engine/Source/Programs/Shared/EpicGames.Horde/Storage/Backends/HttpStorageClient.cs b/Engine/Source/Programs/Shared/EpicGames.Horde/Storage/Backends/HttpStorageClient.cs
--- a/Engine/Source/Programs/Shared/EpicGames.Horde/Storage/Backends/HttpStorageClient.cs
+++ b/Engine/Source/Programs/Shared/EpicGames.Horde/Storage/Backends/HttpStorageClient.cs
@@ -21,6 +21,11 @@ namespace EpicGames.Horde.Storage.Backends
 	/// </summary>
 	public class HttpStorageClient : StorageClientBase
 	{
+		/// <summary>
+		/// Name of clients created from the http client factory
+		/// </summary>
+		public const string HttpClientName = "Horde.HttpStorageClient";
+
 		class WriteBlobResponse
 		{
 			public BlobLocator Locator { get; set; }
@@ -49,6 +54,36 @@ namespace EpicGames.Horde.Storage.Backends
 			_logger = logger;
 		}
 
+		/// <summary>
+		/// Constructor
+		/// </summary>
+		public HttpStorageClient(IHttpClientFactory httpClientFactory, Uri baseAddress, string? bearerToken, ILogger logger)
+			: this(() => CreateAuthenticatedClient(httpClientFactory, baseAddress, bearerToken), () => CreateClient(httpClientFactory), logger)
+		{
+		}
+
+		/// <summary>
+		/// Helper method to create an HTTP client from the given factory
+		/// </summary>
+		static HttpClient CreateClient(IHttpClientFactory httpClientFactory)
+		{
+			return httpClientFactory.CreateClient(HttpClientName);
+		}
+
+		/// <summary>
+		/// Helper method to add the base address and auth header to an HTTP client
+		/// </summary>
+		static HttpClient CreateAuthenticatedClient(IHttpClientFactory httpClientFactory, Uri baseAddress, string? bearerToken)
+		{
+			HttpClient httpClient = CreateClient(httpClientFactory);
+			httpClient.BaseAddress = baseAddress;
+			if (bearerToken != null)
+			{
+				httpClient.DefaultRequestHeaders.Authorization = new AuthenticationHeaderValue("Bearer", bearerToken);
+			}
+			return httpClient;
+		}
+
 		#region Blobs
 
 		/// <inheritdoc/>
