Horde: Missing file.

diff --git a/Engine/Source/Programs/Horde/Horde.Agent/Services/ServerLoggerFactory.cs b/Engine/Source/Programs/Horde/Horde.Agent/Services/ServerLoggerFactory.cs
new file mode 100644
--- /dev/null
+++ b/Engine/Source/Programs/Horde/Horde.Agent/Services/ServerLoggerFactory.cs
@@ -0,0 +1,69 @@
+// Copyright Epic Games, Inc. All Rights Reserved.
+
+using Horde.Agent.Parser;
+using Microsoft.Extensions.Logging;
+
+namespace Horde.Agent.Services
+{
+	/// <summary>
+	/// Interface for creating <see cref="IServerLogger"/> instances
+	/// </summary>
+	interface IServerLoggerFactory
+	{
+		/// <summary>
+		/// Creates a logger which uploads data to the server
+		/// </summary>
+		/// <param name="session">The current session</param>
+		/// <param name="logId">The log id</param>
+		/// <param name="jobId">Job containing the log</param>
+		/// <param name="batchId">Job batch id</param>
+		/// <param name="stepId">The job step id</param>
+		/// <param name="warnings">Whether to suppress warnings</param>
+		/// <returns>New logger instance</returns>
+		IServerLogger CreateLogger(ISession session, string logId, string? jobId, string? batchId, string? stepId, bool? warnings = null);
+	}
+
+	/// <summary>
+	/// Extension methods for <see cref="IServerLoggerFactory"/>
+	/// </summary>
+	static class ServerLoggerFactoryExtensions
+	{
+		/// <summary>
+		/// Creates a logger which uploads data to the server
+		/// </summary>
+		/// <param name="service">Service instance</param>
+		/// <param name="session">The current session</param>
+		/// <param name="logId">The log identifier</param>
+		/// <param name="warnings">Whether to suppress warnings</param>
+		/// <returns>New logger instance</returns>
+		public static IServerLogger CreateLogger(this IServerLoggerFactory service, ISession session, string logId, bool? warnings = null)
+		{
+			return service.CreateLogger(session, logId, null, null, null, warnings);
+		}
+	}
+
+	/// <summary>
+	/// Implementation of <see cref="IServerLoggerFactory"/>
+	/// </summary>
+	class ServerLoggerFactory : IServerLoggerFactory
+	{
+		readonly ILogger _logger;
+
+		/// <summary>
+		/// Constructor
+		/// </summary>
+		public ServerLoggerFactory(ILogger<ServerLoggerFactory> logger)
+		{
+			_logger = logger;
+		}
+
+		/// <inheritdoc/>
+		public IServerLogger CreateLogger(ISession session, string logId, string? jobId, string? batchId, string? stepId, bool? warnings = null)
+		{
+#pragma warning disable CA2000 // Dispose objects before losing scope
+			IJsonRpcLogSink sink = new JsonRpcLogSink(session.RpcConnection, jobId, batchId, stepId, _logger);
+			return new JsonRpcLogger(sink, logId, warnings, _logger);
+#pragma warning restore CA2000 // Dispose objects before losing scope
+		}
+	}
+}
