Horde: Add device ip to checkout confirmation
#rnx
#skipci

diff --git a/Engine/Source/Programs/Horde/HordeDashboard/src/components/DeviceView.tsx b/Engine/Source/Programs/Horde/HordeDashboard/src/components/DeviceView.tsx
--- a/Engine/Source/Programs/Horde/HordeDashboard/src/components/DeviceView.tsx
+++ b/Engine/Source/Programs/Horde/HordeDashboard/src/components/DeviceView.tsx
@@ -1,4 +1,4 @@
-import { DefaultButton, DetailsList, DetailsListLayoutMode, DetailsRow, Dropdown, FocusZone, FocusZoneDirection, getTheme, IColumn, Icon, IconButton, IDetailsListProps, IGroup, Label, Link as FluentLink, mergeStyleSets, Modal, Pivot, PivotItem, PrimaryButton, ScrollablePane, ScrollbarVisibility, SelectionMode, Spinner, SpinnerSize, Stack, Text } from "@fluentui/react";
+import { DefaultButton, DetailsList, DetailsListLayoutMode, Dropdown, FocusZone, FocusZoneDirection, getTheme, IColumn, Icon, IconButton, IGroup, Label, Link as FluentLink, mergeStyleSets, Modal, Pivot, PivotItem, PrimaryButton, ScrollablePane, ScrollbarVisibility, SelectionMode, Spinner, SpinnerSize, Stack, Text } from "@fluentui/react";
 import { observer } from "mobx-react-lite";
 import moment from "moment";
 import React, { useEffect, useState } from "react";
@@ -590,7 +590,7 @@ export const CheckoutConfirmModal: React.FC<{ check: "in" | "out" | "error", dev
       // simple cases
       switch (column.name) {
          case 'Name':
-            return <Text variant="medium" styles={{ root: { fontFamily: "Horde Open Sans SemiBold" } }}>{platform} : {device.name}</Text>
+            return <Text variant="medium" styles={{ root: { fontFamily: "Horde Open Sans SemiBold" } }}>{platform} : {device.name} : {device.address ?? ""} </Text>
          case 'Date':
             if (device.checkOutExpirationTime) {
                return <Text variant="medium">{`Checked out until ${getNiceTime(device.checkOutExpirationTime)}`}</Text>
