Horde: Add custom storage endpoints for uploading log files, and set up a default namespace for receiving them.

diff --git a/Engine/Source/Programs/Horde/Horde.Agent.Tests/JsonRpcLoggerTests.cs b/Engine/Source/Programs/Horde/Horde.Agent.Tests/JsonRpcLoggerTests.cs
--- a/Engine/Source/Programs/Horde/Horde.Agent.Tests/JsonRpcLoggerTests.cs
+++ b/Engine/Source/Programs/Horde/Horde.Agent.Tests/JsonRpcLoggerTests.cs
@@ -117,9 +117,9 @@ namespace Horde.Agent.Tests
 				return Task.CompletedTask;
 			}
 
-			protected override async Task<int> UpdateLogTailAsync(int tailNext, ReadOnlyMemory<byte> tailData, CancellationToken cancellationToken)
+			protected override async Task<int> UpdateLogTailAsync(int tailNext, ReadOnlyMemory<byte> tailData)
 			{
-				await Task.Delay(TimeSpan.FromSeconds(2.0), cancellationToken);
+				await Task.Delay(TimeSpan.FromSeconds(2.0));
 				return -1;
 			}
 		}
diff --git a/Engine/Source/Programs/Horde/Horde.Agent/AgentSettings.cs b/Engine/Source/Programs/Horde/Horde.Agent/AgentSettings.cs
--- a/Engine/Source/Programs/Horde/Horde.Agent/AgentSettings.cs
+++ b/Engine/Source/Programs/Horde/Horde.Agent/AgentSettings.cs
@@ -6,7 +6,6 @@ using System.ComponentModel.DataAnnotations;
 using System.Linq;
 using EpicGames.Core;
 using EpicGames.Horde.Storage;
-using Horde.Agent.Execution;
 using Microsoft.Extensions.Configuration;
 
 namespace Horde.Agent
@@ -187,6 +186,11 @@ namespace Horde.Agent
 		/// </summary>
 		public bool EnableAwsEc2Support { get; set; } = true;
 
+		/// <summary>
+		/// Whether to enable the new storage-API based log system 
+		/// </summary>
+		public bool EnableNewLogger { get; set; }
+
 		/// <summary>
 		/// Key/value properties in addition to those set internally by the agent
 		/// </summary>
diff --git a/Engine/Source/Programs/Horde/Horde.Agent/Services/ServerLoggerFactory.cs b/Engine/Source/Programs/Horde/Horde.Agent/Services/ServerLoggerFactory.cs
--- a/Engine/Source/Programs/Horde/Horde.Agent/Services/ServerLoggerFactory.cs
+++ b/Engine/Source/Programs/Horde/Horde.Agent/Services/ServerLoggerFactory.cs
@@ -1,7 +1,11 @@
 // Copyright Epic Games, Inc. All Rights Reserved.
 
+using System;
 using Horde.Agent.Parser;
 using Microsoft.Extensions.Logging;
+using EpicGames.Horde.Storage.Backends;
+using System.Net.Http;
+using Microsoft.Extensions.Options;
 
 namespace Horde.Agent.Services
 {
@@ -47,13 +51,17 @@ namespace Horde.Agent.Services
 	/// </summary>
 	class ServerLoggerFactory : IServerLoggerFactory
 	{
+		readonly IHttpClientFactory _httpClientFactory;
+		readonly IOptions<AgentSettings> _settings;
 		readonly ILogger _logger;
 
 		/// <summary>
 		/// Constructor
 		/// </summary>
-		public ServerLoggerFactory(ILogger<ServerLoggerFactory> logger)
+		public ServerLoggerFactory(IHttpClientFactory httpClientFactory, IOptions<AgentSettings> settings, ILogger<ServerLoggerFactory> logger)
 		{
+			_httpClientFactory = httpClientFactory;
+			_settings = settings;
 			_logger = logger;
 		}
 
@@ -62,6 +70,11 @@ namespace Horde.Agent.Services
 		{
 #pragma warning disable CA2000 // Dispose objects before losing scope
 			IJsonRpcLogSink sink = new JsonRpcLogSink(session.RpcConnection, jobId, batchId, stepId, _logger);
+			if (_settings.Value.EnableNewLogger)
+			{
+				HttpStorageClient storageClient = new HttpStorageClient(_httpClientFactory,  new Uri(session.ServerUrl, $"api/v1/logs/{logId}/"), session.Token, _logger);
+				sink = new JsonRpcAndStorageLogSink(session.RpcConnection, logId, sink, storageClient, _logger);
+			}
 			return new JsonRpcLogger(sink, logId, warnings, _logger);
 #pragma warning restore CA2000 // Dispose objects before losing scope
 		}
diff --git a/Engine/Source/Programs/Horde/Horde.Agent/Services/WorkerService.cs b/Engine/Source/Programs/Horde/Horde.Agent/Services/WorkerService.cs
--- a/Engine/Source/Programs/Horde/Horde.Agent/Services/WorkerService.cs
+++ b/Engine/Source/Programs/Horde/Horde.Agent/Services/WorkerService.cs
@@ -4,17 +4,10 @@ using System;
 using System.Collections.Generic;
 using System.Diagnostics;
 using System.Linq;
-using System.Reflection;
-using System.Runtime.InteropServices;
 using System.Threading;
 using System.Threading.Tasks;
-using EpicGames.Core;
-using Grpc.Net.Client;
 using Horde.Agent.Leases;
 using Horde.Agent.Utility;
-using HordeCommon;
-using HordeCommon.Rpc;
-using HordeCommon.Rpc.Messages;
 using Microsoft.Extensions.DependencyInjection;
 using Microsoft.Extensions.Hosting;
 using Microsoft.Extensions.Logging;
diff --git a/Engine/Source/Programs/Horde/Horde.Agent/Utility/JsonRpcLogger.cs b/Engine/Source/Programs/Horde/Horde.Agent/Utility/JsonRpcLogger.cs
--- a/Engine/Source/Programs/Horde/Horde.Agent/Utility/JsonRpcLogger.cs
+++ b/Engine/Source/Programs/Horde/Horde.Agent/Utility/JsonRpcLogger.cs
@@ -104,8 +104,8 @@ namespace Horde.Agent.Parser
 		readonly object _lockObject = new object();
 
 		// Tailing task
-		CancellationTokenSource _tailCancellationSource;
 		readonly Task _tailTask;
+		AsyncEvent _tailTaskStop;
 		AsyncEvent _newTailDataEvent = new AsyncEvent();
 
 		public JsonRpcAndStorageLogSink(IRpcConnection connection, string logId, IJsonRpcLogSink? inner, IStorageClient store, ILogger logger)
@@ -117,15 +117,15 @@ namespace Horde.Agent.Parser
 			_writer = new TreeWriter(store, prefix: $"logs/{logId}");
 			_logger = logger;
 
-			_tailCancellationSource = new CancellationTokenSource();
-			_tailTask = Task.Run(() => TickTailAsync(_tailCancellationSource.Token));
+			_tailTaskStop = new AsyncEvent();
+			_tailTask = Task.Run(() => TickTailAsync());
 		}
 
 		public async ValueTask DisposeAsync()
 		{
-			if(_tailCancellationSource != null)
+			if(_tailTaskStop != null)
 			{
-				_tailCancellationSource.Cancel();
+				_tailTaskStop.Latch();
 				try
 				{
 					await _tailTask;
@@ -137,8 +137,7 @@ namespace Horde.Agent.Parser
 				{
 					_logger.LogError(ex, "Exception on log tailing task ({LogId}): {Message}", _logId, ex.Message);
 				}
-				_tailCancellationSource.Dispose();
-				_tailCancellationSource = null!;
+				_tailTaskStop = null!;
 			}
 
 			if (_inner != null)
@@ -147,10 +146,10 @@ namespace Horde.Agent.Parser
 			}
 		}
 
-		async Task TickTailAsync(CancellationToken cancellationToken)
+		async Task TickTailAsync()
 		{
 			int tailNext = -1;
-			for (; ;)
+			while (!_tailTaskStop.IsSet())
 			{
 				Task newTailDataTask;
 
@@ -176,7 +175,7 @@ namespace Horde.Agent.Parser
 				}
 
 				// Update the next tailing position
-				int newTailNext = await UpdateLogTailAsync(tailNext, tailData, cancellationToken);
+				int newTailNext = await UpdateLogTailAsync(tailNext, tailData);
 				if (newTailNext != tailNext)
 				{
 					tailNext = newTailNext;
@@ -237,27 +236,31 @@ namespace Horde.Agent.Parser
 			await _connection.InvokeAsync((LogRpcClient client) => client.UpdateLogAsync(request, cancellationToken: cancellationToken), cancellationToken);
 		}
 
-		protected virtual async Task<int> UpdateLogTailAsync(int tailNext, ReadOnlyMemory<byte> tailData, CancellationToken cancellationToken)
+		protected virtual async Task<int> UpdateLogTailAsync(int tailNext, ReadOnlyMemory<byte> tailData)
 		{
 			DateTime deadline = DateTime.UtcNow.AddMinutes(2.0);
-			using (IRpcClientRef<LogRpcClient> clientRef = await _connection.GetClientRefAsync<LogRpcClient>(cancellationToken))
+			using (IRpcClientRef<LogRpcClient> clientRef = await _connection.GetClientRefAsync<LogRpcClient>(CancellationToken.None))
 			{
-				using (AsyncDuplexStreamingCall<UpdateLogTailRequest, UpdateLogTailResponse> call = clientRef.Client.UpdateLogTail(deadline: deadline, cancellationToken: cancellationToken))
+				using (AsyncDuplexStreamingCall<UpdateLogTailRequest, UpdateLogTailResponse> call = clientRef.Client.UpdateLogTail(deadline: deadline))
 				{
 					// Write the request to the server
 					UpdateLogTailRequest request = new UpdateLogTailRequest();
 					request.TailNext = tailNext;
 					request.TailData = UnsafeByteOperations.UnsafeWrap(tailData);
-					await call.RequestStream.WriteAsync(request, cancellationToken);
+					await call.RequestStream.WriteAsync(request);
 
 					// Wait until the server responds or we need to trigger a new update
 					Task<bool> moveNextAsync = call.ResponseStream.MoveNext();
 
-					Task task = await Task.WhenAny(moveNextAsync, clientRef.DisposingTask, Task.Delay(TimeSpan.FromMinutes(1.0), CancellationToken.None));
+					Task task = await Task.WhenAny(moveNextAsync, clientRef.DisposingTask, _tailTaskStop.Task, Task.Delay(TimeSpan.FromMinutes(1.0), CancellationToken.None));
 					if (task == clientRef.DisposingTask)
 					{
 						_logger.LogDebug("Cancelling long poll from client side (server migration)");
 					}
+					else if (task == _tailTaskStop.Task)
+					{
+						_logger.LogDebug("Cancelling long poll from client side (complete)");
+					}
 
 					// Close the request stream to indicate that we're finished
 					await call.RequestStream.CompleteAsync();
diff --git a/Engine/Source/Programs/Horde/Horde.Build/Defaults/globals.json b/Engine/Source/Programs/Horde/Horde.Build/Defaults/globals.json
--- a/Engine/Source/Programs/Horde/Horde.Build/Defaults/globals.json
+++ b/Engine/Source/Programs/Horde/Horde.Build/Defaults/globals.json
@@ -23,6 +23,26 @@
 			{
 				"id": "default",
 				"backend": "default-backend"
+			},
+			{
+				"id": "horde-artifacts",
+				"prefix": "Artifacts/",
+				"backend": "default-backend"
+			},
+			{
+				"id": "horde-perforce",
+				"prefix": "Perforce/",
+				"backend": "default-backend"
+			},
+			{
+				"id": "horde-logs",
+				"prefix": "Logs/",
+				"backend": "default-backend"
+			},
+			{
+				"id": "horde-tools",
+				"prefix": "Tools/",
+				"backend": "default-backend"
 			}
 		]
 	},
diff --git a/Engine/Source/Programs/Horde/Horde.Build/Logs/ILogFile.cs b/Engine/Source/Programs/Horde/Horde.Build/Logs/ILogFile.cs
--- a/Engine/Source/Programs/Horde/Horde.Build/Logs/ILogFile.cs
+++ b/Engine/Source/Programs/Horde/Horde.Build/Logs/ILogFile.cs
@@ -3,6 +3,7 @@
 using System;
 using System.Collections.Generic;
 using EpicGames.Core;
+using EpicGames.Horde.Storage;
 using Horde.Build.Agents.Sessions;
 using Horde.Build.Jobs;
 using Horde.Build.Utilities;
@@ -78,6 +79,11 @@ namespace Horde.Build.Logs
 		/// Chunks within this file
 		/// </summary>
 		public IReadOnlyList<ILogChunk> Chunks { get; }
+
+		/// <summary>
+		/// Name of the ref used to store data for this log
+		/// </summary>
+		public RefName RefName { get; }
 	}
 
 	/// <summary>
diff --git a/Engine/Source/Programs/Horde/Horde.Build/Logs/LogFileCollection.cs b/Engine/Source/Programs/Horde/Horde.Build/Logs/LogFileCollection.cs
--- a/Engine/Source/Programs/Horde/Horde.Build/Logs/LogFileCollection.cs
+++ b/Engine/Source/Programs/Horde/Horde.Build/Logs/LogFileCollection.cs
@@ -5,9 +5,11 @@ using System.Linq;
 using System.Net;
 using System.Threading;
 using System.Threading.Tasks;
+using EpicGames.Horde.Storage;
 using Horde.Build.Agents.Sessions;
 using Horde.Build.Jobs;
 using Horde.Build.Server;
+using Horde.Build.Streams;
 using Horde.Build.Utilities;
 using MongoDB.Bson.Serialization.Attributes;
 using MongoDB.Driver;
@@ -17,6 +19,7 @@ namespace Horde.Build.Logs
 	using JobId = ObjectId<IJob>;
 	using LogId = ObjectId<ILogFile>;
 	using SessionId = ObjectId<ISession>;
+	using StreamId = StringId<IStream>;
 
 	/// <summary>
 	/// Wrapper around the jobs collection in a mongo DB
@@ -70,6 +73,8 @@ namespace Horde.Build.Logs
 
 			public List<LogChunkDocument> Chunks { get; set; } = new List<LogChunkDocument>();
 
+			public RefName RefName { get; set; }
+
 			[BsonRequired]
 			public int UpdateIndex { get; set; }
 
@@ -87,6 +92,7 @@ namespace Horde.Build.Logs
 				SessionId = sessionId;
 				Type = type;
 				MaxLineIndex = 0;
+				RefName = new RefName(Id.ToString());
 			}
 
 			public LogFileDocument Clone()
diff --git a/Engine/Source/Programs/Horde/Horde.Build/Logs/LogRpcService.cs b/Engine/Source/Programs/Horde/Horde.Build/Logs/LogRpcService.cs
--- a/Engine/Source/Programs/Horde/Horde.Build/Logs/LogRpcService.cs
+++ b/Engine/Source/Programs/Horde/Horde.Build/Logs/LogRpcService.cs
@@ -46,7 +46,7 @@ namespace Horde.Build.Logs
 				throw new StructuredRpcException(StatusCode.PermissionDenied, "Access denied");
 			}
 
-			IStorageClient store = await _storageService.GetClientAsync(new NamespaceId("default"), context.CancellationToken);
+			IStorageClient store = await _storageService.GetClientAsync(Namespace.Logs, context.CancellationToken);
 			_logger.LogInformation("Updating {LogId} to node {Locator}", request.LogId, request.BlobLocator);
 			await store.WriteRefTargetAsync(new RefName(request.LogId), NodeLocator.Parse(request.BlobLocator));
 			return new UpdateLogResponse();
diff --git a/Engine/Source/Programs/Horde/Horde.Build/Logs/LogsController.cs b/Engine/Source/Programs/Horde/Horde.Build/Logs/LogsController.cs
--- a/Engine/Source/Programs/Horde/Horde.Build/Logs/LogsController.cs
+++ b/Engine/Source/Programs/Horde/Horde.Build/Logs/LogsController.cs
@@ -7,13 +7,17 @@ using System.Security.Claims;
 using System.Text;
 using System.Threading;
 using System.Threading.Tasks;
+using EpicGames.Horde.Storage;
 using Horde.Build.Acls;
 using Horde.Build.Issues;
 using Horde.Build.Jobs;
 using Horde.Build.Logs.Data;
+using Horde.Build.Storage;
 using Horde.Build.Utilities;
 using Microsoft.AspNetCore.Authorization;
+using Microsoft.AspNetCore.Http;
 using Microsoft.AspNetCore.Mvc;
+using Microsoft.Extensions.Options;
 using MongoDB.Bson;
 
 namespace Horde.Build.Logs
@@ -49,17 +53,21 @@ namespace Horde.Build.Logs
 		private readonly IIssueCollection _issueCollection;
 		private readonly AclService _aclService;
 		private readonly JobService _jobService;
+		private readonly StorageService _storageService;
+		private readonly IOptions<ServerSettings> _settings;
 
 		/// <summary>
 		/// Constructor
 		/// </summary>
-		public LogsController(ILogFileService logFileService, IIssueCollection issueCollection, AclService aclService, JobService jobService)
+		public LogsController(ILogFileService logFileService, IIssueCollection issueCollection, AclService aclService, JobService jobService, StorageService storageService, IOptions<ServerSettings> settings)
 		{
 			_logFileService = logFileService;
 			_issueCollection = issueCollection;
 			_aclService = aclService;
 			_jobService = jobService;
-		}
+			_storageService = storageService;
+			_settings = settings;
+ 		}
 
 		/// <summary>
 		/// Retrieve metadata about a specific log file
@@ -87,6 +95,31 @@ namespace Horde.Build.Logs
 			return new GetLogFileResponse(logFile, metadata).ApplyFilter(filter);       
 		}
 
+		/// <summary>
+		/// Uploads a blob for a log file. See /api/v1/storage/XXX/blobs.
+		/// </summary>
+		/// <param name="logFileId">Id of the log file to get information about</param>
+		/// <param name="file">Data for the blob</param>
+		/// <param name="cancellationToken">Cancellation token for the request</param>
+		/// <returns>Information about the requested project</returns>
+		[HttpPost]
+		[Route("/api/v1/logs/{logFileId}/blobs")]
+		[ProducesResponseType(typeof(WriteBlobResponse), 200)]
+		public async Task<ActionResult<WriteBlobResponse>> WriteLogBlob(LogId logFileId, IFormFile? file, CancellationToken cancellationToken = default)
+		{
+			ILogFile? logFile = await _logFileService.GetLogFileAsync(logFileId, cancellationToken);
+			if (logFile == null)
+			{
+				return NotFound();
+			}
+			if (!await AuthorizeAsync(logFile, AclAction.WriteLogData, User, null))
+			{
+				return Forbid();
+			}
+
+			return await StorageController.WriteBlobAsync(_storageService, Namespace.Logs, file, $"{logFile.RefName}", cancellationToken);
+		}
+
 		/// <summary>
 		/// Retrieve raw data for a log file
 		/// </summary>
diff --git a/Engine/Source/Programs/Horde/Horde.Build/ServerSettings.cs b/Engine/Source/Programs/Horde/Horde.Build/ServerSettings.cs
--- a/Engine/Source/Programs/Horde/Horde.Build/ServerSettings.cs
+++ b/Engine/Source/Programs/Horde/Horde.Build/ServerSettings.cs
@@ -593,11 +593,6 @@ namespace Horde.Build
 		/// </summary>
 		public StorageOptions? Storage { get; set; }
 
-		/// <summary>
-		/// Namespace to use for storing tools
-		/// </summary>
-		public NamespaceId ToolNamespaceId { get; set; } = new NamespaceId("horde.p4");
-
 		/// <summary>
 		/// Whether to use the local Perforce environment
 		/// </summary>
diff --git a/Engine/Source/Programs/Horde/Horde.Build/Storage/StorageConfig.cs b/Engine/Source/Programs/Horde/Horde.Build/Storage/StorageConfig.cs
--- a/Engine/Source/Programs/Horde/Horde.Build/Storage/StorageConfig.cs
+++ b/Engine/Source/Programs/Horde/Horde.Build/Storage/StorageConfig.cs
@@ -1,5 +1,6 @@
 // Copyright Epic Games, Inc. All Rights Reserved.
 
+using System;
 using System.Collections.Generic;
 using System.ComponentModel.DataAnnotations;
 using EpicGames.Horde.Storage;
@@ -25,6 +26,16 @@ namespace Horde.Build.Storage
 		/// Replicated Perforce data
 		/// </summary>
 		public static NamespaceId Perforce { get; } = new NamespaceId("horde-perforce");
+
+		/// <summary>
+		/// Log data
+		/// </summary>
+		public static NamespaceId Logs { get; } = new NamespaceId("horde-logs");
+
+		/// <summary>
+		/// Storage of tool data
+		/// </summary>
+		public static NamespaceId Tools { get; } = new NamespaceId("horde-tools");
 	}
 
 	/// <summary>
@@ -107,6 +118,11 @@ namespace Horde.Build.Storage
 		[Required]
 		public BackendId Backend { get; set; }
 
+		/// <summary>
+		/// Prefix for items within this namespace
+		/// </summary>
+		public string Prefix { get; set; } = String.Empty;
+
 		/// <summary>
 		/// How frequently to run garbage collection, in hours.
 		/// </summary>
diff --git a/Engine/Source/Programs/Horde/Horde.Build/Storage/StorageController.cs b/Engine/Source/Programs/Horde/Horde.Build/Storage/StorageController.cs
--- a/Engine/Source/Programs/Horde/Horde.Build/Storage/StorageController.cs
+++ b/Engine/Source/Programs/Horde/Horde.Build/Storage/StorageController.cs
@@ -143,6 +143,40 @@ namespace Horde.Build.Storage
 			}
 		}
 
+		/// <summary>
+		/// Writes a blob to storage. Exposed as a public utility method to allow other routes with their own authentication methods to wrap their own authentication/redirection.
+		/// </summary>
+		/// <param name="storageService">The storage service</param>
+		/// <param name="namespaceId">Namespace to write the blob to</param>
+		/// <param name="file">File to be written</param>
+		/// <param name="prefix">Prefix for uploaded blobs</param>
+		/// <param name="cancellationToken">Cancellation token</param>
+		/// <returns>Information about the written blob, or redirect information</returns>
+		public static async Task<ActionResult<WriteBlobResponse>> WriteBlobAsync(StorageService storageService, NamespaceId namespaceId, IFormFile? file, [FromForm] string? prefix = default, CancellationToken cancellationToken = default)
+		{
+			IStorageClientImpl client = await storageService.GetClientAsync(namespaceId, cancellationToken);
+			if (file == null)
+			{
+				(BlobLocator Locator, Uri UploadUrl)? result = await client.GetWriteRedirectAsync(prefix ?? String.Empty, cancellationToken);
+				if (result == null)
+				{
+					return new WriteBlobResponse { SupportsRedirects = false };
+				}
+				else
+				{
+					return new WriteBlobResponse { Locator = result.Value.Locator, UploadUrl = result.Value.UploadUrl };
+				}
+			}
+			else
+			{
+				using (Stream stream = file.OpenReadStream())
+				{
+					BlobLocator locator = await client.WriteBlobAsync(stream, prefix: (prefix == null) ? Utf8String.Empty : new Utf8String(prefix), cancellationToken: cancellationToken);
+					return new WriteBlobResponse { Locator = locator, SupportsRedirects = client.SupportsRedirects ? (bool?)true : null };
+				}
+			}
+		}
+
 		/// <summary>
 		/// Retrieves data from the storage service. 
 		/// </summary>
diff --git a/Engine/Source/Programs/Horde/Horde.Build/Storage/StorageService.cs b/Engine/Source/Programs/Horde/Horde.Build/Storage/StorageService.cs
--- a/Engine/Source/Programs/Horde/Horde.Build/Storage/StorageService.cs
+++ b/Engine/Source/Programs/Horde/Horde.Build/Storage/StorageService.cs
@@ -94,6 +94,7 @@ namespace Horde.Build.Storage
 		sealed class StorageClient : StorageClientBase, IStorageClientImpl, IDisposable
 		{
 			readonly StorageService _outer;
+			readonly string _prefix;
 
 			public NamespaceConfig Config { get; }
 
@@ -109,6 +110,12 @@ namespace Horde.Build.Storage
 				Config = config;
 				Backend = backend;
 				SupportsRedirects = backend is IStorageBackendWithRedirects;
+
+				_prefix = config.Prefix;
+				if (_prefix.Length > 0 && !_prefix.EndsWith("/", StringComparison.Ordinal))
+				{
+					_prefix += "/";
+				}
 			}
 
 			public void Dispose()
@@ -119,7 +126,7 @@ namespace Horde.Build.Storage
 				}
 			}
 
-			static string GetBlobPath(BlobId blobId) => $"{blobId}.blob";
+			string GetBlobPath(BlobId blobId) => $"{_prefix}{blobId}.blob";
 
 			#region Blobs
 
diff --git a/Engine/Source/Programs/Horde/Horde.Build/Tools/ToolCollection.cs b/Engine/Source/Programs/Horde/Horde.Build/Tools/ToolCollection.cs
--- a/Engine/Source/Programs/Horde/Horde.Build/Tools/ToolCollection.cs
+++ b/Engine/Source/Programs/Horde/Horde.Build/Tools/ToolCollection.cs
@@ -5,6 +5,7 @@ using EpicGames.Horde.Storage;
 using EpicGames.Serialization;
 using Horde.Build.Acls;
 using Horde.Build.Server;
+using Horde.Build.Storage;
 using Horde.Build.Utilities;
 using HordeCommon;
 using Microsoft.Extensions.Logging;
@@ -74,7 +75,7 @@ namespace Horde.Build.Tools
 			_tools = new VersionedCollection<ToolId, Tool>(mongoService, "Tools", redisService, s_baseKey, s_types);
 			_storage = storage;
 			_clock = clock;
-			_namespaceId = options.Value.ToolNamespaceId;
+			_namespaceId = Namespace.Tools;
 			_logger = logger;
 		}
 
