Horde: Add an option to return log data from the new log storage backend. Does not currently implement search.

diff --git a/Engine/Source/Programs/Horde/Horde.Build.Tests/DatabaseIntegrationTest.cs b/Engine/Source/Programs/Horde/Horde.Build.Tests/DatabaseIntegrationTest.cs
--- a/Engine/Source/Programs/Horde/Horde.Build.Tests/DatabaseIntegrationTest.cs
+++ b/Engine/Source/Programs/Horde/Horde.Build.Tests/DatabaseIntegrationTest.cs
@@ -14,6 +14,13 @@ using StackExchange.Redis;
 
 namespace Horde.Build.Tests
 {
+	// Stub for fulfilling IOptions interface during testing
+	public sealed class TestOptions<T> : IOptions<T> where T : class
+	{
+		public TestOptions(T options) => Value = options;
+		public T Value { get; }
+	}
+
 	// Stub for fulfilling IOptionsMonitor interface during testing
 	public sealed class TestOptionsMonitor<T> : IOptionsMonitor<T>
         where T : class, new()
diff --git a/Engine/Source/Programs/Horde/Horde.Build.Tests/LogFileServiceTests.cs b/Engine/Source/Programs/Horde/Horde.Build.Tests/LogFileServiceTests.cs
--- a/Engine/Source/Programs/Horde/Horde.Build.Tests/LogFileServiceTests.cs
+++ b/Engine/Source/Programs/Horde/Horde.Build.Tests/LogFileServiceTests.cs
@@ -46,8 +46,8 @@ namespace Horde.Build.Tests
 			ILogBuilder logBuilder = new RedisLogBuilder(GetRedisServiceSingleton().ConnectionPool, NullLogger.Instance);
 			_logStorage = new PersistentLogStorage(new MemoryStorageBackend().ForType<PersistentLogStorage>(), NullLogger<PersistentLogStorage>.Instance);
 			_clock = new FakeClock();
-			TestOptionsMonitor<ServerSettings> settingsOpts = new (new ServerSettings());
-			_logFileService = new LogFileService(logFileCollection, null!, logBuilder, _logStorage, _clock, settingsOpts, logger);
+			TestOptions<ServerSettings> settingsOpts = new (new ServerSettings());
+			_logFileService = new LogFileService(logFileCollection, null!, logBuilder, _logStorage, _clock, null!, settingsOpts, logger);
         }
 
 		protected override void Dispose(bool disposing)
diff --git a/Engine/Source/Programs/Horde/Horde.Build.Tests/LogIndexingTests.cs b/Engine/Source/Programs/Horde/Horde.Build.Tests/LogIndexingTests.cs
--- a/Engine/Source/Programs/Horde/Horde.Build.Tests/LogIndexingTests.cs
+++ b/Engine/Source/Programs/Horde/Horde.Build.Tests/LogIndexingTests.cs
@@ -49,8 +49,8 @@ namespace Horde.Build.Tests
 			ILogBuilder logBuilder = new LocalLogBuilder();
 			_nullLogStorage = new NullLogStorage();
 			_logStorage = new LocalLogStorage(20, _nullLogStorage);
-			TestOptionsMonitor<ServerSettings> settingsOpts = new (new ServerSettings());
-			_logFileService = new LogFileService(logFileCollection, null!, logBuilder, _logStorage, new FakeClock(), settingsOpts, logger);
+			TestOptions<ServerSettings> settingsOpts = new (new ServerSettings());
+			_logFileService = new LogFileService(logFileCollection, null!, logBuilder, _logStorage, new FakeClock(), null!, settingsOpts, logger);
 		}
 
 		protected override void Dispose(bool disposing)
diff --git a/Engine/Source/Programs/Horde/Horde.Build/Logs/LogFileService.cs b/Engine/Source/Programs/Horde/Horde.Build/Logs/LogFileService.cs
--- a/Engine/Source/Programs/Horde/Horde.Build/Logs/LogFileService.cs
+++ b/Engine/Source/Programs/Horde/Horde.Build/Logs/LogFileService.cs
@@ -13,9 +13,12 @@ using System.Text.Json;
 using System.Threading;
 using System.Threading.Tasks;
 using EpicGames.Core;
+using EpicGames.Horde.Logs;
+using EpicGames.Horde.Storage;
 using Horde.Build.Agents.Sessions;
 using Horde.Build.Jobs;
 using Horde.Build.Logs.Data;
+using Horde.Build.Storage;
 using Horde.Build.Utilities;
 using HordeCommon;
 using Microsoft.Extensions.Caching.Memory;
@@ -238,7 +241,8 @@ namespace Horde.Build.Logs
 		private readonly ILogEventCollection _logEvents;
 		private readonly ILogStorage _storage;
 		private readonly ILogBuilder _builder;
-		private readonly IOptionsMonitor<ServerSettings> _settings;
+		private readonly StorageService _storageService;
+		private readonly IOptions<ServerSettings> _settings;
 
 		// Lock object for the <see cref="_writeTasks"/> and <see cref="_writeChunks"/> members
 		private readonly object _writeLock = new object();
@@ -409,12 +413,159 @@ namespace Horde.Build.Logs
 			public override void Write(byte[] buffer, int offset, int count) => throw new NotImplementedException();
 		}
 
+		/// <summary>
+		/// Streams log data to a caller
+		/// </summary>
+		class NewLoggerResponseStream : Stream
+		{
+			readonly TreeReader _reader;
+			readonly LogNode _rootNode;
+
+			/// <summary>
+			/// Starting offset within the file of the data to return 
+			/// </summary>
+			readonly long _responseOffset;
+
+			/// <summary>
+			/// Length of data to return
+			/// </summary>
+			readonly long _responseLength;
+
+			/// <summary>
+			/// Current offset within the stream
+			/// </summary>
+			long _currentOffset;
+
+			/// <summary>
+			/// The current chunk index
+			/// </summary>
+			int _chunkIdx;
+
+			/// <summary>
+			/// Buffer containing a message for missing data
+			/// </summary>
+			ReadOnlyMemory<byte> _sourceBuffer;
+
+			/// <summary>
+			/// Offset within the source buffer
+			/// </summary>
+			int _sourcePos;
+
+			/// <summary>
+			/// Length of the source buffer being copied from
+			/// </summary>
+			int _sourceEnd;
+
+			/// <summary>
+			/// Constructor
+			/// </summary>
+			public NewLoggerResponseStream(TreeReader reader, LogNode rootNode, long offset, long length)
+			{
+				_reader = reader;
+				_rootNode = rootNode;
+
+				_responseOffset = offset;
+				_responseLength = length;
+
+				_currentOffset = offset;
+
+				_chunkIdx = rootNode.TextChunkRefs.GetChunkForOffset(offset);
+				_sourceBuffer = null!;
+			}
+
+			/// <inheritdoc/>
+			public override bool CanRead => true;
+
+			/// <inheritdoc/>
+			public override bool CanSeek => false;
+
+			/// <inheritdoc/>
+			public override bool CanWrite => false;
+
+			/// <inheritdoc/>
+			public override long Length => _responseLength;
+
+			/// <inheritdoc/>
+			public override long Position
+			{
+				get => _currentOffset - _responseOffset;
+				set => throw new NotImplementedException();
+			}
+
+			/// <inheritdoc/>
+			public override void Flush()
+			{
+			}
+
+			/// <inheritdoc/>
+			public override int Read(byte[] buffer, int offset, int count)
+			{
+				return ReadAsync(buffer, offset, count, CancellationToken.None).Result;
+			}
+
+			/// <inheritdoc/>
+			public override async Task<int> ReadAsync(byte[] buffer, int offset, int length, CancellationToken cancellationToken)
+			{
+				return await ReadAsync(buffer.AsMemory(offset, length), cancellationToken);
+			}
+
+			/// <inheritdoc/>
+			public override async ValueTask<int> ReadAsync(Memory<byte> buffer, CancellationToken cancellationToken)
+			{
+				int readBytes = 0;
+				while (readBytes < buffer.Length)
+				{
+					if (_sourcePos < _sourceEnd)
+					{
+						// Try to copy from the current buffer
+						int blockSize = Math.Min(_sourceEnd - _sourcePos, buffer.Length - readBytes);
+						_sourceBuffer.Slice(_sourcePos, blockSize).Span.CopyTo(buffer.Slice(readBytes).Span);
+						_currentOffset += blockSize;
+						readBytes += blockSize;
+						_sourcePos += blockSize;
+					}
+					else if (_currentOffset < _responseOffset + _responseLength)
+					{
+						// Move to the right chunk
+						while (_chunkIdx + 1 < _rootNode.TextChunkRefs.Count && _currentOffset >= _rootNode.TextChunkRefs[_chunkIdx + 1].Offset)
+						{
+							_chunkIdx++;
+						}
+
+						// Get the chunk data
+						LogChunkRef chunk = _rootNode.TextChunkRefs[_chunkIdx];
+						LogChunkNode chunkNode = await chunk.ExpandCopyAsync(_reader, cancellationToken);
+
+						// Get the source data
+						_sourceBuffer = chunkNode.Data;
+						_sourcePos = (int)(_currentOffset - chunk.Offset);
+						_sourceEnd = (int)Math.Min(_sourceBuffer.Length, (_responseOffset + _responseLength) - chunk.Offset);
+					}
+					else
+					{
+						// End of the log
+						break;
+					}
+				}
+				return readBytes;
+			}
+
+			/// <inheritdoc/>
+			public override long Seek(long offset, SeekOrigin origin) => throw new NotImplementedException();
+
+			/// <inheritdoc/>
+			public override void SetLength(long value) => throw new NotImplementedException();
+
+			/// <inheritdoc/>
+			public override void Write(byte[] buffer, int offset, int count) => throw new NotImplementedException();
+		}
+
 		readonly ITicker _ticker;
 
 		/// <summary>
 		/// Constructor
 		/// </summary>
-		public LogFileService(ILogFileCollection logFiles, ILogEventCollection logEvents, ILogBuilder builder, ILogStorage storage, IClock clock, IOptionsMonitor<ServerSettings> settings, ILogger<LogFileService> logger)
+		public LogFileService(ILogFileCollection logFiles, ILogEventCollection logEvents, ILogBuilder builder, ILogStorage storage, IClock clock, StorageService storageService, IOptions<ServerSettings> settings, ILogger<LogFileService> logger)
 		{
 			_logFiles = logFiles;
 			_logEvents = logEvents;
@@ -422,6 +573,7 @@ namespace Horde.Build.Logs
 			_builder = builder;
 			_storage = storage;
 			_ticker = clock.AddSharedTicker<LogFileService>(TimeSpan.FromSeconds(30.0), TickAsync, logger);
+			_storageService = storageService;
 			_settings = settings;
 			_logger = logger;
 		}
@@ -870,29 +1022,62 @@ namespace Horde.Build.Logs
 		/// <inheritdoc/>
 		public async Task<Stream> OpenRawStreamAsync(ILogFile logFile, long offset, long length, CancellationToken cancellationToken)
 		{
-			if (logFile.Chunks.Count == 0)
+			if (_settings.Value.FeatureFlags.EnableNewLogger)
 			{
-				return new MemoryStream(Array.Empty<byte>(), false);
+				TreeReader reader = await GetTreeReaderAsync(cancellationToken);
+
+				LogNode? root = await reader.TryReadNodeAsync<LogNode>(logFile.RefName, cancellationToken: cancellationToken);
+				if (root == null || root.TextChunkRefs.Count == 0)
+				{
+					return new MemoryStream(Array.Empty<byte>(), false);
+				}
+				else
+				{
+					int lastChunkIdx = root.TextChunkRefs.Count - 1;
+
+					// Clamp the length of the request
+					LogChunkRef lastChunk = root.TextChunkRefs[lastChunkIdx];
+					if (length > lastChunk.Offset)
+					{
+						long lastChunkLength = lastChunk.Length;
+						if (lastChunkLength <= 0)
+						{
+							LogChunkNode lastChunkNode = await lastChunk.ExpandAsync(reader, cancellationToken);
+							lastChunkLength = lastChunkNode.Length;
+						}
+						length = Math.Min(length, (lastChunk.Offset + lastChunkLength) - offset);
+					}
+
+					// Create the new stream
+					return new NewLoggerResponseStream(reader, root, offset, length);
+				}
 			}
 			else
 			{
-				int lastChunkIdx = logFile.Chunks.Count - 1;
-
-				// Clamp the length of the request
-				ILogChunk lastChunk = logFile.Chunks[lastChunkIdx];
-				if (length > lastChunk.Offset)
+				if (logFile.Chunks.Count == 0)
 				{
-					long lastChunkLength = lastChunk.Length;
-					if (lastChunkLength <= 0)
+					return new MemoryStream(Array.Empty<byte>(), false);
+				}
+				else
+				{
+					int lastChunkIdx = logFile.Chunks.Count - 1;
+
+					// Clamp the length of the request
+					ILogChunk lastChunk = logFile.Chunks[lastChunkIdx];
+					if (length > lastChunk.Offset)
 					{
-						LogChunkData lastChunkData = await ReadChunkAsync(logFile, lastChunkIdx);
-						lastChunkLength = lastChunkData.Length;
+						long lastChunkLength = lastChunk.Length;
+						if (lastChunkLength <= 0)
+						{
+							LogChunkData lastChunkData = await ReadChunkAsync(logFile, lastChunkIdx);
+							lastChunkLength = lastChunkData.Length;
+						}
+						length = Math.Min(length, (lastChunk.Offset + lastChunkLength) - offset);
 					}
-					length = Math.Min(length, (lastChunk.Offset + lastChunkLength) - offset);
-				}
 
-				// Create the new stream
-				return new ResponseStream(this, logFile, offset, length);
+					// Create the new stream
+					return new ResponseStream(this, logFile, offset, length);
+				}
 			}
 		}
 
@@ -975,27 +1160,64 @@ namespace Horde.Build.Logs
 			}
 		}
 
+		async Task<TreeReader> GetTreeReaderAsync(CancellationToken cancellationToken)
+		{
+			IStorageClient store = await _storageService.GetClientAsync(Namespace.Logs, cancellationToken);
+			return new TreeReader(store, _logFileCache, _logger);
+		}
+
 		/// <inheritdoc/>
 		public async Task<(int, long)> GetLineOffsetAsync(ILogFile logFile, int lineIdx, CancellationToken cancellationToken)
 		{
-			int chunkIdx = logFile.Chunks.GetChunkForLine(lineIdx);
+			if (_settings.Value.FeatureFlags.EnableNewLogger)
+			{
+				TreeReader reader = await GetTreeReaderAsync(cancellationToken);
 
-			ILogChunk chunk = logFile.Chunks[chunkIdx];
-			LogChunkData chunkData = await ReadChunkAsync(logFile, chunkIdx);
+				LogNode? root = await reader.TryReadNodeAsync<LogNode>(logFile.RefName, cancellationToken: cancellationToken);
+				if (root == null)
+				{
+					return (0, 0);
+				}
 
-			if (lineIdx < chunk.LineIndex)
-			{
-				lineIdx = chunk.LineIndex;
-			}
+				int chunkIdx = root.TextChunkRefs.GetChunkForLine(lineIdx);
+				LogChunkRef chunk = root.TextChunkRefs[chunkIdx];
+				LogChunkNode chunkData = await chunk.ExpandAsync(reader, cancellationToken);
 
-			int maxLineIndex = chunk.LineIndex + chunkData.LineCount;
-			if (lineIdx >= maxLineIndex)
-			{
-				lineIdx = maxLineIndex;
+				if (lineIdx < chunk.LineIndex)
+				{
+					lineIdx = chunk.LineIndex;
+				}
+
+				int maxLineIndex = chunk.LineIndex + chunkData.LineCount;
+				if (lineIdx >= maxLineIndex)
+				{
+					lineIdx = maxLineIndex;
+				}
+
+				long offset = chunk.Offset + chunkData.LineOffsets[lineIdx - chunk.LineIndex];
+				return (lineIdx, offset);
 			}
+			else
+			{
+				int chunkIdx = logFile.Chunks.GetChunkForLine(lineIdx);
+
+				ILogChunk chunk = logFile.Chunks[chunkIdx];
+				LogChunkData chunkData = await ReadChunkAsync(logFile, chunkIdx);
 
-			long offset = chunk.Offset + chunkData.GetLineOffsetWithinChunk(lineIdx - chunk.LineIndex);
-			return (lineIdx, offset);
+				if (lineIdx < chunk.LineIndex)
+				{
+					lineIdx = chunk.LineIndex;
+				}
+
+				int maxLineIndex = chunk.LineIndex + chunkData.LineCount;
+				if (lineIdx >= maxLineIndex)
+				{
+					lineIdx = maxLineIndex;
+				}
+
+				long offset = chunk.Offset + chunkData.GetLineOffsetWithinChunk(lineIdx - chunk.LineIndex);
+				return (lineIdx, offset);
+			}
 		}
 
 		/// <summary>
@@ -1039,7 +1261,7 @@ namespace Horde.Build.Logs
 				await _builder.CompleteChunkAsync(logId, offset);
 			}
 
-			if (_settings.CurrentValue.FeatureFlags.LimitConcurrentLogChunkWriting)
+			if (_settings.Value.FeatureFlags.LimitConcurrentLogChunkWriting)
 			{
 				// Flush all the chunks and await completion instead of running them async
 				await WriteCompleteChunksV2Async(flushChunks, true, cancellationToken);
diff --git a/Engine/Source/Programs/Horde/Horde.Build/ServerSettings.cs b/Engine/Source/Programs/Horde/Horde.Build/ServerSettings.cs
--- a/Engine/Source/Programs/Horde/Horde.Build/ServerSettings.cs
+++ b/Engine/Source/Programs/Horde/Horde.Build/ServerSettings.cs
@@ -180,6 +180,10 @@ namespace Horde.Build
 		/// </summary>
 		public bool LimitConcurrentLogChunkWriting { get; set; } = false;
 
+		/// <summary>
+		/// Whether to use the new log storage backend
+		/// </summary>
+		public bool EnableNewLogger { get; set; } = false;
 	}
 
 	/// <summary>
diff --git a/Engine/Source/Programs/Shared/EpicGames.Horde/Logs/LogChunkNode.cs b/Engine/Source/Programs/Shared/EpicGames.Horde/Logs/LogChunkNode.cs
--- a/Engine/Source/Programs/Shared/EpicGames.Horde/Logs/LogChunkNode.cs
+++ b/Engine/Source/Programs/Shared/EpicGames.Horde/Logs/LogChunkNode.cs
@@ -692,6 +692,38 @@ namespace EpicGames.Horde.Logs
 	/// </summary>
 	public static class LogChunkExtensions
 	{
+		/// <summary>
+		/// Gets the chunk index containing the given offset.
+		/// </summary>
+		/// <param name="chunks">The chunks to search</param>
+		/// <param name="offset">The offset to search for</param>
+		/// <returns>The chunk index containing the given offset</returns>
+		public static int GetChunkForOffset(this IReadOnlyList<LogChunkRef> chunks, long offset)
+		{
+			int chunkIndex = chunks.BinarySearch(x => x.Offset, offset);
+			if (chunkIndex < 0)
+			{
+				chunkIndex = ~chunkIndex - 1;
+			}
+			return chunkIndex;
+		}
+
+		/// <summary>
+		/// Gets the starting chunk index for the given line
+		/// </summary>
+		/// <param name="chunks">The chunks to search</param>
+		/// <param name="lineIndex">Index of the line to query</param>
+		/// <returns>Index of the chunk to fetch</returns>
+		public static int GetChunkForLine(this IReadOnlyList<LogChunkRef> chunks, int lineIndex)
+		{
+			int chunkIndex = chunks.BinarySearch(x => x.LineIndex, lineIndex);
+			if (chunkIndex < 0)
+			{
+				chunkIndex = ~chunkIndex - 1;
+			}
+			return chunkIndex;
+		}
+
 		/// <summary>
 		/// Converts a log text instance to plain text
 		/// </summary>
