Horde: Remove byte offset-based addressing for log data.

diff --git a/Engine/Source/Programs/Horde/Horde.Build.Tests/LogFileServiceTests.cs b/Engine/Source/Programs/Horde/Horde.Build.Tests/LogFileServiceTests.cs
--- a/Engine/Source/Programs/Horde/Horde.Build.Tests/LogFileServiceTests.cs
+++ b/Engine/Source/Programs/Horde/Horde.Build.Tests/LogFileServiceTests.cs
@@ -8,6 +8,7 @@ using System.Security.Claims;
 using System.Text;
 using System.Threading;
 using System.Threading.Tasks;
+using EpicGames.Core;
 using Horde.Build.Agents.Sessions;
 using Horde.Build.Jobs;
 using Horde.Build.Logs;
@@ -91,10 +92,11 @@ namespace Horde.Build.Tests
             Assert.AreEqual(6 + 8 + 4, metadata.Length);
             Assert.AreEqual(4, metadata.MaxLineIndex);
 
-            Assert.AreEqual((0, 0), await _logFileService.GetLineOffsetAsync(logFile, 0, CancellationToken.None));
-            Assert.AreEqual((1, 6), await _logFileService.GetLineOffsetAsync(logFile, 1, CancellationToken.None));
-            Assert.AreEqual((2, 10), await _logFileService.GetLineOffsetAsync(logFile, 2, CancellationToken.None));
-            Assert.AreEqual((3, 14), await _logFileService.GetLineOffsetAsync(logFile, 3, CancellationToken.None));
+			List<Utf8String> lines = await _logFileService.ReadLinesAsync(logFile, 0, null, CancellationToken.None);
+			Assert.AreEqual("hello", lines[0].ToString());
+			Assert.AreEqual("foo", lines[1].ToString());
+			Assert.AreEqual("bar", lines[2].ToString());
+			Assert.AreEqual("baz", lines[3].ToString());
         }
         
         [TestMethod]
@@ -126,14 +128,23 @@ namespace Horde.Build.Tests
         private static async Task AssertLineOffset(ILogFileService logFileService, LogId logFileId, int lineIndex, int clampedLineIndex, long offset)
         {
 	        ILogFile? logFile = await logFileService.GetLogFileAsync(logFileId, CancellationToken.None);
-	        Assert.AreEqual((clampedLineIndex, offset), await logFileService.GetLineOffsetAsync(logFile!, lineIndex, CancellationToken.None));
+
+			List<Utf8String> lines = await logFileService.ReadLinesAsync(logFile!, 0, lineIndex + 1, CancellationToken.None);
+			Assert.AreEqual(clampedLineIndex, Math.Min(lineIndex, lines.Count));
+			Assert.AreEqual(offset, lines.Take(lineIndex).Sum(x => x.Length + 1));
         }
         
         protected static async Task<string> ReadLogFile(ILogFileService logFileService, ILogFile logFile, long offset, long length)
         {
-	        using Stream stream = await logFileService.OpenRawStreamAsync(logFile, offset, length);
-			using StreamReader streamReader = new StreamReader(stream);
-	        return await streamReader.ReadToEndAsync();
+			using Stream stream = await logFileService.OpenRawStreamAsync(logFile);
+
+			byte[] prefix = new byte[(int)offset];
+			await stream.ReadGreedyAsync(prefix);
+
+			byte[] data = new byte[(int)length];
+			int actualLength = await stream.ReadGreedyAsync(data);
+
+			return Encoding.UTF8.GetString(data.AsSpan(0, actualLength));
         }
         
         public async Task WriteLogLifecycle(ILogFileService lfs, int maxChunkLength)
diff --git a/Engine/Source/Programs/Horde/Horde.Build.Tests/LogIndexingTests.cs b/Engine/Source/Programs/Horde/Horde.Build.Tests/LogIndexingTests.cs
--- a/Engine/Source/Programs/Horde/Horde.Build.Tests/LogIndexingTests.cs
+++ b/Engine/Source/Programs/Horde/Horde.Build.Tests/LogIndexingTests.cs
@@ -93,7 +93,7 @@ namespace Horde.Build.Tests
 
 			// Read the data back out and check it's the same
 			byte[] readData = new byte[_data.Length];
-			using (Stream stream = await _logFileService.OpenRawStreamAsync(logFile, 0, _data.Length))
+			using (Stream stream = await _logFileService.OpenRawStreamAsync(logFile))
 			{
 				int readSize = await stream.ReadAsync(readData, 0, readData.Length);
 				Assert.AreEqual(readData.Length, readSize);
diff --git a/Engine/Source/Programs/Horde/Horde.Build/Logs/LogFileService.cs b/Engine/Source/Programs/Horde/Horde.Build/Logs/LogFileService.cs
--- a/Engine/Source/Programs/Horde/Horde.Build/Logs/LogFileService.cs
+++ b/Engine/Source/Programs/Horde/Horde.Build/Logs/LogFileService.cs
@@ -93,6 +93,16 @@ namespace Horde.Build.Logs
 		/// <returns>List of logfile documents</returns>
 		Task<List<ILogFile>> GetLogFilesAsync(int? index = null, int? count = null, CancellationToken cancellationToken = default);
 
+		/// <summary>
+		/// Read a set of lines from the given log file
+		/// </summary>
+		/// <param name="logFile">Log file to read</param>
+		/// <param name="index">Index of the first line to read</param>
+		/// <param name="count">Maximum number of lines to return</param>
+		/// <param name="cancellationToken">Cancellation token for the operation</param>
+		/// <returns>List of lines</returns>
+		Task<List<Utf8String>> ReadLinesAsync(ILogFile logFile, int index, int? count, CancellationToken cancellationToken = default);
+
 		/// <summary>
 		/// Writes out chunk data and assigns to a file
 		/// </summary>
@@ -168,31 +178,18 @@ namespace Horde.Build.Logs
 		/// Gets lines from the given log 
 		/// </summary>
 		/// <param name="logFile">The log file</param>
-		/// <param name="offset">Offset of the data to return</param>
-		/// <param name="length">Length of the data to return</param>
 		/// <param name="cancellationToken">Cancellation token for the call</param>
 		/// <returns>Data for the requested range</returns>
-		Task<Stream> OpenRawStreamAsync(ILogFile logFile, long offset, long length, CancellationToken cancellationToken = default);
+		Task<Stream> OpenRawStreamAsync(ILogFile logFile, CancellationToken cancellationToken = default);
 
 		/// <summary>
 		/// Parses a stream of json text and outputs plain text
 		/// </summary>
 		/// <param name="logFile">The log file to query</param>
-		/// <param name="offset">Offset within the data to copy from</param>
-		/// <param name="length">Length of the data to copy</param>
 		/// <param name="outputStream">Output stream to receive the text data</param>
 		/// <param name="cancellationToken">Cancellation token for the call</param>
 		/// <returns>Async text</returns>
-		Task CopyPlainTextStreamAsync(ILogFile logFile, long offset, long length, Stream outputStream, CancellationToken cancellationToken = default);
-
-		/// <summary>
-		/// Gets the offset of the given line number
-		/// </summary>
-		/// <param name="logFile">The log file to search</param>
-		/// <param name="lineIdx">The line index to retrieve the offset for</param>
-		/// <param name="cancellationToken">Cancellation token for the call</param>
-		/// <returns>The actual clamped line number and offset</returns>
-		Task<(int, long)> GetLineOffsetAsync(ILogFile logFile, int lineIdx, CancellationToken cancellationToken);
+		Task CopyPlainTextStreamAsync(ILogFile logFile, Stream outputStream, CancellationToken cancellationToken = default);
 
 		/// <summary>
 		/// Search for the specified text in a log file
@@ -217,14 +214,12 @@ namespace Horde.Build.Logs
 		/// </summary>
 		/// <param name="logFileService">The log file service</param>
 		/// <param name="logFile">The log file to query</param>
-		/// <param name="offset">Offset within the log file to copy</param>
-		/// <param name="length">Length of the data to copy</param>
 		/// <param name="outputStream">Output stream to receive the text data</param>
 		/// <param name="cancellationToken">Cancellation token for the call</param>
 		/// <returns>Async text</returns>
-		public static async Task CopyRawStreamAsync(this ILogFileService logFileService, ILogFile logFile, long offset, long length, Stream outputStream, CancellationToken cancellationToken)
+		public static async Task CopyRawStreamAsync(this ILogFileService logFileService, ILogFile logFile, Stream outputStream, CancellationToken cancellationToken)
 		{
-			await using Stream stream = await logFileService.OpenRawStreamAsync(logFile, offset, length, cancellationToken);
+			await using Stream stream = await logFileService.OpenRawStreamAsync(logFile, cancellationToken);
 			await stream.CopyToAsync(outputStream, cancellationToken);
 		}
 	}
@@ -646,6 +641,34 @@ namespace Horde.Build.Logs
 			return _logFiles.GetLogFilesAsync(index, count, cancellationToken);
 		}
 
+		/// <inheritdoc/>
+		public async Task<List<Utf8String>> ReadLinesAsync(ILogFile logFile, int index, int? count, CancellationToken cancellationToken)
+		{
+			(_, long minOffset) = await GetLineOffsetAsync(logFile, index, cancellationToken);
+			(_, long maxOffset) = await GetLineOffsetAsync(logFile, (count == null) ? Int32.MaxValue : (index + Math.Min(count.Value, Int32.MaxValue - index)), cancellationToken);
+
+			byte[] result;
+			using (System.IO.Stream stream = await OpenRawStreamAsync(logFile, minOffset, maxOffset - minOffset, cancellationToken))
+			{
+				result = new byte[stream.Length];
+				await stream.ReadFixedSizeDataAsync(result, 0, result.Length);
+			}
+
+			List<Utf8String> lines = new List<Utf8String>();
+
+			int offset = 0;
+			for (int idx = 0; idx < result.Length; idx++)
+			{
+				if (result[idx] == (byte)'\n')
+				{
+					lines.Add(new Utf8String(result.AsMemory(offset, idx - offset)));
+					offset = idx + 1;
+				}
+			}
+
+			return lines;
+		}
+
 		class WriteState
 		{
 			public long _offset;
@@ -971,52 +994,28 @@ namespace Horde.Build.Logs
 			scope.Span.SetTag("LineIndex", lineIndex.ToString(CultureInfo.InvariantCulture));
 			scope.Span.SetTag("LineCount", lineCount.ToString(CultureInfo.InvariantCulture));
 
-			(_, long minOffset) = await GetLineOffsetAsync(logFile, lineIndex, cancellationToken);
-			(_, long maxOffset) = await GetLineOffsetAsync(logFile, lineIndex + lineCount, cancellationToken);
-
-			byte[] data = new byte[maxOffset - minOffset];
-			using (Stream stream = await OpenRawStreamAsync(logFile, minOffset, maxOffset - minOffset, cancellationToken))
-			{
-				int length = await stream.ReadAsync(data.AsMemory(), cancellationToken);
-				if(length != data.Length)
-				{
-					_logger.LogWarning("Read less than expected from log stream (Expected {Expected}, Got {Got})", data.Length, length);
-				}
-				return ParseEventData(data.AsSpan(0, length));
-			}
-		}
-
-		/// <summary>
-		/// Parses event data from a buffer
-		/// </summary>
-		/// <param name="data">Data to parse</param>
-		/// <returns>Parsed event data</returns>
-		private ILogEventData ParseEventData(ReadOnlySpan<byte> data)
-		{
-			List<LogEventLine> lines = new List<LogEventLine>();
+			List<Utf8String> lines = await ReadLinesAsync(logFile, lineIndex, lineCount, cancellationToken);
+			List<LogEventLine> eventLines = new List<LogEventLine>(lines.Count);
 
-			ReadOnlySpan<byte> remainingData = data;
-			while(remainingData.Length > 0)
+			foreach (Utf8String line in lines)
 			{
-				int endOfLine = remainingData.IndexOf((byte)'\n');
-				if(endOfLine == -1)
-				{
-					break;
-				}
-
-				ReadOnlySpan<byte> lineData = remainingData.Slice(0, endOfLine);
 				try
 				{
-					lines.Add(new LogEventLine(lineData));
+					eventLines.Add(new LogEventLine(line.Span));
 				}
-				catch(JsonException ex)
+				catch (JsonException ex)
 				{
-					_logger.LogWarning(ex, "Unable to parse line from log file: {Line}", Encoding.UTF8.GetString(lineData));
+					_logger.LogWarning(ex, "Unable to parse line from log file: {Line}", line);
 				}
-				remainingData = remainingData.Slice(endOfLine + 1);
 			}
 
-			return new LogEventData(lines);
+			return new LogEventData(eventLines);
+		}
+
+		/// <inheritdoc/>
+		public Task<Stream> OpenRawStreamAsync(ILogFile logFile, CancellationToken cancellationToken)
+		{
+			return OpenRawStreamAsync(logFile, 0, Int64.MaxValue, cancellationToken);
 		}
 
 		/// <inheritdoc/>
@@ -1099,8 +1098,11 @@ namespace Horde.Build.Logs
 		}
 
 		/// <inheritdoc/>
-		public async Task CopyPlainTextStreamAsync(ILogFile logFile, long offset, long length, Stream outputStream, CancellationToken cancellationToken)
+		public async Task CopyPlainTextStreamAsync(ILogFile logFile, Stream outputStream, CancellationToken cancellationToken)
 		{
+			long offset = 0;
+			long length = Int64.MaxValue;
+
 			using (Stream stream = await OpenRawStreamAsync(logFile, 0, Int64.MaxValue, cancellationToken))
 			{
 				byte[] readBuffer = new byte[4096];
diff --git a/Engine/Source/Programs/Horde/Horde.Build/Logs/LogsController.cs b/Engine/Source/Programs/Horde/Horde.Build/Logs/LogsController.cs
--- a/Engine/Source/Programs/Horde/Horde.Build/Logs/LogsController.cs
+++ b/Engine/Source/Programs/Horde/Horde.Build/Logs/LogsController.cs
@@ -3,10 +3,12 @@
 using System;
 using System.Collections.Generic;
 using System.IO;
+using System.Linq;
 using System.Security.Claims;
 using System.Text;
 using System.Threading;
 using System.Threading.Tasks;
+using EpicGames.Core;
 using EpicGames.Horde.Storage;
 using Horde.Build.Acls;
 using Horde.Build.Issues;
@@ -125,8 +127,6 @@ namespace Horde.Build.Logs
 		/// </summary>
 		/// <param name="logFileId">Id of the log file to get information about</param>
 		/// <param name="format">Format for the returned data</param>
-		/// <param name="offset">The log offset in bytes</param>
-		/// <param name="length">Number of bytes to return</param>
 		/// <param name="fileName">Name of the default filename to download</param>
 		/// <param name="download">Whether to download the file rather than display in the browser</param>
 		/// <param name="cancellationToken">Cancellation token for the request</param>
@@ -136,8 +136,6 @@ namespace Horde.Build.Logs
 		public async Task<ActionResult> GetLogData(
 			LogId logFileId,
 			[FromQuery] LogOutputFormat format = LogOutputFormat.Raw,
-			[FromQuery] long offset = 0,
-			[FromQuery] long length = Int64.MaxValue,
 			[FromQuery] string? fileName = null,
 			[FromQuery] bool download = false,
 			CancellationToken cancellationToken = default)
@@ -155,11 +153,11 @@ namespace Horde.Build.Logs
 			Func<Stream, ActionContext, Task> copyTask;
 			if (format == LogOutputFormat.Text && logFile.Type == LogType.Json)
 			{
-				copyTask = (outputStream, context) => _logFileService.CopyPlainTextStreamAsync(logFile, offset, length, outputStream, cancellationToken);
+				copyTask = (outputStream, context) => _logFileService.CopyPlainTextStreamAsync(logFile, outputStream, cancellationToken);
 			}
 			else
 			{
-				copyTask = (outputStream, context) => _logFileService.CopyRawStreamAsync(logFile, offset, length, outputStream, cancellationToken);
+				copyTask = (outputStream, context) => _logFileService.CopyRawStreamAsync(logFile, outputStream, cancellationToken);
 			}
 
 			return new CustomFileCallbackResult(fileName ?? $"log-{logFileId}.txt", "text/plain", !download, copyTask);
@@ -175,7 +173,7 @@ namespace Horde.Build.Logs
 		/// <returns>Information about the requested project</returns>
 		[HttpGet]
 		[Route("/api/v1/logs/{logFileId}/lines")]
-		public async Task<ActionResult> GetLogLines(LogId logFileId, [FromQuery] int index = 0, [FromQuery] int count = Int32.MaxValue, CancellationToken cancellationToken = default)
+		public async Task<ActionResult> GetLogLines(LogId logFileId, [FromQuery] int index = 0, [FromQuery] int? count = null, CancellationToken cancellationToken = default)
 		{
 			ILogFile? logFile = await _logFileService.GetLogFileAsync(logFileId, cancellationToken);
 			if (logFile == null)
@@ -189,19 +187,8 @@ namespace Horde.Build.Logs
 
 			LogMetadata metadata = await _logFileService.GetMetadataAsync(logFile, cancellationToken);
 
-			(int minIndex, long minOffset) = await _logFileService.GetLineOffsetAsync(logFile, index, cancellationToken);
-			(int maxIndex, long maxOffset) = await _logFileService.GetLineOffsetAsync(logFile, index + Math.Min(count, Int32.MaxValue - index), cancellationToken);
-			index = minIndex;
-			count = maxIndex - minIndex;
-
-			byte[] result;
-			using (System.IO.Stream stream = await _logFileService.OpenRawStreamAsync(logFile, minOffset, maxOffset - minOffset, cancellationToken))
-			{
-				result = new byte[stream.Length];
-				await stream.ReadFixedSizeDataAsync(result, 0, result.Length);
-			}
-
-			using (MemoryStream stream = new MemoryStream(result.Length + (count * 20)))
+			List<Utf8String> lines = await _logFileService.ReadLinesAsync(logFile, index, count, cancellationToken);
+			using (MemoryStream stream = new MemoryStream(lines.Sum(x => x.Length) + (lines.Count * 20)))
 			{
 				stream.WriteByte((byte)'{');
 
@@ -210,69 +197,39 @@ namespace Horde.Build.Logs
 				stream.Write(Encoding.UTF8.GetBytes($"\"maxLineIndex\":{metadata.MaxLineIndex},"));
 				stream.Write(Encoding.UTF8.GetBytes($"\"format\":{ (logFile.Type == LogType.Json ? "\"JSON\"" : "\"TEXT\"")},"));
 
-//				Stream.Write(Encoding.UTF8.GetBytes($"\"minIndex\":{MinIndex},"));
-//				Stream.Write(Encoding.UTF8.GetBytes($"\"minOffset\":{MinOffset},"));
-//				Stream.Write(Encoding.UTF8.GetBytes($"\"maxIndex\":{MaxIndex},"));
-//				Stream.Write(Encoding.UTF8.GetBytes($"\"maxOffset\":{MaxOffset},"));
-//				Stream.Write(Encoding.UTF8.GetBytes($"\"length\":{Result.Length},"));
-
 				stream.Write(Encoding.UTF8.GetBytes($"\"lines\":["));
 				stream.WriteByte((byte)'\n');
 
-				int offset = 0;
-				for (int line = index; line < index + count; line++)
+				for (int lineIdx = 0; lineIdx < lines.Count; lineIdx++)
 				{
+					Utf8String line = lines[lineIdx];
+
 					stream.WriteByte((byte)' ');
 					stream.WriteByte((byte)' ');
 
 					if (logFile.Type == LogType.Json)
 					{
-						// Find the end of the line and output it as an opaque blob
-						int startOffset = offset;
-						for (; ; offset++)
-						{
-							if (offset == result.Length)
-							{
-								stream.WriteByte((byte)'{');
-								stream.WriteByte((byte)'}');
-								break;
-							}
-							else if (result[offset] == (byte)'\n')
-							{
-								await stream.WriteAsync(result.AsMemory(startOffset, offset - startOffset), cancellationToken);
-								offset++;
-								break;
-							}
-						}
+						await stream.WriteAsync(line.Memory, cancellationToken);
 					}
 					else
 					{
 						stream.WriteByte((byte)'\"');
-						for (; offset < result.Length; offset++)
+						for (int idx = 0; idx < line.Length; idx++)
 						{
-							if (result[offset] == '\\' || result[offset] == '\"')
-							{
-								stream.WriteByte((byte)'\\');
-								stream.WriteByte(result[offset]);
-							}
-							else if (result[offset] == (byte)'\n')
-							{
-								offset++;
-								break;
-							}
-							else if (result[offset] >= 32 && result[offset] <= 126)
+							byte character = line[idx];
+							if (character >= 32 && character <= 126 && character != '\\' && character != '\"')
 							{
-								stream.WriteByte(result[offset]);
+								stream.WriteByte(character);
 							}
 							else
 							{
-								stream.Write(Encoding.UTF8.GetBytes($"\\x{result[offset]:x2}"));
+								stream.Write(Encoding.UTF8.GetBytes($"\\x{character:x2}"));
 							}
 						}
 						stream.WriteByte((byte)'\"');
 					}
 
-					if (line + 1 < index + count)
+					if (lineIdx + 1 < lines.Count)
 					{
 						stream.WriteByte((byte)',');
 					}
