Horde: Support for log tailing with new storage backend.

* Once a log is retrieved, server adds an entry with this log id to a sorted set in Redis, scored by expiry time and broadcasts it to any other server instances.
* If a log is being tailed, server keeps the total number of lines in a Redis key and appends chunks of log data to a sorted set for that log scored by index of the first line.
* Chunks are split on fixed boundaries, in order to allow older entries to be purged by score without having to count the number of lines they contain first.
* Agent polls for requests to provide tail data via LogRpcService.UpdateLogTail, which calls LogTailService.WaitForTailNext. WaitForTailNext returns the index of the total line count for this log if it is being tailed (indicating the index of the next line that the agent should send to the server), or blocks until the log is being tailed.
* Once data is flushed to persistent storage, the number of flushed lines is added to LogRpcService._trimQueue and the line data is removed from Redis after LogRpcService.TrimAfter.
* Log node contains a "complete" flag indicating whether it is necessary to check for tail data.

diff --git a/Engine/Source/Programs/Horde/Horde.Agent/Utility/JsonRpcLogger.cs b/Engine/Source/Programs/Horde/Horde.Agent/Utility/JsonRpcLogger.cs
--- a/Engine/Source/Programs/Horde/Horde.Agent/Utility/JsonRpcLogger.cs
+++ b/Engine/Source/Programs/Horde/Horde.Agent/Utility/JsonRpcLogger.cs
@@ -106,7 +106,7 @@ namespace Horde.Agent.Parser
 		// Tailing task
 		readonly Task _tailTask;
 		AsyncEvent _tailTaskStop;
-		AsyncEvent _newTailDataEvent = new AsyncEvent();
+		readonly AsyncEvent _newTailDataEvent = new AsyncEvent();
 
 		public JsonRpcAndStorageLogSink(IRpcConnection connection, string logId, IJsonRpcLogSink? inner, IStorageClient store, ILogger logger)
 		{
@@ -114,7 +114,7 @@ namespace Horde.Agent.Parser
 			_logId = logId;
 			_builder = new LogBuilder(LogFormat.Json, logger);
 			_inner = inner;
-			_writer = new TreeWriter(store, prefix: $"logs/{logId}");
+			_writer = new TreeWriter(store);
 			_logger = logger;
 
 			_tailTaskStop = new AsyncEvent();
@@ -159,12 +159,10 @@ namespace Horde.Agent.Parser
 				{
 					if (tailNext != -1)
 					{
-						tailNext = Math.Max(tailNext, _builder.LineCount);
+						tailNext = Math.Max(tailNext, _builder.FlushedLineCount);
 						tailData = _builder.ReadTailData(tailNext, 16 * 1024);
 					}
-
 					newTailDataTask = _newTailDataEvent.Task;
-					_newTailDataEvent = new AsyncEvent();
 				}
 
 				// If we don't have any updates for the server, wait until we do.
@@ -215,7 +213,7 @@ namespace Horde.Agent.Parser
 
 			if (request.Flush || _bufferLength > FlushLength)
 			{
-				NodeLocator locator = await _builder.FlushAsync(_writer, cancellationToken);
+				NodeLocator locator = await _builder.FlushAsync(_writer, request.Flush, cancellationToken);
 				await UpdateLogAsync(locator, _builder.LineCount, cancellationToken);
 				_bufferLength = 0;
 			}
@@ -245,6 +243,7 @@ namespace Horde.Agent.Parser
 				{
 					// Write the request to the server
 					UpdateLogTailRequest request = new UpdateLogTailRequest();
+					request.LogId = _logId;
 					request.TailNext = tailNext;
 					request.TailData = UnsafeByteOperations.UnsafeWrap(tailData);
 					await call.RequestStream.WriteAsync(request);
diff --git a/Engine/Source/Programs/Horde/Horde.Build.Tests/LogFileServiceTests.cs b/Engine/Source/Programs/Horde/Horde.Build.Tests/LogFileServiceTests.cs
--- a/Engine/Source/Programs/Horde/Horde.Build.Tests/LogFileServiceTests.cs
+++ b/Engine/Source/Programs/Horde/Horde.Build.Tests/LogFileServiceTests.cs
@@ -48,7 +48,7 @@ namespace Horde.Build.Tests
 			_logStorage = new PersistentLogStorage(new MemoryStorageBackend().ForType<PersistentLogStorage>(), NullLogger<PersistentLogStorage>.Instance);
 			_clock = new FakeClock();
 			TestOptions<ServerSettings> settingsOpts = new (new ServerSettings());
-			_logFileService = new LogFileService(logFileCollection, null!, logBuilder, _logStorage, _clock, null!, settingsOpts, logger);
+			_logFileService = new LogFileService(logFileCollection, null!, logBuilder, _logStorage, _clock, null!, null!, settingsOpts, logger);
         }
 
 		protected override void Dispose(bool disposing)
@@ -92,7 +92,7 @@ namespace Horde.Build.Tests
             Assert.AreEqual(6 + 8 + 4, metadata.Length);
             Assert.AreEqual(4, metadata.MaxLineIndex);
 
-			List<Utf8String> lines = await _logFileService.ReadLinesAsync(logFile, 0, null, CancellationToken.None);
+			List<Utf8String> lines = await _logFileService.ReadLinesAsync(logFile, 0, 100, CancellationToken.None);
 			Assert.AreEqual("hello", lines[0].ToString());
 			Assert.AreEqual("foo", lines[1].ToString());
 			Assert.AreEqual("bar", lines[2].ToString());
diff --git a/Engine/Source/Programs/Horde/Horde.Build.Tests/LogIndexingTests.cs b/Engine/Source/Programs/Horde/Horde.Build.Tests/LogIndexingTests.cs
--- a/Engine/Source/Programs/Horde/Horde.Build.Tests/LogIndexingTests.cs
+++ b/Engine/Source/Programs/Horde/Horde.Build.Tests/LogIndexingTests.cs
@@ -50,7 +50,7 @@ namespace Horde.Build.Tests
 			_nullLogStorage = new NullLogStorage();
 			_logStorage = new LocalLogStorage(20, _nullLogStorage);
 			TestOptions<ServerSettings> settingsOpts = new (new ServerSettings());
-			_logFileService = new LogFileService(logFileCollection, null!, logBuilder, _logStorage, new FakeClock(), null!, settingsOpts, logger);
+			_logFileService = new LogFileService(logFileCollection, null!, logBuilder, _logStorage, new FakeClock(), null!, null!, settingsOpts, logger);
 		}
 
 		protected override void Dispose(bool disposing)
diff --git a/Engine/Source/Programs/Horde/Horde.Build.Tests/LogTailServiceTests.cs b/Engine/Source/Programs/Horde/Horde.Build.Tests/LogTailServiceTests.cs
new file mode 100644
--- /dev/null
+++ b/Engine/Source/Programs/Horde/Horde.Build.Tests/LogTailServiceTests.cs
@@ -0,0 +1,190 @@
+// Copyright Epic Games, Inc. All Rights Reserved.
+
+using System;
+using System.Collections.Generic;
+using System.Text;
+using System.Threading;
+using System.Threading.Tasks;
+using EpicGames.Core;
+using Horde.Build.Logs;
+using Horde.Build.Server;
+using Horde.Build.Utilities;
+using HordeCommon;
+using Microsoft.Extensions.DependencyInjection;
+using Microsoft.Extensions.Logging;
+using Microsoft.VisualStudio.TestTools.UnitTesting;
+
+namespace Horde.Build.Tests
+{
+	using LogId = ObjectId<ILogFile>;
+
+	[TestClass]
+	public class LogTailServiceTests : TestSetup
+	{
+		readonly LogId _logId = LogId.Parse("637809e571f04e1d1ba311f1");
+
+		[TestMethod]
+		public async Task Truncated()
+		{
+			LogTailService tailService = ServiceProvider.GetRequiredService<LogTailService>();
+			await tailService.EnableTailingAsync(_logId, 30);
+			await tailService.AppendAsync(_logId, 30, Encoding.UTF8.GetBytes("foo\nbar\nbaz\n"));
+
+			for(int idx = 25; idx <= 29; idx++)
+			{
+				List<Utf8String> lines = await tailService.ReadAsync(_logId, idx, 10);
+				Assert.AreEqual(0, lines.Count);
+			}
+
+			{
+				List<Utf8String> lines = await tailService.ReadAsync(_logId, 30, 10);
+				Assert.AreEqual(3, lines.Count);
+				Assert.AreEqual("foo", lines[0]);
+				Assert.AreEqual("bar", lines[1]);
+				Assert.AreEqual("baz", lines[2]);
+			}
+		}
+
+		[TestMethod]
+		public async Task TruncatedMultiLine()
+		{
+			LogTailService tailService = ServiceProvider.GetRequiredService<LogTailService>();
+			await tailService.EnableTailingAsync(_logId, 30);
+			await tailService.AppendAsync(_logId, 30, Encoding.UTF8.GetBytes("foo\n"));
+			await tailService.AppendAsync(_logId, 31, Encoding.UTF8.GetBytes("bar\n"));
+			await tailService.AppendAsync(_logId, 32, Encoding.UTF8.GetBytes("baz\n"));
+
+			for (int idx = 25; idx <= 29; idx++)
+			{
+				List<Utf8String> lines = await tailService.ReadAsync(_logId, idx, 10);
+				Assert.AreEqual(0, lines.Count);
+			}
+
+			{
+				List<Utf8String> lines = await tailService.ReadAsync(_logId, 30, 10);
+				Assert.AreEqual(3, lines.Count);
+				Assert.AreEqual("foo", lines[0]);
+				Assert.AreEqual("bar", lines[1]);
+				Assert.AreEqual("baz", lines[2]);
+			}
+		}
+
+		[TestMethod]
+		public async Task Splitting()
+		{
+			RedisService redisService = ServiceProvider.GetRequiredService<RedisService>();
+
+			LogTailService tailService = new LogTailService(redisService, ServiceProvider.GetRequiredService<IClock>(), 4, ServiceProvider.GetRequiredService<ILogger<LogTailService>>());
+			await tailService.StartAsync(CancellationToken.None);
+
+			await tailService.EnableTailingAsync(_logId, 30);
+			await tailService.AppendAsync(_logId, 30, Encoding.UTF8.GetBytes("line30\nline31\nline32\nline33\nline34\nline35\nline36\nline37\nline38\nline39\nline40\n"));
+
+			Assert.AreEqual(4, await tailService.RedisTailData(_logId).LengthAsync());
+
+			await tailService.FlushAsync(_logId, 33);
+
+			Assert.AreEqual(4, await tailService.RedisTailData(_logId).LengthAsync());
+
+			await Clock.AdvanceAsync(TimeSpan.FromMinutes(1.0));
+
+			Assert.AreEqual(3, await tailService.RedisTailData(_logId).LengthAsync());
+
+			{
+				List<Utf8String> lines = await tailService.ReadAsync(_logId, 25, 20);
+				Assert.AreEqual(0, lines.Count);
+			}
+			{
+				List<Utf8String> lines = await tailService.ReadAsync(_logId, 31, 20);
+				Assert.AreEqual(0, lines.Count);
+			}
+			{
+				List<Utf8String> lines = await tailService.ReadAsync(_logId, 32, 20);
+				Assert.AreEqual("line32", lines[0]);
+				Assert.AreEqual("line33", lines[1]);
+				Assert.AreEqual("line34", lines[2]);
+				Assert.AreEqual("line35", lines[3]);
+				Assert.AreEqual("line36", lines[4]);
+				Assert.AreEqual("line37", lines[5]);
+				Assert.AreEqual("line38", lines[6]);
+				Assert.AreEqual("line39", lines[7]);
+				Assert.AreEqual("line40", lines[8]);
+			}
+		}
+
+		[TestMethod]
+		public async Task Expiry()
+		{
+			RedisService redisService = GetRedisServiceSingleton();
+
+			LogTailService tailService = ServiceProvider.GetRequiredService<LogTailService>();
+
+			await tailService.StartAsync(CancellationToken.None);
+			await tailService.EnableTailingAsync(_logId, 30);
+			await tailService.AppendAsync(_logId, 30, Encoding.UTF8.GetBytes("foo\n"));
+
+			Assert.IsTrue(await redisService.GetDatabase().KeyExistsAsync(tailService.RedisTailNext(_logId).Key));
+			Assert.IsTrue(await redisService.GetDatabase().KeyExistsAsync(tailService.RedisTailData(_logId).Key));
+
+			await Clock.AdvanceAsync(TimeSpan.FromMinutes(5.0));
+
+			Assert.IsFalse(await redisService.GetDatabase().KeyExistsAsync(tailService.RedisTailNext(_logId).Key));
+			Assert.IsFalse(await redisService.GetDatabase().KeyExistsAsync(tailService.RedisTailData(_logId).Key));
+		}
+
+		[TestMethod]
+		public async Task MissingData()
+		{
+			LogTailService tailService = ServiceProvider.GetRequiredService<LogTailService>();
+
+			await tailService.StartAsync(CancellationToken.None);
+			await tailService.EnableTailingAsync(_logId, 30);
+			await tailService.AppendAsync(_logId, 31, Encoding.UTF8.GetBytes("line31\nline32\nline33\nline34\nline35\n"));
+
+			{
+				List<Utf8String> lines = await tailService.ReadAsync(_logId, 29, 4);
+				Assert.AreEqual(0, lines.Count);
+			}
+			{
+				List<Utf8String> lines = await tailService.ReadAsync(_logId, 29, 4);
+				Assert.AreEqual(0, lines.Count);
+			}
+
+			await tailService.AppendAsync(_logId, 30, Encoding.UTF8.GetBytes("line30\n"));
+			await tailService.AppendAsync(_logId, 32, Encoding.UTF8.GetBytes("line32\nline33\nline34\nline35\n"));
+
+			{
+				List<Utf8String> lines = await tailService.ReadAsync(_logId, 29, 4);
+				Assert.AreEqual(0, lines.Count);
+			}
+			{
+				List<Utf8String> lines = await tailService.ReadAsync(_logId, 30, 4);
+				Assert.AreEqual(1, lines.Count);
+				Assert.AreEqual("line30", lines[0]);
+			}
+			{
+				List<Utf8String> lines = await tailService.ReadAsync(_logId, 32, 4);
+				Assert.AreEqual(0, lines.Count);
+			}
+		}
+
+		[TestMethod]
+		public async Task AsyncEnable()
+		{
+			LogTailService tailService = ServiceProvider.GetRequiredService<LogTailService>();
+			await tailService.StartAsync(CancellationToken.None);
+
+			LogTailService tailService2 = new LogTailService(ServiceProvider.GetRequiredService<RedisService>(), ServiceProvider.GetRequiredService<IClock>(), 4, ServiceProvider.GetRequiredService<ILogger<LogTailService>>());
+
+			Task<int> task = tailService.WaitForTailNext(_logId, CancellationToken.None);
+			Assert.IsTrue(!task.IsCompleted);
+
+			await tailService2.EnableTailingAsync(_logId, 30);
+
+			await Task.WhenAny(task, Task.Delay(1000));
+			Assert.IsTrue(task.IsCompletedSuccessfully);
+
+			Assert.AreEqual(30, await task);
+		}
+	}
+}
diff --git a/Engine/Source/Programs/Horde/Horde.Build.Tests/TestSetup.cs b/Engine/Source/Programs/Horde/Horde.Build.Tests/TestSetup.cs
--- a/Engine/Source/Programs/Horde/Horde.Build.Tests/TestSetup.cs
+++ b/Engine/Source/Programs/Horde/Horde.Build.Tests/TestSetup.cs
@@ -231,6 +231,7 @@ namespace Horde.Build.Tests
 			services.AddSingleton<ILogStorage, NullLogStorage>();
 			services.AddSingleton<ILogBuilder, LocalLogBuilder>();
 			services.AddSingleton<ILogFileService, LogFileService>();
+			services.AddSingleton<LogTailService>();
 			services.AddSingleton<INotificationService, NotificationService>();
 			services.AddSingleton<IPerforceService, PerforceServiceStub>();
 			services.AddSingleton<PerforceLoadBalancer>();
diff --git a/Engine/Source/Programs/Horde/Horde.Build/Logs/LogFileMessages.cs b/Engine/Source/Programs/Horde/Horde.Build/Logs/LogFileMessages.cs
--- a/Engine/Source/Programs/Horde/Horde.Build/Logs/LogFileMessages.cs
+++ b/Engine/Source/Programs/Horde/Horde.Build/Logs/LogFileMessages.cs
@@ -1,5 +1,6 @@
 // Copyright Epic Games, Inc. All Rights Reserved.
 
+using System;
 using System.Collections.Generic;
 using Horde.Build.Logs.Data;
 
@@ -21,6 +22,28 @@ namespace Horde.Build.Logs
 		Json
 	}
 
+	/// <summary>
+	/// Creates a new log file
+	/// </summary>
+	public class CreateLogFileRequest
+	{
+		/// <summary>
+		/// Type of the log file
+		/// </summary>
+		public LogType Type { get; set; } = LogType.Json;
+	}
+
+	/// <summary>
+	/// Response from creating a log file
+	/// </summary>
+	public class CreateLogFileResponse
+	{
+		/// <summary>
+		/// Identifier for the created log file
+		/// </summary>
+		public string Id { get; set; } = String.Empty;
+	}
+
 	/// <summary>
 	/// Response describing a log file
 	/// </summary>
diff --git a/Engine/Source/Programs/Horde/Horde.Build/Logs/LogFileService.cs b/Engine/Source/Programs/Horde/Horde.Build/Logs/LogFileService.cs
--- a/Engine/Source/Programs/Horde/Horde.Build/Logs/LogFileService.cs
+++ b/Engine/Source/Programs/Horde/Horde.Build/Logs/LogFileService.cs
@@ -101,7 +101,7 @@ namespace Horde.Build.Logs
 		/// <param name="count">Maximum number of lines to return</param>
 		/// <param name="cancellationToken">Cancellation token for the operation</param>
 		/// <returns>List of lines</returns>
-		Task<List<Utf8String>> ReadLinesAsync(ILogFile logFile, int index, int? count, CancellationToken cancellationToken = default);
+		Task<List<Utf8String>> ReadLinesAsync(ILogFile logFile, int index, int count, CancellationToken cancellationToken = default);
 
 		/// <summary>
 		/// Writes out chunk data and assigns to a file
@@ -555,12 +555,13 @@ namespace Horde.Build.Logs
 			public override void Write(byte[] buffer, int offset, int count) => throw new NotImplementedException();
 		}
 
+		readonly LogTailService _logTailService;
 		readonly ITicker _ticker;
 
 		/// <summary>
 		/// Constructor
 		/// </summary>
-		public LogFileService(ILogFileCollection logFiles, ILogEventCollection logEvents, ILogBuilder builder, ILogStorage storage, IClock clock, StorageService storageService, IOptions<ServerSettings> settings, ILogger<LogFileService> logger)
+		public LogFileService(ILogFileCollection logFiles, ILogEventCollection logEvents, ILogBuilder builder, ILogStorage storage, IClock clock, LogTailService logTailService, StorageService storageService, IOptions<ServerSettings> settings, ILogger<LogFileService> logger)
 		{
 			_logFiles = logFiles;
 			_logEvents = logEvents;
@@ -568,6 +569,7 @@ namespace Horde.Build.Logs
 			_builder = builder;
 			_storage = storage;
 			_ticker = clock.AddSharedTicker<LogFileService>(TimeSpan.FromSeconds(30.0), TickAsync, logger);
+			_logTailService = logTailService;
 			_storageService = storageService;
 			_settings = settings;
 			_logger = logger;
@@ -642,27 +644,66 @@ namespace Horde.Build.Logs
 		}
 
 		/// <inheritdoc/>
-		public async Task<List<Utf8String>> ReadLinesAsync(ILogFile logFile, int index, int? count, CancellationToken cancellationToken)
+		public async Task<List<Utf8String>> ReadLinesAsync(ILogFile logFile, int index, int count, CancellationToken cancellationToken)
 		{
-			(_, long minOffset) = await GetLineOffsetAsync(logFile, index, cancellationToken);
-			(_, long maxOffset) = await GetLineOffsetAsync(logFile, (count == null) ? Int32.MaxValue : (index + Math.Min(count.Value, Int32.MaxValue - index)), cancellationToken);
+			List<Utf8String> lines = new List<Utf8String>();
 
-			byte[] result;
-			using (System.IO.Stream stream = await OpenRawStreamAsync(logFile, minOffset, maxOffset - minOffset, cancellationToken))
+			if (_settings.Value.FeatureFlags.EnableNewLogger)
 			{
-				result = new byte[stream.Length];
-				await stream.ReadFixedSizeDataAsync(result, 0, result.Length);
-			}
+				TreeReader reader = await GetTreeReaderAsync(cancellationToken);
 
-			List<Utf8String> lines = new List<Utf8String>();
+				int maxIndex = index + count;
 
-			int offset = 0;
-			for (int idx = 0; idx < result.Length; idx++)
+				LogNode? root = await reader.TryReadNodeAsync<LogNode>(logFile.RefName, cancellationToken: cancellationToken);
+				if (root != null)
+				{
+					int chunkIdx = root.TextChunkRefs.GetChunkForLine(index);
+					for (; index < maxIndex && chunkIdx < root.TextChunkRefs.Count; chunkIdx++)
+					{
+						LogChunkRef chunk = root.TextChunkRefs[chunkIdx];
+						LogChunkNode chunkData = await chunk.ExpandAsync(reader, cancellationToken);
+
+						for (; index < maxIndex && index < chunk.LineIndex; index++)
+						{
+							lines.Add($"Internal error; missing data for line {index}\n");
+						}
+
+						for (; index < maxIndex && index < chunk.LineIndex + chunk.LineCount; index++)
+						{
+							lines.Add(chunkData.GetLine(index - chunk.LineIndex));
+						}
+					}
+				}
+
+				if (root == null || !root.Complete)
+				{
+					await _logTailService.EnableTailingAsync(logFile.Id, root?.LineCount ?? 0);
+					if (index < maxIndex)
+					{
+						await _logTailService.ReadAsync(logFile.Id, index, maxIndex - index, lines);
+					}
+				}
+			}
+			else
 			{
-				if (result[idx] == (byte)'\n')
+				(_, long minOffset) = await GetLineOffsetAsync(logFile, index, cancellationToken);
+				(_, long maxOffset) = await GetLineOffsetAsync(logFile, index + Math.Min(count, Int32.MaxValue - index), cancellationToken);
+
+				byte[] result;
+				using (System.IO.Stream stream = await OpenRawStreamAsync(logFile, minOffset, maxOffset - minOffset, cancellationToken))
 				{
-					lines.Add(new Utf8String(result.AsMemory(offset, idx - offset)));
-					offset = idx + 1;
+					result = new byte[stream.Length];
+					await stream.ReadFixedSizeDataAsync(result, 0, result.Length);
+				}
+
+				int offset = 0;
+				for (int idx = 0; idx < result.Length; idx++)
+				{
+					if (result[idx] == (byte)'\n')
+					{
+						lines.Add(new Utf8String(result.AsMemory(offset, idx - offset)));
+						offset = idx + 1;
+					}
 				}
 			}
 
diff --git a/Engine/Source/Programs/Horde/Horde.Build/Logs/LogRpcService.cs b/Engine/Source/Programs/Horde/Horde.Build/Logs/LogRpcService.cs
--- a/Engine/Source/Programs/Horde/Horde.Build/Logs/LogRpcService.cs
+++ b/Engine/Source/Programs/Horde/Horde.Build/Logs/LogRpcService.cs
@@ -8,6 +8,7 @@ using Horde.Build.Utilities;
 using EpicGames.Horde.Storage;
 using Horde.Build.Storage;
 using Microsoft.Extensions.Logging;
+using System.Threading;
 
 namespace Horde.Build.Logs
 {
@@ -20,15 +21,17 @@ namespace Horde.Build.Logs
 	public class LogRpcService : LogRpc.LogRpcBase
 	{
 		readonly ILogFileService _logFileService;
+		readonly LogTailService _logTailService;
 		readonly StorageService _storageService;
 		readonly ILogger<LogRpcService> _logger;
 
 		/// <summary>
 		/// Constructor
 		/// </summary>
-		public LogRpcService(ILogFileService logFileService, StorageService storageService, ILogger<LogRpcService> logger)
+		public LogRpcService(ILogFileService logFileService, LogTailService logTailService, StorageService storageService, ILogger<LogRpcService> logger)
 		{
 			_logFileService = logFileService;
+			_logTailService = logTailService;
 			_storageService = storageService;
 			_logger = logger;
 		}
@@ -49,17 +52,45 @@ namespace Horde.Build.Logs
 			IStorageClient store = await _storageService.GetClientAsync(Namespace.Logs, context.CancellationToken);
 			_logger.LogInformation("Updating {LogId} to node {Locator}", request.LogId, request.BlobLocator);
 			await store.WriteRefTargetAsync(new RefName(request.LogId), NodeLocator.Parse(request.BlobLocator));
+
+			await _logTailService.FlushAsync(logFile.Id, request.LineCount);
+
 			return new UpdateLogResponse();
 		}
 
 		/// <inheritdoc/>
 		public override async Task UpdateLogTail(IAsyncStreamReader<UpdateLogTailRequest> requestStream, IServerStreamWriter<UpdateLogTailResponse> responseStream, ServerCallContext context)
 		{
-			while (await requestStream.MoveNext())
+			UpdateLogTailResponse response = new UpdateLogTailResponse();
+			response.TailNext = -1;
+
+			Task<bool> moveNextTask = requestStream.MoveNext();
+			while (await moveNextTask)
 			{
-				// TODO: request tail data if log has been read recently
+				UpdateLogTailRequest request = requestStream.Current;
+				LogId logId = new LogId(request.LogId);
+
+				if (request.TailData.Length > 0)
+				{
+					await _logTailService.AppendAsync(logId, request.TailNext, request.TailData.Memory);
+				}
+
+				moveNextTask = requestStream.MoveNext();
+
+				using (CancellationTokenSource cancellationSource = new CancellationTokenSource())
+				{
+					Task<int> waitTask = _logTailService.WaitForTailNext(logId, cancellationSource.Token);
+
+					await Task.WhenAny(waitTask, moveNextTask);
+					cancellationSource.Cancel();
+
+					response.TailNext = await waitTask;
+				}
+
+				await responseStream.WriteAsync(response);
 			}
-			await responseStream.WriteAsync(new UpdateLogTailResponse { TailNext = -1 }); 
+
+			await responseStream.WriteAsync(response);
 		}
 	}
 }
diff --git a/Engine/Source/Programs/Horde/Horde.Build/Logs/LogTailService.cs b/Engine/Source/Programs/Horde/Horde.Build/Logs/LogTailService.cs
new file mode 100644
--- /dev/null
+++ b/Engine/Source/Programs/Horde/Horde.Build/Logs/LogTailService.cs
@@ -0,0 +1,366 @@
+// Copyright Epic Games, Inc. All Rights Reserved.
+
+using System.Threading.Tasks;
+using Horde.Build.Utilities;
+using System.Collections.Generic;
+using System.Threading;
+using System;
+using EpicGames.Core;
+using Horde.Build.Server;
+using EpicGames.Redis;
+using StackExchange.Redis;
+using System.Collections.Concurrent;
+using Microsoft.Extensions.Hosting;
+using HordeCommon;
+using Microsoft.Extensions.Logging;
+using System.Globalization;
+
+namespace Horde.Build.Logs
+{
+	using LogId = ObjectId<ILogFile>;
+
+	/// <summary>
+	/// Stores data for logs which has not yet been flushed to persistent storage.
+	/// 
+	/// * Once a log is retrieved, server adds an entry with this log id to a sorted set in Redis (see <see cref="_expireQueue"/>, scored by 
+	///   expiry time (defaulting to <see cref="ExpireAfter"/> from present) and broadcasts it to any other server instances.
+	/// * If a log is being tailed, server keeps the total number of lines in a Redis key (<see cref="RedisTailNext"/>, and appends chunks of log 
+	///   data to a sorted set for that log scored by index of the first line (see <see cref="RedisTailData"/>).
+	/// * Chunks are split on fixed boundaries, in order to allow older entries to be purged by score without having to count the number of lines
+	///   they contain first.
+	/// * Agent polls for requests to provide tail data via <see cref="LogRpcService.UpdateLogTail"/>, which calls <see cref="WaitForTailNext"/>.
+	///   <see cref="WaitForTailNext"/> returns the index of the total line count for this log if it is being tailed (indicating the index of
+	///   the next line that the agent should send to the server), or blocks until the log is being tailed.
+	/// * Once data is flushed to persistent storage, the number of flushed lines is added to <see cref="_trimQueue"/> and the line data is removed
+	///   from Redis after <see cref="TrimAfter"/>.
+	/// * <see cref="EpicGames.Horde.Logs.LogNode"/> contains a "complete" flag indicating whether it is necessary to check for tail data.
+	///   
+	/// </summary>
+	public class LogTailService : IHostedService
+	{
+		/// <summary>
+		/// Amount of time to keep tail data in the cache after it has been flushed to persistent storage. This gives some wiggle room with order of operations when sequencing persisted data with tail data.
+		/// </summary>
+		public static TimeSpan TrimAfter { get; } = TimeSpan.FromMinutes(0.5);
+
+		/// <summary>
+		/// Amount of time to continue tailing a log after the last time it was updated.
+		/// </summary>
+		public static TimeSpan ExpireAfter { get; } = TimeSpan.FromMinutes(2.0);
+
+		readonly RedisService _redisService;
+
+		readonly ConcurrentDictionary<LogId, AsyncEvent> _notifyLogEvents = new ConcurrentDictionary<LogId, AsyncEvent>();
+		readonly RedisChannel<LogId> _tailStartChannel;
+		readonly RedisSortedSet<string> _trimQueue;
+		readonly RedisSortedSet<LogId> _expireQueue;
+
+		readonly IClock _clock;
+		readonly ITicker _expireTailsTicker;
+
+		readonly ILogger _logger;
+
+		RedisChannelSubscription<LogId>? _tailStartSubscription;
+
+		internal RedisString<int> RedisTailNext(LogId logId) => new RedisString<int>(_redisService.ConnectionPool, $"logs:{logId}:tail-next");
+		internal RedisSortedSet<ReadOnlyMemory<byte>> RedisTailData(LogId logId) => new RedisSortedSet<ReadOnlyMemory<byte>>(_redisService.ConnectionPool, $"logs:{logId}:tail-data");
+
+		/// <summary>
+		/// Boundary to split Redis chunks along. 
+		/// </summary>
+		public int ChunkLineCount { get; }
+
+		/// <summary>
+		/// Constructor
+		/// </summary>
+		public LogTailService(RedisService redisService, IClock clock, ILogger<LogTailService> logger)
+			: this(redisService, clock, 4, logger)
+		{
+		}
+
+		/// <summary>
+		/// Constructor
+		/// </summary>
+		public LogTailService(RedisService redisService, IClock clock, int chunkLineCount, ILogger<LogTailService> logger)
+		{
+			if ((chunkLineCount & (chunkLineCount - 1)) != 0)
+			{
+				throw new ArgumentException($"{nameof(ChunkLineCount)} must be a power of two", nameof(chunkLineCount));
+			}
+
+			ChunkLineCount = chunkLineCount;
+
+			_redisService = redisService;
+			_tailStartChannel = new RedisChannel<LogId>("logs:notify");
+			_trimQueue = new RedisSortedSet<string>(_redisService.ConnectionPool, "logs:trim");
+			_expireQueue = new RedisSortedSet<LogId>(_redisService.ConnectionPool, "logs:expire");
+
+			_clock = clock;
+			_expireTailsTicker = clock.AddSharedTicker<LogTailService>(TimeSpan.FromSeconds(30.0), ExpireTailsAsync, logger);
+			_logger = logger;
+		}
+
+		/// <inheritdoc/>
+		public async Task StartAsync(CancellationToken cancellationToken)
+		{
+			_tailStartSubscription = await _redisService.DatabaseSingleton.Multiplexer.GetSubscriber().SubscribeAsync(_tailStartChannel, (_, x) => OnTailStart(x));
+			await _expireTailsTicker.StartAsync();
+		}
+
+		/// <inheritdoc/>
+		public async Task StopAsync(CancellationToken cancellationToken)
+		{
+			await _expireTailsTicker.StopAsync();
+			if (_tailStartSubscription != null)
+			{
+				await _tailStartSubscription.DisposeAsync();
+			}
+		}
+
+		private async ValueTask ExpireTailsAsync(CancellationToken cancellationToken)
+		{
+			double maxScore = (_clock.UtcNow - DateTime.UnixEpoch).TotalSeconds;
+
+			SortedSetEntry<string>[] trimEntries = await _trimQueue.RangeByScoreWithScoresAsync(stop: maxScore);
+			foreach (SortedSetEntry<string> trimEntry in trimEntries)
+			{
+				string[] values = trimEntry.Element.Split('=');
+
+				LogId logId = LogId.Parse(values[0]);
+				int lineCount = Int32.Parse(values[1], NumberStyles.None, CultureInfo.InvariantCulture);
+
+				int minLineIndex = (lineCount & ~(ChunkLineCount - 1)) - 1;
+				_ = await RedisTailData(logId).RemoveRangeByScoreAsync(Double.NegativeInfinity, minLineIndex, Exclude.None, CommandFlags.FireAndForget);
+			}
+
+			SortedSetEntry<LogId>[] expireEntries = await _expireQueue.RangeByScoreWithScoresAsync(stop: maxScore);
+			foreach (SortedSetEntry<LogId> expireEntry in expireEntries)
+			{
+				LogId logId = expireEntry.Element;
+
+				ITransaction transaction = _redisService.GetDatabase().CreateTransaction();
+				transaction.AddCondition(Condition.SortedSetEqual(_expireQueue.Key, expireEntry.ElementValue, expireEntry.Score));
+				_ = transaction.KeyDeleteAsync(RedisTailNext(logId).Key, CommandFlags.FireAndForget);
+				_ = transaction.KeyDeleteAsync(RedisTailData(logId).Key, CommandFlags.FireAndForget);
+				_ = transaction.ExecuteAsync(CommandFlags.FireAndForget);
+			}
+		}
+
+		private void OnTailStart(LogId logId)
+		{
+			if (_notifyLogEvents.TryGetValue(logId, out AsyncEvent? notifyEvent))
+			{
+				notifyEvent.Latch();
+			}
+		}
+
+		/// <summary>
+		/// Enable tailing for the given log file
+		/// </summary>
+		/// <param name="logId">Log file identifier</param>
+		/// <param name="lineCount">The current flushed line count for the log file</param>
+		public async ValueTask EnableTailingAsync(LogId logId, int lineCount)
+		{
+			if (await RedisTailNext(logId).SetAsync(lineCount, when: When.NotExists))
+			{
+				_logger.LogDebug("Enabled tailing for log {LogId}", logId);
+			}
+
+			await _redisService.GetDatabase().PublishAsync(_tailStartChannel, logId);
+
+			double score = (_clock.UtcNow + ExpireAfter - DateTime.UnixEpoch).TotalSeconds;
+			_ = _expireQueue.AddAsync(logId, score, flags: CommandFlags.FireAndForget);
+		}
+
+		/// <summary>
+		/// Appends data to the log tail, starting at the given line index
+		/// </summary>
+		/// <param name="logId">Log file identifier</param>
+		/// <param name="tailNext">Index of the first line to append</param>
+		/// <param name="tailData">Data for the log starting at the given line number</param>
+		/// <returns></returns>
+		public async Task AppendAsync(LogId logId, int tailNext, ReadOnlyMemory<byte> tailData)
+		{
+			List<SortedSetEntry<ReadOnlyMemory<byte>>> entries = new List<SortedSetEntry<ReadOnlyMemory<byte>>>();
+			int newLineCount = SplitLogDataToChunks(tailNext, tailData, ChunkLineCount, entries);
+
+			if (entries.Count > 0)
+			{
+				RedisString<int> redisTailNext = RedisTailNext(logId);
+				RedisSortedSet<ReadOnlyMemory<byte>> redisTailData = RedisTailData(logId);
+
+				ITransaction transaction = _redisService.GetDatabase().CreateTransaction();
+				transaction.AddCondition(Condition.StringEqual(RedisTailNext(logId).Key, tailNext));
+				_ = transaction.With(redisTailData).AddAsync(entries.ToArray(), flags: CommandFlags.FireAndForget);
+				_ = transaction.With(redisTailNext).SetAsync(newLineCount, flags: CommandFlags.FireAndForget);
+
+				bool result = await transaction.ExecuteAsync();
+				_logger.LogTrace("Tail data for {LogId}, line {TailNext} -> {NewTailNext}, result: {Result}", logId, tailNext, newLineCount, result);
+			}
+		}
+
+		/// <summary>
+		/// Reads data for the given log file, starting at the requested line index
+		/// </summary>
+		/// <param name="logId">The log to query</param>
+		/// <param name="lineIndex">Index of the first line to return</param>
+		/// <param name="lineCount">Maximum number of lines to return</param>
+		/// <returns>List of lines for the file</returns>
+		public async Task<List<Utf8String>> ReadAsync(LogId logId, int lineIndex, int lineCount)
+		{
+			List<Utf8String> lines = new List<Utf8String>();
+			await ReadAsync(logId, lineIndex, lineCount, lines);
+			return lines;
+		}
+
+		/// <summary>
+		/// Reads data for the given log file, starting at the requested line index
+		/// </summary>
+		/// <param name="logId">The log to query</param>
+		/// <param name="lineIndex">Index of the first line to return</param>
+		/// <param name="lineCount">Maximum number of lines to return</param>
+		/// <param name="lines">Buffer for the lines to be returned</param>
+		/// <returns>List of lines for the file</returns>
+		public async Task ReadAsync(LogId logId, int lineIndex, int lineCount, List<Utf8String> lines)
+		{
+			int lowerBound = lineIndex;
+			int upperBound = lineIndex + lineCount;
+			int chunkLowerBound = lowerBound & ~(ChunkLineCount - 1);
+			int chunkUpperBound = (upperBound + (ChunkLineCount - 1)) & ~(ChunkLineCount - 1);
+
+			int nextIndex = lineIndex;
+
+			SortedSetEntry<ReadOnlyMemory<byte>>[] entries = await RedisTailData(logId).RangeByScoreWithScoresAsync(chunkLowerBound, chunkUpperBound);
+			foreach (SortedSetEntry<ReadOnlyMemory<byte>> entry in entries)
+			{
+				ReadOnlyMemory<byte> memory = entry.Element;
+
+				int index = (int)entry.Score;
+				if (index > nextIndex)
+				{
+					break;
+				}
+				for (; ; index++)
+				{
+					int lineLength = memory.Span.IndexOf((byte)'\n');
+					if (lineLength == -1)
+					{
+						break;
+					}
+					if (index >= nextIndex)
+					{
+						lines.Add(new Utf8String(memory.Slice(0, lineLength)));
+						nextIndex++;
+					}
+					memory = memory.Slice(lineLength + 1);
+				}
+			}
+		}
+
+		/// <summary>
+		/// Flushes the given log file to the given number of lines
+		/// </summary>
+		/// <param name="logId">Log to flush</param>
+		/// <param name="lineCount">Number of lines in the flushed log file</param>
+		/// <returns></returns>
+		public async ValueTask FlushAsync(LogId logId, int lineCount)
+		{
+			string trimEntry = $"{logId}={lineCount}";
+			double trimTime = (_clock.UtcNow + TrimAfter - DateTime.UnixEpoch).TotalSeconds;
+			_ = await _trimQueue.AddAsync(trimEntry, trimTime, flags: CommandFlags.FireAndForget);
+		}
+
+		/// <summary>
+		/// Return the desired next line to append for the given log file, or -1 if it is not being tailed.
+		/// </summary>
+		/// <param name="logId">Log to query</param>
+		/// <returns>Number of lines for the log file</returns>
+		public async Task<int> GetTailNext(LogId logId)
+		{
+			int? tailNext = await RedisTailNext(logId).GetValueAsync();
+			return tailNext ?? -1;
+		}
+
+		/// <summary>
+		/// Waits for a request to read tail data
+		/// </summary>
+		/// <param name="logId">The log file to query</param>
+		/// <param name="cancellationToken">Cancellation token for the operation</param>
+		/// <returns></returns>
+		public async Task<int> WaitForTailNext(LogId logId, CancellationToken cancellationToken)
+		{
+			AsyncEvent notifyEvent = _notifyLogEvents.GetOrAdd(logId, new AsyncEvent());
+			try
+			{
+				using (IDisposable registration = cancellationToken.Register(() => notifyEvent.Latch()))
+				{
+					for (; ; )
+					{
+						Task notifyTask = notifyEvent.Task;
+
+						int? tailNext = await RedisTailNext(logId).GetValueAsync();
+						if (tailNext != null)
+						{
+							return tailNext.Value;
+						}
+						else if (cancellationToken.IsCancellationRequested)
+						{
+							return -1;
+						}
+
+						await notifyTask;
+					}
+				}
+			}
+			finally
+			{
+				_notifyLogEvents.TryRemove(logId, out _);
+			}
+		}
+
+		/// <summary>
+		/// Splits data into chunks aligned onto <see cref="ChunkLineCount"/> line boundaries.
+		/// </summary>
+		static int SplitLogDataToChunks(int lineIdx, ReadOnlyMemory<byte> data, int chunkLineCount, List<SortedSetEntry<ReadOnlyMemory<byte>>> entries)
+		{
+			for (; ; )
+			{
+				int minLineIdx = lineIdx;
+
+				// Find the next chunk
+				int length = 0;
+				for (; ; )
+				{
+					// Add the next line to the buffer
+					int lineLength = data.Span.Slice(length).IndexOf((byte)'\n') + 1;
+					if (lineLength == 0)
+					{
+						break;
+					}
+					length += lineLength;
+
+					// Move to the next line
+					lineIdx++;
+
+					// If this is a chunk boundary, break here
+					if ((lineIdx & (chunkLineCount - 1)) == 0)
+					{
+						break;
+					}
+				}
+
+				// If we couldn't find any more lines, bail out
+				if (length == 0)
+				{
+					break;
+				}
+
+				// Add this entry
+				ReadOnlyMemory<byte> chunk = data.Slice(0, length);
+				entries.Add(new SortedSetEntry<ReadOnlyMemory<byte>>(chunk, minLineIdx));
+				data = data.Slice(length);
+			}
+			return lineIdx;
+		}
+	}
+}
diff --git a/Engine/Source/Programs/Horde/Horde.Build/Logs/LogsController.cs b/Engine/Source/Programs/Horde/Horde.Build/Logs/LogsController.cs
--- a/Engine/Source/Programs/Horde/Horde.Build/Logs/LogsController.cs
+++ b/Engine/Source/Programs/Horde/Horde.Build/Logs/LogsController.cs
@@ -71,6 +71,27 @@ namespace Horde.Build.Logs
 			_settings = settings;
  		}
 
+		/// <summary>
+		/// Creates a new log file. This endpoint is used mainly for debugging; log documents for specific uses are usually 
+		/// created by the server and have their id passed into clients to append to.
+		/// </summary>
+		/// <param name="request">Request to create the log</param>
+		/// <param name="cancellationToken">Cancellation token for the request</param>
+		/// <returns>Information about the requested project</returns>
+		[HttpPost]
+		[Route("/api/v1/logs")]
+		[ProducesResponseType(typeof(GetLogFileResponse), 200)]
+		public async Task<ActionResult<object>> CreateLog(CreateLogFileRequest request, CancellationToken cancellationToken = default)
+		{
+			if (!await _aclService.AuthorizeAsync(AclAction.CreateLog, User, null))
+			{
+				return Forbid();
+			}
+
+			ILogFile logFile = await _logFileService.CreateLogFileAsync(JobId.Empty, null, request.Type, cancellationToken: cancellationToken);
+			return new CreateLogFileResponse { Id = logFile.Id.ToString() };
+		}
+
 		/// <summary>
 		/// Retrieve metadata about a specific log file
 		/// </summary>
@@ -173,7 +194,7 @@ namespace Horde.Build.Logs
 		/// <returns>Information about the requested project</returns>
 		[HttpGet]
 		[Route("/api/v1/logs/{logFileId}/lines")]
-		public async Task<ActionResult> GetLogLines(LogId logFileId, [FromQuery] int index = 0, [FromQuery] int? count = null, CancellationToken cancellationToken = default)
+		public async Task<ActionResult> GetLogLines(LogId logFileId, [FromQuery] int index = 0, [FromQuery] int count = 100, CancellationToken cancellationToken = default)
 		{
 			ILogFile? logFile = await _logFileService.GetLogFileAsync(logFileId, cancellationToken);
 			if (logFile == null)
diff --git a/Engine/Source/Programs/Horde/Horde.Build/Startup.cs b/Engine/Source/Programs/Horde/Horde.Build/Startup.cs
--- a/Engine/Source/Programs/Horde/Horde.Build/Startup.cs
+++ b/Engine/Source/Programs/Horde/Horde.Build/Startup.cs
@@ -408,6 +408,7 @@ namespace Horde.Build
 			services.AddSingleton<JobService>();
 			services.AddSingleton<LifetimeService>();
 			services.AddSingleton<ILogFileService, LogFileService>();
+			services.AddSingleton<LogTailService>();
 			services.AddSingleton<INotificationService, NotificationService>();
 			services.AddSingleton<PerforceServiceCache>();
 			services.AddSingleton<IPerforceService>(sp => sp.GetRequiredService<PerforceServiceCache>());
@@ -612,6 +613,7 @@ namespace Horde.Build
 				services.AddHostedService(provider => provider.GetRequiredService<IssueService>());
 				services.AddHostedService<IssueReportService>();
 				services.AddHostedService(provider => (LogFileService)provider.GetRequiredService<ILogFileService>());
+				services.AddHostedService(provider => provider.GetRequiredService<LogTailService>());
 				services.AddHostedService(provider => (NotificationService)provider.GetRequiredService<INotificationService>());
 				services.AddHostedService(provider => provider.GetRequiredService<ReplicationService>());
 				if (!settings.DisableSchedules)
diff --git a/Engine/Source/Programs/Shared/EpicGames.Horde.Tests/LogTests.cs b/Engine/Source/Programs/Shared/EpicGames.Horde.Tests/LogTests.cs
--- a/Engine/Source/Programs/Shared/EpicGames.Horde.Tests/LogTests.cs
+++ b/Engine/Source/Programs/Shared/EpicGames.Horde.Tests/LogTests.cs
@@ -87,7 +87,7 @@ namespace EpicGames.Horde.Tests
 			}
 
 			// Flush it to storage, and read the finished log node
-			NodeLocator locator = await builder.FlushAsync(new TreeWriter(_store), CancellationToken.None);
+			NodeLocator locator = await builder.FlushAsync(new TreeWriter(_store), true, CancellationToken.None);
 			LogNode log = await _reader.ReadNodeAsync<LogNode>(locator);
 
 			// Read the data back out and check it's the same
@@ -129,7 +129,7 @@ namespace EpicGames.Horde.Tests
 			{
 				builder.WriteData(Encoding.UTF8.GetBytes(lines[lineIdx]));
 			}
-			NodeLocator rootNodeLocator = await builder.FlushAsync(new TreeWriter(_store), CancellationToken.None);
+			NodeLocator rootNodeLocator = await builder.FlushAsync(new TreeWriter(_store), true, CancellationToken.None);
 
 			// Read it back in and test the index
 			LogNode rootNode = await _reader.ReadNodeAsync<LogNode>(rootNodeLocator);
diff --git a/Engine/Source/Programs/Shared/EpicGames.Horde/Logs/LogChunkNode.cs b/Engine/Source/Programs/Shared/EpicGames.Horde/Logs/LogChunkNode.cs
--- a/Engine/Source/Programs/Shared/EpicGames.Horde/Logs/LogChunkNode.cs
+++ b/Engine/Source/Programs/Shared/EpicGames.Horde/Logs/LogChunkNode.cs
@@ -330,6 +330,13 @@ namespace EpicGames.Horde.Logs
 			_lineOffsets.RemoveRange(1, _lineOffsets.Count - 1);
 		}
 
+		/// <summary>
+		/// Gets a line at the given index
+		/// </summary>
+		/// <param name="index">Index of the line</param>
+		/// <returns>Text for the line</returns>
+		public Utf8String GetLine(int index) => new Utf8String(_data.AsMemory(_lineOffsets[index], _lineOffsets[index + 1] - _lineOffsets[index]));
+
 		/// <summary>
 		/// Create a new chunk data object with the given data appended. The internal buffers are reused, with the assumption that
 		/// there is no contention over writing to the same location in the chunk.
@@ -398,7 +405,7 @@ namespace EpicGames.Horde.Logs
 		/// </summary>
 		/// <param name="input">The input data</param>
 		/// <returns>True if the given text is empty</returns>
-		public static bool IsEmptyOrWhitespace(ReadOnlySpan<byte> input)
+		static bool IsEmptyOrWhitespace(ReadOnlySpan<byte> input)
 		{
 			for (int idx = 0; idx < input.Length; idx++)
 			{
@@ -418,7 +425,7 @@ namespace EpicGames.Horde.Logs
 		/// <param name="output">Output buffer for the converted line</param>
 		/// <param name="outputOffset">Offset within the buffer to write the converted data</param>
 		/// <returns></returns>
-		public static int ConvertToPlainText(ReadOnlySpan<byte> input, byte[] output, int outputOffset)
+		static int ConvertToPlainText(ReadOnlySpan<byte> input, byte[] output, int outputOffset)
 		{
 			if (IsEmptyOrWhitespace(input))
 			{
@@ -672,6 +679,10 @@ namespace EpicGames.Horde.Logs
 				}
 				lineIdx -= chunk.LineCount;
 			}
+			for (; lineIdx < _nextChunkBuilder.LineCount; lineIdx++)
+			{
+				yield return _nextChunkBuilder.GetLine(lineIdx);
+			}
 		}
 
 		/// <summary>
diff --git a/Engine/Source/Programs/Shared/EpicGames.Horde/Logs/LogNode.cs b/Engine/Source/Programs/Shared/EpicGames.Horde/Logs/LogNode.cs
--- a/Engine/Source/Programs/Shared/EpicGames.Horde/Logs/LogNode.cs
+++ b/Engine/Source/Programs/Shared/EpicGames.Horde/Logs/LogNode.cs
@@ -41,7 +41,7 @@ namespace EpicGames.Horde.Logs
 		/// <summary>
 		/// Default value for an empty log file
 		/// </summary>
-		public static LogNode Empty { get; } = new LogNode(LogFormat.Json, 0, 0, Array.Empty<LogChunkRef>(), new TreeNodeRef<LogIndexNode>(LogIndexNode.Empty));
+		public static LogNode Empty { get; } = new LogNode(LogFormat.Json, 0, 0, Array.Empty<LogChunkRef>(), new TreeNodeRef<LogIndexNode>(LogIndexNode.Empty), false);
 
 		/// <summary>
 		/// Format for this log file
@@ -68,16 +68,22 @@ namespace EpicGames.Horde.Logs
 		/// </summary>
 		public TreeNodeRef<LogIndexNode> IndexRef { get; }
 
+		/// <summary>
+		/// Whether this log is complete
+		/// </summary>
+		public bool Complete { get; }
+
 		/// <summary>
 		/// Deserializing constructor
 		/// </summary>
-		public LogNode(LogFormat format, int lineCount, long length, IReadOnlyList<LogChunkRef> textChunkRefs, TreeNodeRef<LogIndexNode> indexRef)
+		public LogNode(LogFormat format, int lineCount, long length, IReadOnlyList<LogChunkRef> textChunkRefs, TreeNodeRef<LogIndexNode> indexRef, bool complete)
 		{
 			Format = format;
 			LineCount = lineCount;
 			Length = length;
 			TextChunkRefs = textChunkRefs.ToArray();
 			IndexRef = indexRef;
+			Complete = complete;
 		}
 
 		/// <summary>
@@ -91,6 +97,7 @@ namespace EpicGames.Horde.Logs
 			Length = (long)reader.ReadUnsignedVarInt();
 			IndexRef = reader.ReadRef<LogIndexNode>();
 			TextChunkRefs = reader.ReadList(() => new LogChunkRef(reader));
+			Complete = reader.ReadBoolean();
 		}
 
 		/// <inheritdoc/>
@@ -101,6 +108,7 @@ namespace EpicGames.Horde.Logs
 			writer.WriteUnsignedVarInt((ulong)Length);
 			writer.WriteRef(IndexRef);
 			writer.WriteList(TextChunkRefs, x => x.Serialize(writer));
+			writer.WriteBoolean(Complete);
 		}
 
 		/// <inheritdoc/>
@@ -259,8 +267,9 @@ namespace EpicGames.Horde.Logs
 		/// Flushes the written data to the log
 		/// </summary>
 		/// <param name="writer">Writer for the output nodes</param>
+		/// <param name="complete">Whether the log is complete</param>
 		/// <param name="cancellationToken">Cancellation token for the operation</param>
-		public async Task<NodeLocator> FlushAsync(TreeWriter writer, CancellationToken cancellationToken)
+		public async Task<NodeLocator> FlushAsync(TreeWriter writer, bool complete, CancellationToken cancellationToken)
 		{
 			// Capture the new data that needs to be written
 			IReadOnlyList<LogChunkNode> writeTextChunks;
@@ -288,7 +297,7 @@ namespace EpicGames.Horde.Logs
 				length += writeTextChunk.Length;
 			}
 
-			LogNode newRoot = new LogNode(_format, lineCount, length, newJsonChunkRefs, new TreeNodeRef<LogIndexNode>(newIndex));
+			LogNode newRoot = new LogNode(_format, lineCount, length, newJsonChunkRefs, new TreeNodeRef<LogIndexNode>(newIndex), complete);
 
 			IoHash newRootHash = await writer.WriteAsync(new TreeNodeRef(newRoot), cancellationToken);
 			await writer.FlushAsync(cancellationToken);
diff --git a/Engine/Source/Programs/Shared/EpicGames.Redis/RedisString.cs b/Engine/Source/Programs/Shared/EpicGames.Redis/RedisString.cs
--- a/Engine/Source/Programs/Shared/EpicGames.Redis/RedisString.cs
+++ b/Engine/Source/Programs/Shared/EpicGames.Redis/RedisString.cs
@@ -95,7 +95,7 @@ namespace EpicGames.Redis
 			RedisValue value = await str.GetDatabase().StringGetAsync(str.Key, flags);
 			if (value.IsNullOrEmpty)
 			{
-				return default(TElement);
+				return null;
 			}
 			return RedisSerializer.Deserialize<TElement>(value);
 		}
