Horde: Fix line count returned using new logger backend.

diff --git a/Engine/Source/Programs/Horde/Horde.Build/Logs/ILogFile.cs b/Engine/Source/Programs/Horde/Horde.Build/Logs/ILogFile.cs
--- a/Engine/Source/Programs/Horde/Horde.Build/Logs/ILogFile.cs
+++ b/Engine/Source/Programs/Horde/Horde.Build/Logs/ILogFile.cs
@@ -84,6 +84,11 @@ namespace Horde.Build.Logs
 		/// Name of the ref used to store data for this log
 		/// </summary>
 		public RefName RefName { get; }
+
+		/// <summary>
+		/// Number of lines (V2 storage backend)
+		/// </summary>
+		public int LineCount { get; }
 	}
 
 	/// <summary>
diff --git a/Engine/Source/Programs/Horde/Horde.Build/Logs/ILogFileCollection.cs b/Engine/Source/Programs/Horde/Horde.Build/Logs/ILogFileCollection.cs
--- a/Engine/Source/Programs/Horde/Horde.Build/Logs/ILogFileCollection.cs
+++ b/Engine/Source/Programs/Horde/Horde.Build/Logs/ILogFileCollection.cs
@@ -63,6 +63,15 @@ namespace Horde.Build.Logs
 		/// <returns>The new log file document</returns>
 		Task<ILogFile> CreateLogFileAsync(JobId jobId, SessionId? sessionId, LogType type, LogId? logId = null, CancellationToken cancellationToken = default);
 
+		/// <summary>
+		/// Updates the line count for a log file (v2 backend only)
+		/// </summary>
+		/// <param name="logFileInterface">The current log file</param>
+		/// <param name="lineCount">New line count for the log file</param>
+		/// <param name="cancellationToken">Cancellation token for the call</param>
+		/// <returns>The updated log file document</returns>
+		Task<ILogFile> UpdateLineCountAsync(ILogFile logFileInterface, int lineCount, CancellationToken cancellationToken);
+
 		/// <summary>
 		/// Adds a new chunk
 		/// </summary>
diff --git a/Engine/Source/Programs/Horde/Horde.Build/Logs/LogFileCollection.cs b/Engine/Source/Programs/Horde/Horde.Build/Logs/LogFileCollection.cs
--- a/Engine/Source/Programs/Horde/Horde.Build/Logs/LogFileCollection.cs
+++ b/Engine/Source/Programs/Horde/Horde.Build/Logs/LogFileCollection.cs
@@ -73,6 +73,7 @@ namespace Horde.Build.Logs
 
 			public List<LogChunkDocument> Chunks { get; set; } = new List<LogChunkDocument>();
 
+			public int LineCount { get; set; }
 			public RefName RefName { get; set; }
 
 			[BsonRequired]
@@ -131,6 +132,14 @@ namespace Horde.Build.Logs
 			return newLogFile;
 		}
 
+		/// <inheritdoc/>
+		public async Task<ILogFile> UpdateLineCountAsync(ILogFile logFileInterface, int lineCount, CancellationToken cancellationToken)
+		{
+			FilterDefinition<LogFileDocument> filter = Builders<LogFileDocument>.Filter.Eq(x => x.Id, logFileInterface.Id);
+			UpdateDefinition<LogFileDocument> update = Builders<LogFileDocument>.Update.Set(x => x.LineCount, lineCount).Inc(x => x.UpdateIndex, 1);
+			return await _logFiles.FindOneAndUpdateAsync(filter, update, new FindOneAndUpdateOptions<LogFileDocument, ILogFile> { ReturnDocument = ReturnDocument.After }, cancellationToken);
+		}
+
 		/// <inheritdoc/>
 		public async Task<ILogFile?> TryAddChunkAsync(ILogFile logFileInterface, long offset, int lineIndex, CancellationToken cancellationToken)
 		{
diff --git a/Engine/Source/Programs/Horde/Horde.Build/Logs/LogFileMessages.cs b/Engine/Source/Programs/Horde/Horde.Build/Logs/LogFileMessages.cs
--- a/Engine/Source/Programs/Horde/Horde.Build/Logs/LogFileMessages.cs
+++ b/Engine/Source/Programs/Horde/Horde.Build/Logs/LogFileMessages.cs
@@ -64,11 +64,6 @@ namespace Horde.Build.Logs
 		/// </summary>
 		public LogType Type { get; set; }
 
-		/// <summary>
-        /// Length of the log, in bytes
-        /// </summary>
-        public long Length { get; set; }
-
 		/// <summary>
         /// Number of lines in the file
         /// </summary>
@@ -84,7 +79,6 @@ namespace Horde.Build.Logs
             Id = logFile.Id.ToString();
             JobId = logFile.JobId.ToString();
 			Type = logFile.Type;
-            Length = metadata.Length;
             LineCount = metadata.MaxLineIndex;
 		}
 	}
diff --git a/Engine/Source/Programs/Horde/Horde.Build/Logs/LogFileService.cs b/Engine/Source/Programs/Horde/Horde.Build/Logs/LogFileService.cs
--- a/Engine/Source/Programs/Horde/Horde.Build/Logs/LogFileService.cs
+++ b/Engine/Source/Programs/Horde/Horde.Build/Logs/LogFileService.cs
@@ -925,19 +925,32 @@ namespace Horde.Build.Logs
 		public async Task<LogMetadata> GetMetadataAsync(ILogFile logFile, CancellationToken cancellationToken)
 		{
 			LogMetadata metadata = new LogMetadata();
-			if (logFile.Chunks.Count > 0)
+			if (_settings.Value.FeatureFlags.EnableNewLogger)
 			{
-				ILogChunk chunk = logFile.Chunks[logFile.Chunks.Count - 1];
-				if (logFile.MaxLineIndex == null || chunk.Length == 0)
+				metadata.MaxLineIndex = logFile.LineCount;
+
+				int nextIdx = await _logTailService.GetTailNext(logFile.Id);
+				if (nextIdx != -1)
 				{
-					LogChunkData chunkData = await ReadChunkAsync(logFile, logFile.Chunks.Count - 1);
-					metadata.Length = chunk.Offset + chunkData.Length;
-					metadata.MaxLineIndex = chunk.LineIndex + chunkData.LineCount;
+					metadata.MaxLineIndex = nextIdx;
 				}
-				else
+			}
+			else
+			{
+				if (logFile.Chunks.Count > 0)
 				{
-					metadata.Length = chunk.Offset + chunk.Length;
-					metadata.MaxLineIndex = logFile.MaxLineIndex.Value;
+					ILogChunk chunk = logFile.Chunks[logFile.Chunks.Count - 1];
+					if (logFile.MaxLineIndex == null || chunk.Length == 0)
+					{
+						LogChunkData chunkData = await ReadChunkAsync(logFile, logFile.Chunks.Count - 1);
+						metadata.Length = chunk.Offset + chunkData.Length;
+						metadata.MaxLineIndex = chunk.LineIndex + chunkData.LineCount;
+					}
+					else
+					{
+						metadata.Length = chunk.Offset + chunk.Length;
+						metadata.MaxLineIndex = logFile.MaxLineIndex.Value;
+					}
 				}
 			}
 			return metadata;
diff --git a/Engine/Source/Programs/Horde/Horde.Build/Logs/LogRpcService.cs b/Engine/Source/Programs/Horde/Horde.Build/Logs/LogRpcService.cs
--- a/Engine/Source/Programs/Horde/Horde.Build/Logs/LogRpcService.cs
+++ b/Engine/Source/Programs/Horde/Horde.Build/Logs/LogRpcService.cs
@@ -21,6 +21,7 @@ namespace Horde.Build.Logs
 	public class LogRpcService : LogRpc.LogRpcBase
 	{
 		readonly ILogFileService _logFileService;
+		readonly ILogFileCollection _logFileCollection;
 		readonly LogTailService _logTailService;
 		readonly StorageService _storageService;
 		readonly ILogger<LogRpcService> _logger;
@@ -28,9 +29,10 @@ namespace Horde.Build.Logs
 		/// <summary>
 		/// Constructor
 		/// </summary>
-		public LogRpcService(ILogFileService logFileService, LogTailService logTailService, StorageService storageService, ILogger<LogRpcService> logger)
+		public LogRpcService(ILogFileService logFileService, ILogFileCollection logFileCollection, LogTailService logTailService, StorageService storageService, ILogger<LogRpcService> logger)
 		{
 			_logFileService = logFileService;
+			_logFileCollection = logFileCollection;
 			_logTailService = logTailService;
 			_storageService = storageService;
 			_logger = logger;
@@ -53,6 +55,8 @@ namespace Horde.Build.Logs
 			_logger.LogInformation("Updating {LogId} to node {Locator}", request.LogId, request.BlobLocator);
 			await store.WriteRefTargetAsync(new RefName(request.LogId), NodeLocator.Parse(request.BlobLocator));
 
+			await _logFileCollection.UpdateLineCountAsync(logFile, request.LineCount, CancellationToken.None);
+
 			await _logTailService.FlushAsync(logFile.Id, request.LineCount);
 
 			return new UpdateLogResponse();
diff --git a/Engine/Source/Programs/Horde/Horde.Build/Logs/LogsController.cs b/Engine/Source/Programs/Horde/Horde.Build/Logs/LogsController.cs
--- a/Engine/Source/Programs/Horde/Horde.Build/Logs/LogsController.cs
+++ b/Engine/Source/Programs/Horde/Horde.Build/Logs/LogsController.cs
@@ -214,8 +214,8 @@ namespace Horde.Build.Logs
 				stream.WriteByte((byte)'{');
 
 				stream.Write(Encoding.UTF8.GetBytes($"\"index\":{index},"));
-				stream.Write(Encoding.UTF8.GetBytes($"\"count\":{count},"));
-				stream.Write(Encoding.UTF8.GetBytes($"\"maxLineIndex\":{metadata.MaxLineIndex},"));
+				stream.Write(Encoding.UTF8.GetBytes($"\"count\":{lines.Count},"));
+				stream.Write(Encoding.UTF8.GetBytes($"\"maxLineIndex\":{Math.Max(metadata.MaxLineIndex, index + lines.Count)},"));
 				stream.Write(Encoding.UTF8.GetBytes($"\"format\":{ (logFile.Type == LogType.Json ? "\"JSON\"" : "\"TEXT\"")},"));
 
 				stream.Write(Encoding.UTF8.GetBytes($"\"lines\":["));
