Horde: Support indexed searches for new log storage format. Persisted data is indexed using ngram tries, tail data is converted to plain text and searched directly.

diff --git a/Engine/Source/Programs/Horde/Horde.Build.Tests/LogIndexingTests.cs b/Engine/Source/Programs/Horde/Horde.Build.Tests/LogIndexingTests.cs
--- a/Engine/Source/Programs/Horde/Horde.Build.Tests/LogIndexingTests.cs
+++ b/Engine/Source/Programs/Horde/Horde.Build.Tests/LogIndexingTests.cs
@@ -19,6 +19,7 @@ using Microsoft.Extensions.DependencyInjection;
 using Microsoft.Extensions.Logging;
 using Microsoft.VisualStudio.TestTools.UnitTesting;
 using Horde.Build.Logs.Data;
+using EpicGames.Horde.Logs;
 
 namespace Horde.Build.Tests
 {
@@ -161,7 +162,7 @@ namespace Horde.Build.Tests
 					{
 						string str = lines[lineIdx].Substring(strOfs, strLen);
 
-						LogSearchStats stats = new LogSearchStats();
+						SearchStats stats = new SearchStats();
 						List<int> results = await _logFileService.SearchLogDataAsync(logFile, str, 0, 5, stats, CancellationToken.None);
 						Assert.AreEqual(1, results.Count);
 						Assert.AreEqual(lineIdx, results[0]);
@@ -187,7 +188,7 @@ namespace Horde.Build.Tests
 			await ((LogFileService)_logFileService).FlushPendingWritesAsync();
 
 			{
-				LogSearchStats stats = new LogSearchStats();
+				SearchStats stats = new SearchStats();
 				List<int> results = await _logFileService.SearchLogDataAsync(logFile, "abc", 0, 5, stats, CancellationToken.None);
 				Assert.AreEqual(1, results.Count);
 				Assert.AreEqual(0, results[0]);
@@ -197,7 +198,7 @@ namespace Horde.Build.Tests
 				Assert.AreEqual(0, stats.NumFalsePositiveBlocks);
 			}
 			{
-				LogSearchStats stats = new LogSearchStats();
+				SearchStats stats = new SearchStats();
 				List<int> results = await _logFileService.SearchLogDataAsync(logFile, "def", 0, 5, stats, CancellationToken.None);
 				Assert.AreEqual(1, results.Count);
 				Assert.AreEqual(1, results[0]);
@@ -207,7 +208,7 @@ namespace Horde.Build.Tests
 				Assert.AreEqual(0, stats.NumFalsePositiveBlocks);
 			}
 			{
-				LogSearchStats stats = new LogSearchStats();
+				SearchStats stats = new SearchStats();
 				List<int> results = await _logFileService.SearchLogDataAsync(logFile, "ghi", 0, 5, stats, CancellationToken.None);
 				Assert.AreEqual(1, results.Count);
 				Assert.AreEqual(2, results[0]);
@@ -239,7 +240,7 @@ namespace Horde.Build.Tests
 
 		async Task SearchLogDataTestAsync(ILogFile logFile, string text, int firstLine, int count, int[] expectedLines)
 		{
-			LogSearchStats stats = new LogSearchStats();
+			SearchStats stats = new SearchStats();
 			List<int> lines = await _logFileService.SearchLogDataAsync(logFile, text, firstLine, count, stats, CancellationToken.None);
 			Assert.IsTrue(lines.SequenceEqual(expectedLines));
 		}
diff --git a/Engine/Source/Programs/Horde/Horde.Build.Tests/LogTailServiceTests.cs b/Engine/Source/Programs/Horde/Horde.Build.Tests/LogTailServiceTests.cs
--- a/Engine/Source/Programs/Horde/Horde.Build.Tests/LogTailServiceTests.cs
+++ b/Engine/Source/Programs/Horde/Horde.Build.Tests/LogTailServiceTests.cs
@@ -176,7 +176,7 @@ namespace Horde.Build.Tests
 
 			LogTailService tailService2 = new LogTailService(ServiceProvider.GetRequiredService<RedisService>(), ServiceProvider.GetRequiredService<IClock>(), 4, ServiceProvider.GetRequiredService<ILogger<LogTailService>>());
 
-			Task<int> task = tailService.WaitForTailNext(_logId, CancellationToken.None);
+			Task<int> task = tailService.WaitForTailNextAsync(_logId, CancellationToken.None);
 			Assert.IsTrue(!task.IsCompleted);
 
 			await tailService2.EnableTailingAsync(_logId, 30);
diff --git a/Engine/Source/Programs/Horde/Horde.Build/Logs/Data/LogIndexData.cs b/Engine/Source/Programs/Horde/Horde.Build/Logs/Data/LogIndexData.cs
--- a/Engine/Source/Programs/Horde/Horde.Build/Logs/Data/LogIndexData.cs
+++ b/Engine/Source/Programs/Horde/Horde.Build/Logs/Data/LogIndexData.cs
@@ -5,6 +5,7 @@ using System.Collections.Generic;
 using System.IO;
 using System.Linq;
 using EpicGames.Core;
+using EpicGames.Horde.Logs;
 using Horde.Build.Utilities;
 using OpenTracing;
 using OpenTracing.Util;
@@ -101,37 +102,6 @@ namespace Horde.Build.Logs.Data
 		}
 	}
 
-	/// <summary>
-	/// Stats for a search
-	/// </summary>
-	public class LogSearchStats
-	{
-		/// <summary>
-		/// Number of blocks that were scanned
-		/// </summary>
-		public int NumScannedBlocks { get; set; }
-
-		/// <summary>
-		/// Number of bytes that had to be scanned for results
-		/// </summary>
-		public int NumScannedBytes { get; set; }
-
-		/// <summary>
-		/// Number of blocks that were skipped
-		/// </summary>
-		public int NumSkippedBlocks { get; set; }
-
-		/// <summary>
-		/// Number of blocks that had to be decompressed
-		/// </summary>
-		public int NumDecompressedBlocks { get; set; }
-
-		/// <summary>
-		/// Number of blocks that were searched but did not contain the search term
-		/// </summary>
-		public int NumFalsePositiveBlocks { get; set; }
-	}
-
 	/// <summary>
 	/// Data for a log chunk
 	/// </summary>
@@ -257,7 +227,7 @@ namespace Horde.Build.Logs.Data
 		/// <param name="text">Text to search for</param>
 		/// <param name="stats">Receives stats for the search</param>
 		/// <returns>List of line numbers for the text</returns>
-		public IEnumerable<int> Search(int firstLineIndex, SearchText text, LogSearchStats stats)
+		public IEnumerable<int> Search(int firstLineIndex, SearchText text, SearchStats stats)
 		{
 			int lastBlockCount = 0;
 			foreach (int blockIdx in EnumeratePossibleBlocks(text.Bytes, firstLineIndex))
diff --git a/Engine/Source/Programs/Horde/Horde.Build/Logs/LogFileMessages.cs b/Engine/Source/Programs/Horde/Horde.Build/Logs/LogFileMessages.cs
--- a/Engine/Source/Programs/Horde/Horde.Build/Logs/LogFileMessages.cs
+++ b/Engine/Source/Programs/Horde/Horde.Build/Logs/LogFileMessages.cs
@@ -2,6 +2,7 @@
 
 using System;
 using System.Collections.Generic;
+using EpicGames.Horde.Logs;
 using Horde.Build.Logs.Data;
 
 namespace Horde.Build.Logs
@@ -96,7 +97,7 @@ namespace Horde.Build.Logs
 		/// <summary>
 		/// Stats for the search
 		/// </summary>
-		public LogSearchStats? Stats { get; set; }
+		public SearchStats? Stats { get; set; }
 	}
 }
 
diff --git a/Engine/Source/Programs/Horde/Horde.Build/Logs/LogFileService.cs b/Engine/Source/Programs/Horde/Horde.Build/Logs/LogFileService.cs
--- a/Engine/Source/Programs/Horde/Horde.Build/Logs/LogFileService.cs
+++ b/Engine/Source/Programs/Horde/Horde.Build/Logs/LogFileService.cs
@@ -7,6 +7,7 @@ using System.Diagnostics;
 using System.Globalization;
 using System.IO;
 using System.Linq;
+using System.Runtime.CompilerServices;
 using System.Security.Claims;
 using System.Text;
 using System.Text.Json;
@@ -201,7 +202,7 @@ namespace Horde.Build.Logs
 		/// <param name="stats">Receives stats for the search</param>
 		/// <param name="cancellationToken">Cancellation token for the call</param>
 		/// <returns>List of line numbers containing the given term</returns>
-		Task<List<int>> SearchLogDataAsync(ILogFile logFile, string text, int firstLine, int count, LogSearchStats stats, CancellationToken cancellationToken);
+		Task<List<int>> SearchLogDataAsync(ILogFile logFile, string text, int firstLine, int count, SearchStats stats, CancellationToken cancellationToken);
 	}
 
 	/// <summary>
@@ -929,7 +930,7 @@ namespace Horde.Build.Logs
 			{
 				metadata.MaxLineIndex = logFile.LineCount;
 
-				int nextIdx = await _logTailService.GetTailNext(logFile.Id);
+				int nextIdx = await _logTailService.GetTailNextAsync(logFile.Id);
 				if (nextIdx != -1)
 				{
 					metadata.MaxLineIndex = nextIdx;
@@ -1863,7 +1864,7 @@ namespace Horde.Build.Logs
 		}
 
 		/// <inheritdoc/>
-		public async Task<List<int>> SearchLogDataAsync(ILogFile logFile, string text, int firstLine, int count, LogSearchStats searchStats, CancellationToken cancellationToken)
+		public async Task<List<int>> SearchLogDataAsync(ILogFile logFile, string text, int firstLine, int count, SearchStats searchStats, CancellationToken cancellationToken)
 		{
 			Stopwatch timer = Stopwatch.StartNew();
 
@@ -1875,7 +1876,11 @@ namespace Horde.Build.Logs
 			List<int> results = new List<int>();
 			if (count > 0)
 			{
-				await using IAsyncEnumerator<int> enumerator = SearchLogDataInternalAsync(logFile, text, firstLine, searchStats).GetAsyncEnumerator(cancellationToken);
+				IAsyncEnumerable<int> enumerable =_settings.Value.FeatureFlags.EnableNewLogger ?
+						SearchLogDataInternalNewAsync(logFile, text, firstLine, searchStats, cancellationToken) :
+						SearchLogDataInternalAsync(logFile, text, firstLine, searchStats);
+
+				await using IAsyncEnumerator<int> enumerator = enumerable.GetAsyncEnumerator(cancellationToken);
 				while (await enumerator.MoveNextAsync() && results.Count < count)
 				{
 					results.Add(enumerator.Current);
@@ -1886,8 +1891,87 @@ namespace Horde.Build.Logs
 			return results;
 		}
 
-		/// <inheritdoc/>
-		public async IAsyncEnumerable<int> SearchLogDataInternalAsync(ILogFile logFile, string text, int firstLine, LogSearchStats searchStats)
+		async IAsyncEnumerable<int> SearchLogDataInternalNewAsync(ILogFile logFile, string text, int firstLine, SearchStats searchStats, [EnumeratorCancellation] CancellationToken cancellationToken)
+		{
+			SearchTerm searchText = new SearchTerm(text);
+			TreeReader reader = await GetTreeReaderAsync(cancellationToken);
+
+			// Search the index
+			if (logFile.LineCount > 0)
+			{
+				LogNode? root = await reader.ReadNodeAsync<LogNode>(logFile.RefName, cancellationToken: cancellationToken);
+				if(root != null)
+				{
+					LogIndexNode index = await root.IndexRef.ExpandAsync(reader, cancellationToken);
+					await foreach (int lineIdx in index.Search(reader, firstLine, searchText, searchStats, cancellationToken))
+					{
+						yield return lineIdx;
+					}
+					if (root.Complete)
+					{
+						yield break;
+					}
+					firstLine = root.LineCount;
+				}
+			}
+
+			// Search any tail data we have
+			for (; ; )
+			{
+				Utf8String[] lines = await ReadTailAsync(logFile, firstLine, cancellationToken);
+				if (lines.Length == 0)
+				{
+					break;
+				}
+
+				for (int idx = 0; idx < lines.Length; idx++)
+				{
+					if (SearchTerm.FindNextOcurrence(lines[idx].Span, 0, searchText) != -1)
+					{
+						yield return firstLine + idx;
+					}
+				}
+
+				firstLine += lines.Length;
+			}
+		}
+
+		async Task<Utf8String[]> ReadTailAsync(ILogFile logFile, int index, CancellationToken cancellationToken)
+		{
+			const int BatchSize = 128;
+
+			string cacheKey = $"{logFile.Id}@{index}";
+			if (_logFileCache.TryGetValue(cacheKey, out Utf8String[]? lines))
+			{
+				return lines!;
+			}
+
+			lines = (await _logTailService.ReadAsync(logFile.Id, index, BatchSize)).ToArray();
+			if (logFile.Type == LogType.Json)
+			{
+				LogChunkBuilder builder = new LogChunkBuilder(lines.Sum(x => x.Length));
+				foreach (Utf8String line in lines)
+				{
+					builder.AppendJsonAsPlainText(line.Span, _logger);
+				}
+				lines = lines.ToArray();
+			}
+
+			if (lines.Length == BatchSize)
+			{
+				int length = lines.Sum(x => x.Length);
+				using (ICacheEntry entry = _logFileCache.CreateEntry(cacheKey))
+				{
+					entry.SetSlidingExpiration(TimeSpan.FromMinutes(1.0));
+					entry.SetSize(length);
+					entry.SetValue(lines);
+				}
+			}
+
+			return lines;
+		}
+
+		async IAsyncEnumerable<int> SearchLogDataInternalAsync(ILogFile logFile, string text, int firstLine, SearchStats searchStats)
 		{
 			SearchText searchText = new SearchText(text);
 
diff --git a/Engine/Source/Programs/Horde/Horde.Build/Logs/LogRpcService.cs b/Engine/Source/Programs/Horde/Horde.Build/Logs/LogRpcService.cs
--- a/Engine/Source/Programs/Horde/Horde.Build/Logs/LogRpcService.cs
+++ b/Engine/Source/Programs/Horde/Horde.Build/Logs/LogRpcService.cs
@@ -83,7 +83,7 @@ namespace Horde.Build.Logs
 
 				using (CancellationTokenSource cancellationSource = new CancellationTokenSource())
 				{
-					Task<int> waitTask = _logTailService.WaitForTailNext(logId, cancellationSource.Token);
+					Task<int> waitTask = _logTailService.WaitForTailNextAsync(logId, cancellationSource.Token);
 
 					await Task.WhenAny(waitTask, moveNextTask);
 					cancellationSource.Cancel();
diff --git a/Engine/Source/Programs/Horde/Horde.Build/Logs/LogTailService.cs b/Engine/Source/Programs/Horde/Horde.Build/Logs/LogTailService.cs
--- a/Engine/Source/Programs/Horde/Horde.Build/Logs/LogTailService.cs
+++ b/Engine/Source/Programs/Horde/Horde.Build/Logs/LogTailService.cs
@@ -28,8 +28,8 @@ namespace Horde.Build.Logs
 	///   data to a sorted set for that log scored by index of the first line (see <see cref="RedisTailData"/>).
 	/// * Chunks are split on fixed boundaries, in order to allow older entries to be purged by score without having to count the number of lines
 	///   they contain first.
-	/// * Agent polls for requests to provide tail data via <see cref="LogRpcService.UpdateLogTail"/>, which calls <see cref="WaitForTailNext"/>.
-	///   <see cref="WaitForTailNext"/> returns the index of the total line count for this log if it is being tailed (indicating the index of
+	/// * Agent polls for requests to provide tail data via <see cref="LogRpcService.UpdateLogTail"/>, which calls <see cref="WaitForTailNextAsync"/>.
+	///   <see cref="WaitForTailNextAsync"/> returns the index of the total line count for this log if it is being tailed (indicating the index of
 	///   the next line that the agent should send to the server), or blocks until the log is being tailed.
 	/// * Once data is flushed to persistent storage, the number of flushed lines is added to <see cref="_trimQueue"/> and the line data is removed
 	///   from Redis after <see cref="TrimAfter"/>.
@@ -74,7 +74,7 @@ namespace Horde.Build.Logs
 		/// Constructor
 		/// </summary>
 		public LogTailService(RedisService redisService, IClock clock, ILogger<LogTailService> logger)
-			: this(redisService, clock, 4, logger)
+			: this(redisService, clock, 32, logger)
 		{
 		}
 
@@ -275,7 +275,7 @@ namespace Horde.Build.Logs
 		/// </summary>
 		/// <param name="logId">Log to query</param>
 		/// <returns>Number of lines for the log file</returns>
-		public async Task<int> GetTailNext(LogId logId)
+		public async Task<int> GetTailNextAsync(LogId logId)
 		{
 			int? tailNext = await RedisTailNext(logId).GetValueAsync();
 			return tailNext ?? -1;
@@ -287,7 +287,7 @@ namespace Horde.Build.Logs
 		/// <param name="logId">The log file to query</param>
 		/// <param name="cancellationToken">Cancellation token for the operation</param>
 		/// <returns></returns>
-		public async Task<int> WaitForTailNext(LogId logId, CancellationToken cancellationToken)
+		public async Task<int> WaitForTailNextAsync(LogId logId, CancellationToken cancellationToken)
 		{
 			AsyncEvent notifyEvent = _notifyLogEvents.GetOrAdd(logId, new AsyncEvent());
 			try
diff --git a/Engine/Source/Programs/Horde/Horde.Build/Logs/LogsController.cs b/Engine/Source/Programs/Horde/Horde.Build/Logs/LogsController.cs
--- a/Engine/Source/Programs/Horde/Horde.Build/Logs/LogsController.cs
+++ b/Engine/Source/Programs/Horde/Horde.Build/Logs/LogsController.cs
@@ -9,6 +9,7 @@ using System.Text;
 using System.Threading;
 using System.Threading.Tasks;
 using EpicGames.Core;
+using EpicGames.Horde.Logs;
 using EpicGames.Horde.Storage;
 using Horde.Build.Acls;
 using Horde.Build.Issues;
@@ -302,7 +303,7 @@ namespace Horde.Build.Logs
 			}
 
 			SearchLogFileResponse response = new SearchLogFileResponse();
-			response.Stats = new LogSearchStats();
+			response.Stats = new SearchStats();
 			response.Lines = await _logFileService.SearchLogDataAsync(logFile, text, firstLine, count, response.Stats, cancellationToken);
 			return response;
 		}
@@ -400,6 +401,10 @@ namespace Horde.Build.Logs
 			{
 				return true;
 			}
+			if (await _aclService.AuthorizeAsync(action, user, permissionsCache))
+			{
+				return true;
+			}
 			return false;
 		}
 	}
diff --git a/Engine/Source/Programs/Shared/EpicGames.Horde/Logs/LogChunkNode.cs b/Engine/Source/Programs/Shared/EpicGames.Horde/Logs/LogChunkNode.cs
--- a/Engine/Source/Programs/Shared/EpicGames.Horde/Logs/LogChunkNode.cs
+++ b/Engine/Source/Programs/Shared/EpicGames.Horde/Logs/LogChunkNode.cs
@@ -529,6 +529,20 @@ namespace EpicGames.Horde.Logs
 			}
 		}
 
+		/// <summary>
+		/// Create an array of lines from the text
+		/// </summary>
+		/// <returns>Array of lines</returns>
+		public Utf8String[] ToArray()
+		{
+			Utf8String[] lines = new Utf8String[LineCount];
+			for (int idx = 0; idx < LineCount; idx++)
+			{
+				lines[idx] = new Utf8String(_data, _lineOffsets[idx], _lineOffsets[idx + 1] - 1).Clone();
+			}
+			return lines;
+		}
+
 		/// <summary>
 		/// Create a <see cref="LogChunkNode"/> object from the current state
 		/// </summary>
diff --git a/Engine/Source/Programs/Shared/EpicGames.Horde/Logs/SearchTerm.cs b/Engine/Source/Programs/Shared/EpicGames.Horde/Logs/SearchTerm.cs
--- a/Engine/Source/Programs/Shared/EpicGames.Horde/Logs/SearchTerm.cs
+++ b/Engine/Source/Programs/Shared/EpicGames.Horde/Logs/SearchTerm.cs
@@ -95,7 +95,7 @@ namespace EpicGames.Horde.Logs
 		}
 
 		/// <summary>
-		/// Perform a case sensitive search for the next occurerence of the search term in a given buffer
+		/// Perform a case insensitive search for the next occurerence of the search term in a given buffer
 		/// </summary>
 		/// <param name="buffer">The buffer to search</param>
 		/// <param name="offset">Starting offset for the search</param>
