Horde: Prevent users being able to modify logs if they only have view session permissions.

diff --git a/Engine/Source/Programs/Horde/Horde.Build/Logs/LogsController.cs b/Engine/Source/Programs/Horde/Horde.Build/Logs/LogsController.cs
--- a/Engine/Source/Programs/Horde/Horde.Build/Logs/LogsController.cs
+++ b/Engine/Source/Programs/Horde/Horde.Build/Logs/LogsController.cs
@@ -397,7 +397,7 @@ namespace Horde.Build.Logs
 			{
 				return true;
 			}
-			if (logFile.SessionId != null && await _aclService.AuthorizeAsync(AclAction.ViewSession, user, permissionsCache))
+			if (action == AclAction.ViewLog && logFile.SessionId != null && await _aclService.AuthorizeAsync(AclAction.ViewSession, user, permissionsCache))
 			{
 				return true;
 			}
