Horde: Ensure instance data is copied in tests to avoid leaking state

diff --git a/Engine/Source/Programs/Horde/Horde.Build.Tests/Fleet/FakeAmazonEc2.cs b/Engine/Source/Programs/Horde/Horde.Build.Tests/Fleet/FakeAmazonEc2.cs
--- a/Engine/Source/Programs/Horde/Horde.Build.Tests/Fleet/FakeAmazonEc2.cs
+++ b/Engine/Source/Programs/Horde/Horde.Build.Tests/Fleet/FakeAmazonEc2.cs
@@ -33,6 +33,11 @@ public class FakeAmazonEc2
 	/// </summary>
 	internal Action<StartInstancesRequest, CancellationToken> OnStartInstances { get; set; } = (_, _) => { };
 	
+	/// <summary>
+	/// Callback when ModifyInstanceAttribute method is called. Can be used for mocking exception throwing.
+	/// </summary>
+	internal Func<ModifyInstanceAttributeRequest, CancellationToken, Task> OnModifyInstanceAttribute { get; set; } = (_, _) => Task.CompletedTask;
+	
 	internal IReadOnlyDictionary<string, Instance> Instances => _instances;
 
 	private readonly Dictionary<string, Dictionary<string, int>> _availabilityZoneCapacity = new()
@@ -88,6 +93,18 @@ public class FakeAmazonEc2
 		return _instances[instanceId];
 	}
 
+	public static Instance CopyInstance(Instance i)
+	{
+		return new Instance
+		{
+			InstanceId = i.InstanceId,
+			InstanceType = i.InstanceType,
+			State = i.State,
+			Placement = i.Placement,
+			LaunchTime = i.LaunchTime
+		};
+	}
+
 	public int GetStoppedInstanceCount() { return _instances.Values.Count(x => x.State == StateStopped); }
 	public int GetPendingInstanceCount() { return _instances.Values.Count(x => x.State == StatePending); }
 
@@ -95,7 +112,7 @@ public class FakeAmazonEc2
 	{
 		return Task.FromResult(new DescribeInstancesResponse
 		{
-			Reservations = new () { new Reservation { Instances = _instances.Values.ToList() } },
+			Reservations = new () { new Reservation { Instances = _instances.Values.Select(CopyInstance).ToList() } },
 			HttpStatusCode = HttpStatusCode.OK
 		});
 	}
@@ -149,7 +166,7 @@ public class FakeAmazonEc2
 		});
 	}
 	
-	private Task<ModifyInstanceAttributeResponse> ModifyInstanceAttributeAsync(ModifyInstanceAttributeRequest request, CancellationToken cancellationToken)
+	private async Task<ModifyInstanceAttributeResponse> ModifyInstanceAttributeAsync(ModifyInstanceAttributeRequest request, CancellationToken cancellationToken)
 	{
 		if (!_instances.ContainsKey(request.InstanceId))
 		{
@@ -157,11 +174,12 @@ public class FakeAmazonEc2
 		}
 
 		_instances[request.InstanceId].InstanceType = request.InstanceType;
+		await OnModifyInstanceAttribute(request, cancellationToken);
 		
-		return Task.FromResult(new ModifyInstanceAttributeResponse()
+		return new ModifyInstanceAttributeResponse()
 		{
 			ContentLength = 0,
 			HttpStatusCode = HttpStatusCode.OK
-		});
+		};
 	}
 }
\ No newline at end of file
