Horde: Allow configuring OIDC tokens to be set via the agent type.

diff --git a/Engine/Source/Programs/Horde/Horde.Build/Server/RpcService.cs b/Engine/Source/Programs/Horde/Horde.Build/Server/RpcService.cs
--- a/Engine/Source/Programs/Horde/Horde.Build/Server/RpcService.cs
+++ b/Engine/Source/Programs/Horde/Horde.Build/Server/RpcService.cs
@@ -692,7 +692,15 @@ namespace Horde.Build.Server
 				IStream? stream = await _streamService.GetStreamAsync(job.StreamId);
 				if (stream != null)
 				{
-					foreach (TokenConfig tokenConfig in stream.Config.Tokens)
+					List<TokenConfig> tokenConfigs = new List<TokenConfig>(stream.Config.Tokens);
+
+					INodeGroup group = graph.Groups[batch.GroupIdx];
+					if (stream.Config.AgentTypes.TryGetValue(group.AgentType, out AgentConfig? agentConfig) && agentConfig.Tokens != null)
+					{
+						tokenConfigs.AddRange(agentConfig.Tokens);
+					}
+
+					foreach (TokenConfig tokenConfig in tokenConfigs)
 					{
 						string? value = await AllocateTokenAsync(tokenConfig, context.CancellationToken);
 						if (value == null)
diff --git a/Engine/Source/Programs/Horde/Horde.Build/Streams/StreamConfig.cs b/Engine/Source/Programs/Horde/Horde.Build/Streams/StreamConfig.cs
--- a/Engine/Source/Programs/Horde/Horde.Build/Streams/StreamConfig.cs
+++ b/Engine/Source/Programs/Horde/Horde.Build/Streams/StreamConfig.cs
@@ -429,6 +429,11 @@ namespace Horde.Build.Streams
 		/// </summary>
 		public Dictionary<string, string>? Environment { get; set; }
 
+		/// <summary>
+		/// Tokens to allocate for this agent type
+		/// </summary>
+		public List<TokenConfig>? Tokens { get; set; }
+
 		/// <summary>
 		/// Creates an API response object from this stream
 		/// </summary>
