Horde: Lower fleet service tick interval

diff --git a/Engine/Source/Programs/Horde/Horde.Build/Agents/Fleet/FleetService.cs b/Engine/Source/Programs/Horde/Horde.Build/Agents/Fleet/FleetService.cs
--- a/Engine/Source/Programs/Horde/Horde.Build/Agents/Fleet/FleetService.cs
+++ b/Engine/Source/Programs/Horde/Horde.Build/Agents/Fleet/FleetService.cs
@@ -86,12 +86,11 @@ namespace Horde.Build.Agents.Fleet
 			_clock = clock;
 			_cache = cache;
 			_logger = logger;
-			_ticker = clock.AddSharedTicker<FleetService>(TimeSpan.FromMinutes(3.0), TickLeaderAsync, _logger);
+			_ticker = clock.AddSharedTicker<FleetService>(TimeSpan.FromSeconds(30), TickLeaderAsync, _logger);
 			_tickerHighFrequency = clock.AddSharedTicker("FleetService.TickHighFrequency", TimeSpan.FromSeconds(30), TickHighFrequencyAsync, _logger);
 			_settings = settings;
 			_defaultScaleOutCooldown = TimeSpan.FromSeconds(settings.Value.AgentPoolScaleOutCooldownSeconds);
 			_defaultScaleInCooldown = TimeSpan.FromSeconds(settings.Value.AgentPoolScaleInCooldownSeconds);
-
 		}
 
 		/// <inheritdoc/>
