Horde: Add more logging to workspaces getting conformed

diff --git a/Engine/Source/Programs/Horde/Horde.Agent/Execution/PerforceExecutor.cs b/Engine/Source/Programs/Horde/Horde.Agent/Execution/PerforceExecutor.cs
--- a/Engine/Source/Programs/Horde/Horde.Agent/Execution/PerforceExecutor.cs
+++ b/Engine/Source/Programs/Horde/Horde.Agent/Execution/PerforceExecutor.cs
@@ -367,9 +367,13 @@ namespace Horde.Agent.Execution
 				List<Func<Task>> syncFuncs = new List<Func<Task>>();
 				foreach (IGrouping<DirectoryReference, WorkspaceInfo> workspaceGroup in workspaces.GroupBy(x => x.MetadataDir).OrderBy(x => x.Key.FullName))
 				{
+					logger.LogInformation("Queuing workspaces for sync/populate:");
+					
 					List<PopulateRequest> populateRequests = new List<PopulateRequest>();
 					foreach (WorkspaceInfo workspace in workspaceGroup)
 					{
+						logger.LogInformation("  Stream={StreamName} RemoveUntrackedFiles={RemoveUntrackedFiles} MetadataDir={MetadataDir} WorkspaceDir={WorkspaceDir} ClientName={ClientName}",
+							workspace.StreamName, workspace.RemoveUntrackedFiles, workspace.MetadataDir.FullName, workspace.WorkspaceDir.FullName, workspace.ClientName);
 						IPerforceConnection perforceClient = await workspace.PerforceClient.WithClientAsync(workspace.ClientName);
 						perforceConnections.Add(perforceClient);
 						populateRequests.Add(new PopulateRequest(perforceClient, workspace.StreamName, workspace.View));
