Horde: Only use PerforceServiceCache when commit replication is enabled.

diff --git a/Engine/Source/Programs/Horde/Horde.Build/Perforce/PerforceServiceCache.cs b/Engine/Source/Programs/Horde/Horde.Build/Perforce/PerforceServiceCache.cs
--- a/Engine/Source/Programs/Horde/Horde.Build/Perforce/PerforceServiceCache.cs
+++ b/Engine/Source/Programs/Horde/Horde.Build/Perforce/PerforceServiceCache.cs
@@ -206,19 +206,13 @@ namespace Horde.Build.Perforce
 		/// <inheritdoc/>
 		public async Task StartAsync(CancellationToken cancellationToken)
 		{
-			if (_settings.Commits.ReplicateMetadata)
-			{
-				await _updateCommitsTicker.StartAsync();
-			}
+			await _updateCommitsTicker.StartAsync();
 		}
 
 		/// <inheritdoc/>
 		public async Task StopAsync(CancellationToken cancellationToken)
 		{
-			if (_settings.Commits.ReplicateMetadata)
-			{
-				await _updateCommitsTicker.StopAsync();
-			}
+			await _updateCommitsTicker.StopAsync();
 		}
 
 		#region Commit updates
diff --git a/Engine/Source/Programs/Horde/Horde.Build/Startup.cs b/Engine/Source/Programs/Horde/Horde.Build/Startup.cs
--- a/Engine/Source/Programs/Horde/Horde.Build/Startup.cs
+++ b/Engine/Source/Programs/Horde/Horde.Build/Startup.cs
@@ -410,8 +410,17 @@ namespace Horde.Build
 			services.AddSingleton<ILogFileService, LogFileService>();
 			services.AddSingleton<LogTailService>();
 			services.AddSingleton<INotificationService, NotificationService>();
-			services.AddSingleton<PerforceServiceCache>();
-			services.AddSingleton<IPerforceService>(sp => sp.GetRequiredService<PerforceServiceCache>());
+
+			if (settings.Commits.ReplicateMetadata)
+			{
+				services.AddSingleton<PerforceServiceCache>();
+				services.AddSingleton<IPerforceService>(sp => sp.GetRequiredService<PerforceServiceCache>());
+			}
+			else
+			{
+				services.AddSingleton<PerforceService>();
+				services.AddSingleton<IPerforceService>(sp => sp.GetRequiredService<PerforceService>());
+			}
 
 			services.AddSingleton<PerforceLoadBalancer>();
 			services.AddSingleton<PoolService>();
@@ -608,7 +617,12 @@ namespace Horde.Build
 				services.AddHostedService(provider => provider.GetRequiredService<FleetService>());
 				
 				services.AddHostedService(provider => provider.GetRequiredService<AgentService>());
-				services.AddHostedService(provider => provider.GetRequiredService<PerforceServiceCache>());
+
+				if (settings.Commits.ReplicateMetadata)
+				{
+					services.AddHostedService(provider => provider.GetRequiredService<PerforceServiceCache>());
+				}
+
 				services.AddHostedService(provider => provider.GetRequiredService<ConsistencyService>());
 				services.AddHostedService(provider => provider.GetRequiredService<IssueService>());
 				services.AddHostedService<IssueReportService>();
