Horde: Add pub/sub model for commits in a stream. ICommitCollection exposes a SubscribeAsync method that uses a redis channel to post updates of newly committed changes.

diff --git a/Engine/Source/Programs/Horde/Horde.Build/Perforce/ICommitCollection.cs b/Engine/Source/Programs/Horde/Horde.Build/Perforce/ICommitCollection.cs
--- a/Engine/Source/Programs/Horde/Horde.Build/Perforce/ICommitCollection.cs
+++ b/Engine/Source/Programs/Horde/Horde.Build/Perforce/ICommitCollection.cs
@@ -41,6 +41,15 @@ namespace Horde.Build.Perforce
 		/// <param name="cancellationToken">Cancellation token for the operation</param>
 		/// <returns>Changelist information</returns>
 		IAsyncEnumerable<ICommit> FindAsync(int? minChange = null, int? maxChange = null, int? maxResults = null, IReadOnlyList<CommitTag>? tags = null, CancellationToken cancellationToken = default);
+
+		/// <summary>
+		/// Subscribes to changes from this commit source
+		/// </summary>
+		/// <param name="minChange">Minimum changelist number (exclusive)</param>
+		/// <param name="tags">Tags for the commit to return</param>
+		/// <param name="cancellationToken">Cancellation token for the operation</param>
+		/// <returns>New change information</returns>
+		IAsyncEnumerable<ICommit> SubscribeAsync(int minChange, IReadOnlyList<CommitTag>? tags = null, CancellationToken cancellationToken = default);
 	}
 
 	/// <summary>
diff --git a/Engine/Source/Programs/Horde/Horde.Build/Perforce/PerforceService.cs b/Engine/Source/Programs/Horde/Horde.Build/Perforce/PerforceService.cs
--- a/Engine/Source/Programs/Horde/Horde.Build/Perforce/PerforceService.cs
+++ b/Engine/Source/Programs/Horde/Horde.Build/Perforce/PerforceService.cs
@@ -1250,10 +1250,24 @@ namespace Horde.Build.Perforce
 				}
 			}
 
+			/// <inheritdoc/>
 			public virtual Task<ICommit> GetAsync(int changeNumber, CancellationToken cancellationToken = default)
 			{
 				return _perforceService.GetChangeDetailsAsync(_stream, changeNumber, cancellationToken);
 			}
+
+			/// <inheritdoc/>
+			public virtual async IAsyncEnumerable<ICommit> SubscribeAsync(int minChange, IReadOnlyList<CommitTag>? tags = null, [EnumeratorCancellation] CancellationToken cancellationToken = default)
+			{
+				for (; ; )
+				{
+					await foreach (ICommit commit in FindAsync(minChange + 1, null, 10, tags, cancellationToken))
+					{
+						yield return commit;
+					}
+					await Task.Delay(TimeSpan.FromSeconds(10.0), cancellationToken);
+				}
+			}
 		}
 
 		/// <inheritdoc/>
diff --git a/Engine/Source/Programs/Horde/Horde.Build/Perforce/PerforceServiceCache.cs b/Engine/Source/Programs/Horde/Horde.Build/Perforce/PerforceServiceCache.cs
--- a/Engine/Source/Programs/Horde/Horde.Build/Perforce/PerforceServiceCache.cs
+++ b/Engine/Source/Programs/Horde/Horde.Build/Perforce/PerforceServiceCache.cs
@@ -1,14 +1,17 @@
 // Copyright Epic Games, Inc. All Rights Reserved.
 
 using System;
+using System.Collections.Concurrent;
 using System.Collections.Generic;
 using System.Diagnostics;
+using System.IO;
 using System.Linq;
 using System.Runtime.CompilerServices;
 using System.Threading;
 using System.Threading.Tasks;
 using EpicGames.Core;
 using EpicGames.Perforce;
+using EpicGames.Redis;
 using Horde.Build.Server;
 using Horde.Build.Streams;
 using Horde.Build.Users;
@@ -21,6 +24,7 @@ using MongoDB.Bson;
 using MongoDB.Bson.Serialization.Attributes;
 using MongoDB.Bson.Serialization.Options;
 using MongoDB.Driver;
+using StackExchange.Redis;
 
 namespace Horde.Build.Perforce
 {
@@ -167,22 +171,25 @@ namespace Horde.Build.Perforce
 		}
 
 		readonly MongoService _mongoService;
+		readonly RedisService _redisService;
 		readonly IDowntimeService _downtimeService;
-		readonly ServerSettings _settings;
 		readonly IMongoCollection<CachedCommitDoc> _commits;
 		readonly IStreamCollection _streamCollection;
 		readonly ILogger _logger;
+
+		static readonly RedisChannel<StreamId> s_commitUpdateChannel = new RedisChannel<StreamId>("commit-update");
+
 		readonly ITicker _updateCommitsTicker;
 
 		/// <summary>
 		/// Constructor
 		/// </summary>
-		public PerforceServiceCache(PerforceLoadBalancer loadBalancer, MongoService mongoService, IDowntimeService downtimeService, GlobalsService globalsService, IUserCollection userCollection, IStreamCollection streamCollection, IClock clock, IOptions<ServerSettings> settings, ILogger<PerforceService> logger)
+		public PerforceServiceCache(PerforceLoadBalancer loadBalancer, MongoService mongoService, RedisService redisService, IDowntimeService downtimeService, GlobalsService globalsService, IUserCollection userCollection, IStreamCollection streamCollection, IClock clock, IOptions<ServerSettings> settings, ILogger<PerforceService> logger)
 			: base(loadBalancer, globalsService, userCollection, settings, logger)
 		{
 			_mongoService = mongoService;
+			_redisService = redisService;
 			_downtimeService = downtimeService;
-			_settings = settings.Value;
 
 			List<MongoIndex<CachedCommitDoc>> indexes = new List<MongoIndex<CachedCommitDoc>>();
 			indexes.Add(MongoIndex.Create<CachedCommitDoc>(keys => keys.Ascending(x => x.StreamId).Descending(x => x.Number), true));
@@ -419,6 +426,8 @@ namespace Horde.Build.Perforce
 
 							CachedCommitDoc commitDoc = new CachedCommitDoc(commit, commitTags);
 							await AddCachedCommitAsync(commitDoc, cancellationToken);
+
+							await _redisService.PublishAsync(s_commitUpdateChannel, streamInfo.Stream.Id, CommandFlags.FireAndForget);
 						}
 					}
 				}
@@ -605,6 +614,60 @@ namespace Horde.Build.Perforce
 				}
 				return await base.GetAsync(changeNumber, cancellationToken);
 			}
+
+			/// <inheritdoc/>
+			public override async IAsyncEnumerable<ICommit> SubscribeAsync(int minChange, IReadOnlyList<CommitTag>? tags = null, [EnumeratorCancellation] CancellationToken cancellationToken = default)
+			{
+				AsyncEvent updateEvent = new AsyncEvent();
+
+				void OnUpdate(StreamId streamId)
+				{
+					if (streamId == _stream.Id)
+					{
+						updateEvent.Set();
+					}
+				}
+
+				RedisChannelSubscription<StreamId>? subscription = null;
+				try
+				{
+					for (; ; )
+					{
+						cancellationToken.ThrowIfCancellationRequested();
+
+						Task task = updateEvent.Task;
+
+						int numResults = 10;
+						await foreach (ICommit commit in FindAsync(minChange + 1, null, numResults, tags, cancellationToken))
+						{
+							yield return commit;
+							minChange = commit.Number;
+							numResults--;
+						}
+
+						if (numResults == 0)
+						{
+							// Query again; received the max requested number of changes.
+							continue;
+						}
+
+						if (subscription == null)
+						{
+							subscription = await _owner._redisService.SubscribeAsync(s_commitUpdateChannel, OnUpdate);
+							continue;
+						}
+
+						await task;
+					}
+				}
+				finally
+				{
+					if (subscription != null)
+					{
+						await subscription.DisposeAsync();
+					}
+				}
+			}
 		}
 
 		/// <inheritdoc/>
diff --git a/Engine/Source/Programs/Horde/Horde.Build/Server/RedisService.cs b/Engine/Source/Programs/Horde/Horde.Build/Server/RedisService.cs
--- a/Engine/Source/Programs/Horde/Horde.Build/Server/RedisService.cs
+++ b/Engine/Source/Programs/Horde/Horde.Build/Server/RedisService.cs
@@ -221,5 +221,33 @@ namespace Horde.Build.Server
 			}
 			redisLogger.LogInformation("Exit code {ExitCode}", _redisProcess.ExitCode);
 		}
+
+		/// <summary>
+		/// Publish a message to a channel
+		/// </summary>
+		/// <typeparam name="T">Type of elements sent over the channel</typeparam>
+		/// <param name="channel">Channel to post to</param>
+		/// <param name="message">Message to post to the channel</param>
+		/// <param name="flags">Flags for the request</param>
+		public Task PublishAsync<T>(RedisChannel<T> channel, T message, CommandFlags flags = CommandFlags.None)
+		{
+			return ConnectionPool.GetDatabase().PublishAsync(channel, message, flags);
+		}
+
+		/// <inheritdoc cref="SubscribeAsync{T}(RedisChannel{T}, Action{RedisChannel{T}, T})"/>
+		public Task<RedisChannelSubscription<T>> SubscribeAsync<T>(RedisChannel<T> channel, Action<T> callback) => SubscribeAsync(channel, (ch, x) => callback(x));
+
+		/// <summary>
+		/// Subscribe to notifications on a channel
+		/// </summary>
+		/// <typeparam name="T">Type of elements sent over the channel</typeparam>
+		/// <param name="channel">Channel to monitor</param>
+		/// <param name="callback">Callback for new events</param>
+		/// <returns>Subscription object</returns>
+		public async Task<RedisChannelSubscription<T>> SubscribeAsync<T>(RedisChannel<T> channel, Action<RedisChannel<T>, T> callback)
+		{
+			ISubscriber subscriber = ConnectionPool.GetConnection().GetSubscriber();
+			return await subscriber.SubscribeAsync(channel, callback);
+		}
 	}
 }
