Horde: Remove code for last version of replication service.

diff --git a/Engine/Source/Programs/Horde/Horde.Build/Commands/Bundles/ExtractCommand.cs b/Engine/Source/Programs/Horde/Horde.Build/Commands/Bundles/ExtractCommand.cs
--- a/Engine/Source/Programs/Horde/Horde.Build/Commands/Bundles/ExtractCommand.cs
+++ b/Engine/Source/Programs/Horde/Horde.Build/Commands/Bundles/ExtractCommand.cs
@@ -1,5 +1,6 @@
 // Copyright Epic Games, Inc. All Rights Reserved.
 
+using System;
 using System.Threading;
 using System.Threading.Tasks;
 using EpicGames.Core;
@@ -31,18 +32,14 @@ namespace Horde.Build.Commands.Bundles
 			_loggerProvider = loggerProvider;
 		}
 
-		public override async Task<int> ExecuteAsync(ILogger logger)
+		public override Task<int> ExecuteAsync(ILogger logger)
 		{
 			using ServiceProvider serviceProvider = Startup.CreateServiceProvider(_configuration, _loggerProvider);
 
 			IStorageClient store = serviceProvider.GetRequiredService<IStorageClient<ReplicationService>>();
-			TreeReader reader = new TreeReader(store, serviceProvider.GetRequiredService<IMemoryCache>(), serviceProvider.GetRequiredService<ILogger<ExtractCommand>>()); 
+			TreeReader reader = new TreeReader(store, serviceProvider.GetRequiredService<IMemoryCache>(), serviceProvider.GetRequiredService<ILogger<ExtractCommand>>());
 
-			ReplicationNode node = await reader.ReadNodeAsync<ReplicationNode>(RefName);
-			DirectoryNode root = await node.Contents.ExpandAsync(reader);
-			await root.CopyToDirectoryAsync(reader, OutputDir.ToDirectoryInfo(), logger, CancellationToken.None);
-
-			return 0;
+			throw new NotImplementedException();
 		}
 	}
 }
diff --git a/Engine/Source/Programs/Horde/Horde.Build/Commands/Bundles/PerforceCommand.cs b/Engine/Source/Programs/Horde/Horde.Build/Commands/Bundles/PerforceCommand.cs
--- a/Engine/Source/Programs/Horde/Horde.Build/Commands/Bundles/PerforceCommand.cs
+++ b/Engine/Source/Programs/Horde/Horde.Build/Commands/Bundles/PerforceCommand.cs
@@ -78,34 +78,7 @@ namespace Horde.Build.Commands.Bundles
 				throw new FatalErrorException($"Stream '{stream}' not found");
 			}
 
-			Dictionary<IStream, int> streamToFirstChange = new Dictionary<IStream, int>();
-			streamToFirstChange[stream] = Change;
-
-			ReplicationNode baseContents;
-			if (BaseChange == 0)
-			{
-				baseContents = new ReplicationNode(new DirectoryNode());
-			}
-			else
-			{
-				baseContents = await ReplicationService.ReadCommitTreeAsync(reader, stream, BaseChange, Filter, RevisionsOnly, CancellationToken.None);
-			}
-
-			await foreach (ICommit commit in commitService.GetCollection(stream).FindAsync(Change, null, Count))
-			{
-				string briefSummary = commit.Description.Replace('\n', ' ');
-				logger.LogInformation("");
-				logger.LogInformation("Commit {Change} by {AuthorId}: {Summary}", commit.Number, commit.AuthorId, briefSummary.Substring(0, Math.Min(50, briefSummary.Length)));
-				logger.LogInformation("Base path: {BasePath}", commit.BasePath);
-
-				if (Content)
-				{
-					baseContents = await replicationService.WriteCommitTreeAsync(storage, reader, stream, commit.Number, BaseChange, baseContents, Filter, RevisionsOnly, CancellationToken.None);
-					BaseChange = commit.Number;
-				}
-			}
-
-			return 0;
+			throw new NotImplementedException();
 		}
 	}
 }
diff --git a/Engine/Source/Programs/Horde/Horde.Build/Perforce/ReplicationService.cs b/Engine/Source/Programs/Horde/Horde.Build/Perforce/ReplicationService.cs
--- a/Engine/Source/Programs/Horde/Horde.Build/Perforce/ReplicationService.cs
+++ b/Engine/Source/Programs/Horde/Horde.Build/Perforce/ReplicationService.cs
@@ -1,30 +1,13 @@
 // Copyright Epic Games, Inc. All Rights Reserved.
 
 using System;
-using System.Collections.Generic;
-using System.Diagnostics;
-using System.IO;
-using System.Linq;
-using System.Reflection.Metadata;
-using System.Security.Cryptography;
-using System.Text;
 using System.Threading;
 using System.Threading.Tasks;
-using EpicGames.Core;
-using EpicGames.Horde.Storage;
-using EpicGames.Horde.Storage.Git;
-using EpicGames.Horde.Storage.Nodes;
-using EpicGames.Perforce;
-using EpicGames.Redis;
 using Horde.Build.Server;
-using Horde.Build.Storage;
 using Horde.Build.Streams;
 using Horde.Build.Utilities;
-using Microsoft.Extensions.Caching.Memory;
 using Microsoft.Extensions.Hosting;
 using Microsoft.Extensions.Logging;
-using Microsoft.Extensions.Options;
-using StackExchange.Redis;
 
 namespace Horde.Build.Perforce
 {
@@ -40,947 +23,27 @@ namespace Horde.Build.Perforce
 		}
 	}
 
-	/// <summary>
-	/// Root node for a commit snapshot
-	/// </summary>
-	[TreeNode("{8C874966-2E89-4273-A4AB-AC9F5491C86D}")]
-	public class ReplicationNode : TreeNode
-	{
-		/// <summary>
-		/// Paths that have been synced. Empty once complete.
-		/// </summary>
-		public List<Utf8String> Paths { get; }
-
-		/// <summary>
-		/// Contents of this snapshot
-		/// </summary>
-		public TreeNodeRef<DirectoryNode> Contents { get; }
-
-		/// <summary>
-		/// Constructor
-		/// </summary>
-		public ReplicationNode()
-			: this(new DirectoryNode())
-		{
-		}
-
-		/// <summary>
-		/// Constructor
-		/// </summary>
-		public ReplicationNode(DirectoryNode contents)
-		{
-			Paths = new List<Utf8String>();
-			Contents = new TreeNodeRef<DirectoryNode>(contents);
-		}
-
-		/// <summary>
-		/// Deserializing constructor
-		/// </summary>
-		/// <param name="reader"></param>
-		public ReplicationNode(ITreeNodeReader reader)
-		{
-			Paths = reader.ReadList(() => reader.ReadUtf8String());
-			Contents = reader.ReadRef<DirectoryNode>();
-		}
-
-		/// <inheritdoc/>
-		public override void Serialize(ITreeNodeWriter writer)
-		{
-			writer.WriteList(Paths, x => writer.WriteUtf8String(x));
-			writer.WriteRef(Contents);
-		}
-
-		/// <inheritdoc/>
-		public override IEnumerable<TreeNodeRef> EnumerateRefs()
-		{
-			yield return Contents;
-		}
-	}
-
 	/// <summary>
 	/// Service which replicates content from Perforce
 	/// </summary>
-	sealed class ReplicationService : IHostedService, IDisposable
+	sealed class ReplicationService : IHostedService
 	{
-		/// <summary>
-		/// Metadata about a stream that needs to have commits mirrored
-		/// </summary>
-		class StreamInfo
-		{
-			public IStream Stream { get; set; }
-			public ViewMap View { get; }
-
-			public StreamInfo(IStream stream, ViewMap view)
-			{
-				Stream = stream;
-				View = view;
-			}
-		}
-
-		readonly ServerSettings _settings;
-
-		// Redis
-		readonly RedisConnectionPool _redisConnectionPool;
-
-		// Collections
-//		readonly CommitService _commitService;
-//		readonly IAsyncDisposable _commitListener;
-
-		readonly IStreamCollection _streamCollection;
-		readonly IPerforceService _perforceService;
-		readonly StorageService _storageService;
-		readonly IMemoryCache _cache;
+		readonly RedisService _redisService;
 		readonly ILogger _logger;
 
-		const int MaxBackgroundTasks = 2;
-
-		const double ReservationTimeSeconds = 60.0;
-		const double ExtendReservationAfterSeconds = 40.0;
-
-		RedisKey RedisBaseKey { get; } = new RedisKey("commits/replication/");
-		RedisChannel<StreamId> RedisUpdateChannel { get; } = new RedisChannel<StreamId>("commits/replication/streams");
-		bool _stopping;
-		readonly AsyncEvent _updateStreamsEvent = new AsyncEvent();
-		readonly RedisSet<StreamId> _redisDirtyStreams;
-		readonly RedisSortedSet<StreamId> _redisReservations;
-		RedisChannelSubscription<StreamId>? _redisUpdateSubscription;
-		Task? _streamUpdateTask;
-
 		/// <summary>
 		/// Constructor
 		/// </summary>
-		public ReplicationService(RedisService redisService, /* CommitService commitService, */IStreamCollection streamCollection, IPerforceService perforceService, StorageService storageService, IMemoryCache cache, IOptions<ServerSettings> settings, ILogger<ReplicationService> logger)
+		public ReplicationService(RedisService redisService, ILogger<ReplicationService> logger)
 		{
-			_settings = settings.Value;
-
-			_redisConnectionPool = redisService.ConnectionPool;
-			_redisDirtyStreams = new RedisSet<StreamId>(_redisConnectionPool, RedisBaseKey.Append("streams"));
-			_redisReservations = new RedisSortedSet<StreamId>(_redisConnectionPool, RedisBaseKey.Append("reservations"));
-
-//			_commitService = commitService;
-//			_commitListener = _commitService.AddListener(OnCommitAdded);
-
-			_streamCollection = streamCollection;
-			_perforceService = perforceService;
-			_storageService = storageService;
-			_cache = cache;
+			_redisService = redisService;
 			_logger = logger;
 		}
 
 		/// <inheritdoc/>
-		public void Dispose()
-		{
-//			_commitListener.DisposeAsync().AsTask().Wait();
-		}
+		public Task StartAsync(CancellationToken cancellationToken) => Task.CompletedTask;
 
 		/// <inheritdoc/>
-		public async Task StartAsync(CancellationToken cancellationToken)
-		{
-			if (_settings.Commits.ReplicateContent)
-			{
-				_stopping = false;
-				_updateStreamsEvent.Reset();
-				_redisUpdateSubscription = await _redisConnectionPool.GetDatabase().Multiplexer.GetSubscriber().SubscribeAsync(RedisUpdateChannel, (_, _) => _updateStreamsEvent.Pulse());
-				_streamUpdateTask = Task.Run(() => UpdateContentAsync(), cancellationToken);
-			}
-		}
-
-		/// <inheritdoc/>
-		public async Task StopAsync(CancellationToken cancellationToken)
-		{
-			if (_settings.Commits.ReplicateContent)
-			{
-				_stopping = true;
-				_updateStreamsEvent.Latch();
-				if (_redisUpdateSubscription != null)
-				{
-					await _redisUpdateSubscription.DisposeAsync();
-					_redisUpdateSubscription = null;
-				}
-				await _streamUpdateTask!;
-//				await _commitListener.DisposeAsync();
-			}
-		}
-
-		RedisList<int> RedisStreamChanges(StreamId streamId) => new RedisList<int>(_redisConnectionPool, RedisBaseKey.Append($"stream/{streamId}/changes"));
-
-		async Task OnCommitAdded(ICommit commit)
-		{
-			// Post the new commit key
-			RedisList<int> streamCommitsKey = RedisStreamChanges(commit.StreamId);
-			await streamCommitsKey.RightPushAsync(commit.Number);
-
-			// Signal to any listeners that we have new data to process
-			await _redisDirtyStreams.AddAsync(commit.StreamId);
-			await _redisConnectionPool.GetDatabase().PublishAsync(RedisUpdateChannel, commit.StreamId);
-		}
-
-		async Task UpdateContentAsync()
-		{
-			List<(StreamId, Task)> backgroundTasks = new List<(StreamId, Task)>(MaxBackgroundTasks);
-			while (!_stopping || backgroundTasks.Count > 0)
-			{
-				try
-				{
-					await UpdateContentInternalAsync(backgroundTasks);
-				}
-				catch (Exception ex)
-				{
-					_logger.LogError(ex, "Exception when updating commit service");
-				}
-			}
-		}
-
-		async Task UpdateContentInternalAsync(List<(StreamId, Task)> backgroundTasks)
-		{
-			// Remove any complete background tasks
-			for (int idx = backgroundTasks.Count - 1; idx >= 0; idx--)
-			{
-				(_, Task backgroundTask) = backgroundTasks[idx];
-				if (backgroundTask.IsCompleted)
-				{
-					if (backgroundTask.Exception != null)
-					{
-						_logger.LogError(backgroundTask.Exception, "Update background task faulted");
-					}
-					backgroundTasks.RemoveAt(idx);
-				}
-			}
-
-			// Create a list of events to wait for
-			List<Task> waitTasks = new List<Task>(backgroundTasks.Select(x => x.Item2));
-			if (!_stopping)
-			{
-				// If we have spare slots for executing background tasks, check if there are any dirty streams
-				Task newStreamTask = _updateStreamsEvent.Task;
-				if (backgroundTasks.Count < MaxBackgroundTasks)
-				{
-					// Expire any reservations that are no longer valid
-					DateTime utcNow = DateTime.UtcNow;
-					await _redisReservations.RemoveRangeByScoreAsync(Double.NegativeInfinity, utcNow.Ticks);
-
-					// Find the streams that we can wait for
-					HashSet<StreamId> checkStreams = new HashSet<StreamId>(await _redisDirtyStreams.MembersAsync());
-					checkStreams.ExceptWith(backgroundTasks.Select(x => x.Item1));
-
-					if (checkStreams.Count > 0)
-					{
-						// Try to start new background tasks
-						double newScore = (utcNow + TimeSpan.FromSeconds(ReservationTimeSeconds)).Ticks;
-						foreach (StreamId checkStream in checkStreams)
-						{
-							if (await _redisReservations.AddAsync(checkStream, newScore, When.NotExists))
-							{
-								Task newTask = Task.Run(() => UpdateStreamContentAsync(checkStream, CancellationToken.None));
-								backgroundTasks.Add((checkStream, newTask));
-
-								if (backgroundTasks.Count == MaxBackgroundTasks)
-								{
-									break;
-								}
-							}
-						}
-						checkStreams.ExceptWith(backgroundTasks.Select(x => x.Item1));
-
-						// If we still have spare tasks, check how long we should wait for a reservation to expire
-						if (backgroundTasks.Count < MaxBackgroundTasks)
-						{
-							long waitTime = Int64.MaxValue;
-							await foreach (SortedSetEntry<StreamId> entry in _redisReservations.ScanAsync())
-							{
-								if (checkStreams.Contains(entry.Element))
-								{
-									waitTime = Math.Min(waitTime, (long)entry.Score - utcNow.Ticks);
-								}
-							}
-							if (waitTime < Int32.MaxValue)
-							{
-								waitTasks.Add(Task.Delay(new TimeSpan(Math.Max(waitTime, 0))));
-							}
-						}
-					}
-				}
-
-				// If we still have bandwidth to process more tasks, wait for new streams to become available
-				if (backgroundTasks.Count < MaxBackgroundTasks)
-				{
-					waitTasks.Add(newStreamTask);
-				}
-			}
-
-			// Wait until any task has completed
-			if (waitTasks.Count > 0)
-			{
-				await Task.WhenAny(waitTasks);
-			}
-		}
-
-		async Task UpdateStreamContentAsync(StreamId streamId, CancellationToken cancellationToken)
-		{
-			Stopwatch timer = Stopwatch.StartNew();
-
-			RedisList<int> streamChanges = RedisStreamChanges(streamId);
-			for (; ; )
-			{
-				IStream? stream = await _streamCollection.GetAsync(streamId);
-				if (stream == null || stream.Config.ReplicationMode == ContentReplicationMode.None)
-				{
-					// Remove all but the last item
-					await streamChanges.TrimAsync(0, -2);
-				}
-				else
-				{
-					IStorageClient storage = await _storageService.GetClientAsync(Namespace.Perforce, cancellationToken);
-					TreeReader reader = new TreeReader(storage, _cache, _logger);
-
-					// Update the stream, updating the reservation every 30 seconds
-					Task internalTask = Task.Run(() => UpdateStreamContentInternalAsync(storage, reader, stream, streamChanges, cancellationToken), cancellationToken);
-					while (!internalTask.IsCompleted)
-					{
-						Task delayTask = Task.Delay(TimeSpan.FromSeconds(ExtendReservationAfterSeconds), cancellationToken);
-						if (await Task.WhenAny(internalTask, delayTask) == delayTask)
-						{
-							DateTime newTime = DateTime.UtcNow + TimeSpan.FromSeconds(ReservationTimeSeconds);
-							await _redisReservations.AddAsync(streamId, newTime.Ticks);
-							_logger.LogInformation("Extending reservation for content update of {StreamId} (elapsed: {Time}s)", streamId, (int)timer.Elapsed.TotalSeconds);
-						}
-					}
-
-					// Log any error during the update
-					if (internalTask.IsFaulted)
-					{
-						_logger.LogError(internalTask.Exception, "Exception while updating stream content for {StreamId}", streamId);
-					}
-					else
-					{
-						_logger.LogInformation("Finished update for {StreamId}", streamId);
-					}
-				}
-
-				// Remove this stream from the dirty list if it's empty
-				ITransaction transaction = _redisConnectionPool.GetDatabase().CreateTransaction();
-				transaction.AddCondition(Condition.ListLengthLessThan(streamChanges.Key, 2));
-				_ = transaction.With(_redisDirtyStreams).RemoveAsync(streamId, flags: CommandFlags.FireAndForget);
-				if (await transaction.ExecuteAsync())
-				{
-					break;
-				}
-			}
-			await _redisReservations.RemoveAsync(streamId);
-		}
-
-		async Task UpdateStreamContentInternalAsync(IStorageClient storage, TreeReader reader, IStream stream, RedisList<int> changes, CancellationToken cancellationToken)
-		{
-			int prevChange = 0;
-			ReplicationNode prevContents = new ReplicationNode();
-			for (; ; )
-			{
-				// Get the first two changes to be mirrored. The first one should already exist, unless it's the start of replication for this stream.
-				int[] values = await changes.RangeAsync(0, 1);
-				if (values.Length == 0)
-				{
-					break;
-				}
-
-				// Get or add the tree for the first change.
-				if (prevChange != values[0])
-				{
-					RefName prevRefName = GetRefName(stream, values[0]);
-
-					ReplicationNode? contents = await reader.TryReadNodeAsync<ReplicationNode>(prevRefName, cancellationToken: cancellationToken);
-					if (contents == null)
-					{
-						_logger.LogInformation("No content for CL {Change}; creating full snapshot", values[0]);
-						prevContents = await WriteCommitTreeAsync(storage, reader, stream, values[0], 0, new ReplicationNode(), cancellationToken);
-					}
-					else
-					{
-						_logger.LogInformation("Reading existing commit tree for CL {Change} from ref {RefName}", values[0], prevRefName);
-						prevContents = contents;
-					}
-
-					prevChange = values[0];
-				}
-				else if (values.Length == 2)
-				{
-					// Perform a snapshot of the new change, then remove it from the list
-					prevContents = await WriteCommitTreeAsync(storage, reader, stream, values[1], prevChange, prevContents, cancellationToken);
-					prevChange = values[1];
-
-					// Remove the first item from the list
-					ITransaction transaction = _redisConnectionPool.GetDatabase().CreateTransaction();
-					transaction.AddCondition(Condition.ListIndexEqual(changes.Key, 0, values[0]));
-					transaction.AddCondition(Condition.ListIndexEqual(changes.Key, 1, values[1]));
-					_ = transaction.With(changes).LeftPopAsync(flags: CommandFlags.FireAndForget);
-					await transaction.ExecuteAsync();
-				}
-				else
-				{
-					// Nothing to do
-					break;
-				}
-			}
-		}
-
-		/// <summary>
-		/// Gets the ref for a particular change
-		/// </summary>
-		/// <param name="streamId"></param>
-		/// <param name="change"></param>
-		/// <param name="filter"></param>
-		/// <param name="revisionsOnly"></param>
-		/// <returns></returns>
-		static RefName GetRefName(StreamId streamId, int change, string? filter, bool revisionsOnly)
-		{
-			StringBuilder builder = new StringBuilder($"v5/{streamId}/{change}");
-			if (filter != null)
-			{
-				builder.Append($"-flt_{IoHash.Compute(Encoding.UTF8.GetBytes(filter)).ToString().Substring(0, 6)}");
-			}
-			if (revisionsOnly)
-			{
-				builder.Append("-revs");
-			}
-			return new RefName(builder.ToString());
-		}
-
-		/// <summary>
-		/// Gets the ref for a particular change
-		/// </summary>
-		/// <param name="stream"></param>
-		/// <param name="change"></param>
-		/// <returns></returns>
-		static RefName GetRefName(IStream stream, int change)
-		{
-			return GetRefName(stream.Id, change, stream.Config.ReplicationFilter, stream.Config.ReplicationMode == ContentReplicationMode.RevisionsOnly);
-		}
-
-		/// <summary>
-		/// Reads a tree from the given stream
-		/// </summary>
-		/// <param name="reader">Storage client instance</param>
-		/// <param name="stream"></param>
-		/// <param name="change"></param>
-		/// <param name="filter"></param>
-		/// <param name="revisionsOnly"></param>
-		/// <param name="cancellationToken"></param>
-		/// <returns></returns>
-		public static async Task<ReplicationNode> ReadCommitTreeAsync(TreeReader reader, IStream stream, int change, string? filter, bool revisionsOnly, CancellationToken cancellationToken)
-		{
-			RefName name = GetRefName(stream.Id, change, filter, revisionsOnly);
-			return await reader.ReadNodeAsync<ReplicationNode>(name, cancellationToken: cancellationToken);
-		}
-
-		/// <summary>
-		/// Replicates the contents of a stream to Horde storage, optionally using the given change as a starting point
-		/// </summary>
-		/// <param name="storage">Storage client instance</param>
-		/// <param name="reader"></param>
-		/// <param name="stream">The stream to replicate</param>
-		/// <param name="change">Commit to store the tree ref</param>
-		/// <param name="baseChange">The base change to update from</param>
-		/// <param name="baseContents">Initial contents of the tree at baseChange</param>
-		/// <param name="cancellationToken">Cancellation token for the operation</param>
-		/// <returns>Root tree object</returns>
-		public async Task<ReplicationNode> WriteCommitTreeAsync(IStorageClient storage, TreeReader reader, IStream stream, int change, int baseChange, ReplicationNode baseContents, CancellationToken cancellationToken)
-		{
-			bool revisionsOnly = stream.Config.ReplicationMode == ContentReplicationMode.RevisionsOnly;
-			return await WriteCommitTreeAsync(storage, reader, stream, change, baseChange, baseContents, stream.Config.ReplicationFilter, revisionsOnly, cancellationToken);
-		}
-
-		[DebuggerDisplay("{_path}")]
-		class DirectoryToSync
-		{
-			public readonly Utf8String Path;
-			public readonly Dictionary<Utf8String, long> FileNameToSize;
-			public long _size;
-
-			public DirectoryToSync(Utf8String path, Utf8StringComparer comparer)
-			{
-				Path = path;
-				FileNameToSize = new Dictionary<Utf8String, long>(comparer);
-			}
-		}
-
-		class FileWriter
-		{
-			readonly Utf8String _path;
-			FileNode _node;
-			readonly long _size;
-
-			public Utf8String Path => _path;
-			public FileNode Node => _node;
-			public long Size => _size;
-
-			public FileWriter(Utf8String path, long size)
-			{
-				_path = path;
-				_node = new LeafFileNode();
-				_size = size;
-			}
-
-			public async Task AppendAsync(ReadOnlyMemory<byte> data, ChunkingOptions options, TreeWriter writer, CancellationToken cancellationToken)
-			{
-				_node = await _node.AppendAsync(data, options, writer, cancellationToken);
-			}
-		}
-
-		/// <summary>
-		/// Replicates the contents of a stream to Horde storage, optionally using the given change as a starting point
-		/// </summary>
-		/// <param name="storage">Storage client</param>
-		/// <param name="reader">Reader for nodes</param>
-		/// <param name="stream">The stream to replicate</param>
-		/// <param name="change">Commit to store the tree ref</param>
-		/// <param name="baseChange">The base change to update from</param>
-		/// <param name="baseContents">Initial contents of the tree at baseChange</param>
-		/// <param name="filter"></param>
-		/// <param name="revisionsOnly"></param>
-		/// <param name="cancellationToken">Cancellation token for the operation</param>
-		/// <returns>Root tree object</returns>
-		public async Task<ReplicationNode> WriteCommitTreeAsync(IStorageClient storage, TreeReader reader, IStream stream, int change, int? baseChange, ReplicationNode baseContents, string? filter, bool revisionsOnly, CancellationToken cancellationToken)
-		{
-			// Create a client to replicate from this stream
-			ReplicationClient clientInfo = await FindOrAddReplicationClientAsync(stream);
-
-			// Connect to the server and flush the workspace
-			using IPerforceConnection perforce = await PerforceConnection.CreateAsync(clientInfo.Settings, _logger);
-			await FlushWorkspaceAsync(clientInfo, perforce, baseChange ?? 0);
-
-			// Apply all the updates
-			_logger.LogInformation("Syncing client {Client} from changelist {BaseChange} to {Change}", clientInfo.Client.Name, baseChange ?? 0, change);
-			clientInfo.Change = -1;
-
-			Utf8String clientRoot = new Utf8String(clientInfo.Client.Root);
-			string filterOrDefault = filter ?? "...";
-			string queryPath = $"//{clientInfo.Client.Name}/{filterOrDefault}";
-
-			RefName refName = GetRefName(stream.Id, change, filter, revisionsOnly);
-			RefName incRefName = new RefName($"{refName}-inc");
-			TreeWriter writer = new TreeWriter(storage, prefix: refName.Text);
-
-			// Get the current sync state for this change
-			ReplicationNode? syncNode = await reader.TryReadNodeAsync<ReplicationNode>(incRefName, cancellationToken: cancellationToken);
-			bool deleteIncRef = syncNode != null;
-			if (syncNode == null)
-			{
-				DirectoryNode contents = await baseContents.Contents.ExpandAsync(reader, cancellationToken);
-				syncNode = new ReplicationNode(contents);
-			}
-
-			// Get the contents of the current tree
-			DirectoryNode root = await syncNode.Contents.ExpandAsync(reader, cancellationToken);
-
-			if (revisionsOnly)
-			{
-				await foreach (PerforceResponse<SyncRecord> record in perforce.StreamCommandAsync<SyncRecord>("sync", new[] { "-k", $"{queryPath}@{change}" }, null, default, cancellationToken))
-				{
-					PerforceError? error = record.Error;
-					if (error != null && error.Generic == PerforceGenericCode.Empty)
-					{
-						continue;
-					}
-
-					SyncRecord syncRecord = record.Data;
-					if (!syncRecord.Path.StartsWith(clientRoot))
-					{
-						throw new ArgumentException($"Unable to make path {clientInfo.Client.Root} relative to client root {clientInfo.Client.Root}");
-					}
-
-					Utf8String path = syncRecord.Path.Substring(clientRoot.Length);
-
-					byte[] data = Encoding.UTF8.GetBytes($"{syncRecord.DepotFile}#{syncRecord.Revision}");
-					FileNode node = await FileNode.CreateAsync(data, _settings.Commits.Chunking, writer, cancellationToken);
-					await root.AddFileByPathAsync(reader, path, FileEntryFlags.PerforceDepotPathAndRevision, node, cancellationToken);
-				}
-			}
-			else
-			{
-				// Replay the files that have already been synced
-				foreach (Utf8String path in syncNode.Paths)
-				{
-					string flushPath = $"//{clientInfo.Client.Name}/{path}@{change}";
-					_logger.LogInformation("Flushing {FlushPath}", flushPath);
-					await perforce.SyncQuietAsync(SyncOptions.Force | SyncOptions.KeepWorkspaceFiles, -1, flushPath, cancellationToken);
-				}
-
-				// Add the root directory from the filter to the list of files to sync. This prevents traversing above it.
-				Dictionary<Utf8String, DirectoryToSync> pathToDirectory = new Dictionary<Utf8String, DirectoryToSync>(clientInfo.ServerInfo.Utf8PathComparer);
-				if (filter != null)
-				{
-					string rootPath = filter.Substring(0, filter.LastIndexOf('/') + 1);
-					pathToDirectory.Add(rootPath.Replace('/', Path.DirectorySeparatorChar), new DirectoryToSync(rootPath, clientInfo.ServerInfo.Utf8PathComparer));
-				}
-
-				// Do a sync preview to find everything that's left, and sort the remaining list of paths
-				await foreach (PerforceResponse<SyncRecord> response in perforce.StreamCommandAsync<SyncRecord>("sync", new[] { "-n" }, new string[] { $"{queryPath}@{change}" }, null, cancellationToken))
-				{
-					PerforceError? error = response.Error;
-					if (error != null)
-					{
-						_logger.LogWarning("Perforce: {Message}", error.Data);
-						continue;
-					}
-
-					Utf8String path = response.Data.Path;
-					if (!path.StartsWith(clientRoot, Utf8StringComparer.Ordinal))
-					{
-						throw new ArgumentException($"Unable to make path {path} relative to client root {clientRoot}");
-					}
-
-					int fileIdx = path.Length;
-					while (fileIdx > 0 && path[fileIdx - 1] != '/' && path[fileIdx - 1] != '\\')
-					{
-						fileIdx--;
-					}
-
-					Utf8String directoryPath = path.Slice(clientRoot.Length, fileIdx - clientRoot.Length);
-
-					DirectoryToSync directory = FindOrAddDirectoryTree(pathToDirectory, directoryPath, clientInfo.ServerInfo.Utf8PathComparer);
-					directory.FileNameToSize.Add(path.Slice(fileIdx).Clone(), response.Data.FileSize);
-					directory._size += response.Data.FileSize;
-				}
-
-				// Sort the directories by name
-				List<DirectoryToSync> directories = pathToDirectory.Values.OrderBy(x => x.Path, clientInfo.ServerInfo.Utf8PathComparer).ToList();
-
-				// Output some stats for the sync
-				long totalSize = directories.Sum(x => x._size);
-				_logger.LogInformation("Total sync size: {Size:n1}mb", totalSize / (1024.0 * 1024.0));
-
-				// Sync incrementally
-				long syncedSize = 0;
-				while (directories.Count > 0)
-				{
-					// Save the incremental state
-					if (syncedSize > 0)
-					{
-						Stopwatch flushTimer = Stopwatch.StartNew();
-						await writer.WriteRefAsync(incRefName, syncNode, cancellationToken: cancellationToken);
-						flushTimer.Stop();
-						deleteIncRef = true;
-					}
-
-					// Find the next paths to sync
-					const long MaxBatchSize = 1L * 1024 * 1024 * 1024;
-
-					int dirIdx = directories.Count - 1;
-					long size = directories[dirIdx]._size;
-
-					for (; dirIdx > 0; dirIdx--)
-					{
-						long nextSize = size + directories[dirIdx - 1]._size;
-						if (size > 0 && nextSize > MaxBatchSize)
-						{
-							break;
-						}
-						size = nextSize;
-					}
-
-					syncedSize += size;
-					double syncPct = (totalSize == 0) ? 100.0 : (syncedSize * 100.0) / totalSize;
-					_logger.LogInformation("Syncing {StreamId} to {Change} [{SyncPct:n1}%] ({Size:n1}mb)", stream.Id, change, syncPct, size / (1024.0 * 1024.0));
-
-					// Copy them to a separate list and remove any redundant paths
-					List<string> syncPaths = new List<string>();
-					for (int idx = dirIdx; idx < directories.Count; idx++)
-					{
-						Utf8String basePath = directories[idx].Path;
-						syncPaths.Add($"//{clientInfo.Client.Name}/{basePath}...@{change}");
-
-						long dirSize = directories[idx]._size;
-						while (idx + 1 < directories.Count && directories[idx + 1].Path.StartsWith(basePath, clientInfo.ServerInfo.Utf8PathComparer))
-						{
-							dirSize += directories[idx + 1]._size;
-							idx++;
-						}
-
-						_logger.LogInformation("  {Directory} ({Size:n1}mb)", directories[idx].Path + "...", dirSize / (1024.0 * 1024.0));
-					}
-
-					Stopwatch syncTimer = Stopwatch.StartNew();
-					Stopwatch processTimer = new Stopwatch();
-					Stopwatch gcTimer = new Stopwatch();
-
-					Dictionary<int, FileWriter> handles = new Dictionary<int, FileWriter>();
-					await foreach (PerforceResponse response in perforce.StreamCommandAsync("sync", Array.Empty<string>(), syncPaths, null, typeof(SyncRecord), true, default))
-					{
-						PerforceError? error = response.Error;
-						if (error != null)
-						{
-							if (error.Generic == PerforceGenericCode.Empty)
-							{
-								continue;
-							}
-							else
-							{
-								throw new ReplicationException($"Perforce error while replicating content - {error}");
-							}
-						}
-
-						processTimer.Start();
-
-						PerforceIo? io = response.Io;
-						if (io != null)
-						{
-							if (io.Command == PerforceIoCommand.Open)
-							{
-								Utf8String file = GetClientRelativePath(io.Payload, clientInfo.Client.Root);
-								int offset = GetFileOffset(file);
-
-								long fileSize = 0;
-								if (!pathToDirectory.TryGetValue(file.Substring(0, offset), out DirectoryToSync? directory))
-								{
-									throw new ReplicationException($"Unable to find directory for {file}");
-								}
-								if (!directory.FileNameToSize.TryGetValue(file.Substring(offset), out fileSize))
-								{
-									throw new ReplicationException($"Unable to find file entry for {file}");
-								}
-
-								handles[io.File] = new FileWriter(file, fileSize);
-							}
-							else if (io.Command == PerforceIoCommand.Write)
-							{
-								FileWriter file = handles[io.File];
-								await file.AppendAsync(io.Payload, _settings.Commits.Chunking, writer, cancellationToken);
-							}
-							else if (io.Command == PerforceIoCommand.Close)
-							{
-								FileWriter? file = null;
-								if (handles.Remove(io.File, out file))
-								{
-									FileNode node = file.Node;
-									if (node.Length != file.Size)
-									{
-										throw new ReplicationException($"Invalid size for replicated file '{file.Path}'. Expected {file.Size}, got {node.Length}.");
-									}
-
-									FileEntry entry = await root.AddFileByPathAsync(reader, file.Path, FileEntryFlags.None, node, cancellationToken);
-									await writer.WriteAsync(entry, cancellationToken);
-								}
-							}
-							else if (io.Command == PerforceIoCommand.Unlink)
-							{
-								Utf8String file = GetClientRelativePath(io.Payload, clientInfo.Client.Root);
-								await root.DeleteFileByPathAsync(reader, file, cancellationToken);
-							}
-							else
-							{
-								_logger.LogWarning("Unhandled command code {Code}", io.Command);
-							}
-						}
-
-						processTimer.Stop();
-					}
-
-					TimeSpan stallTime = (perforce as NativePerforceConnection)?.StallTime ?? TimeSpan.Zero;
-					_logger.LogInformation("Completed batch in {TimeSeconds:n1}s ({ProcessTimeSeconds:n1}s processing, {StallTimeSeconds:n1}s stalled, {Throughput:n1}mb/s)", syncTimer.Elapsed.TotalSeconds, processTimer.Elapsed.TotalSeconds, stallTime.TotalSeconds, size / (1024.0 * 1024.0 * syncTimer.Elapsed.TotalSeconds));
-
-					// Combine all the existing sync paths together with a new wildcard.
-					while (directories.Count > dirIdx)
-					{
-						Utf8String nextPath = directories[^1].Path;
-						if (syncNode.Paths.Count > 0)
-						{
-							Utf8String lastPath = syncNode.Paths[^1];
-							for (int endIdx = 0; endIdx < lastPath.Length; endIdx++)
-							{
-								if (lastPath[endIdx] == '/')
-								{
-									Utf8String prefix = lastPath.Substring(0, endIdx + 1);
-									if (!nextPath.StartsWith(prefix, clientInfo.ServerInfo.Utf8PathComparer))
-									{
-										// Remove any paths that start with this prefix
-										while (syncNode.Paths.Count > 0 && syncNode.Paths[^1].StartsWith(prefix, clientInfo.ServerInfo.Utf8PathComparer))
-										{
-											syncNode.Paths.RemoveAt(syncNode.Paths.Count - 1);
-										}
-
-										// Replace it with a wildcard
-										syncNode.Paths.Add(prefix + "...");
-										break;
-									}
-								}
-							}
-						}
-						syncNode.Paths.Add(nextPath + "...");
-						directories.RemoveAt(directories.Count - 1);
-					}
-				}
-			}
-//			await writer.WriteAsync(syncNode, cancellationToken);
-
-			clientInfo.Change = change;
-
-			// Return the new root object
-			_logger.LogInformation("Writing ref {RefId} for {StreamId} change {Change}", refName, stream.Id, change);
-			await writer.WriteRefAsync(refName, syncNode, cancellationToken: cancellationToken);
-
-			// Delete the incremental state
-			if (deleteIncRef)
-			{
-				await storage.DeleteRefAsync(incRefName, cancellationToken);
-			}
-			return syncNode;
-		}
-
-		static int GetFileOffset(Utf8String path)
-		{
-			int fileIdx = path.Length;
-			while (fileIdx > 0 && path[fileIdx - 1] != '/' && path[fileIdx - 1] != '\\')
-			{
-				fileIdx--;
-			}
-			return fileIdx;
-		}
-
-		static DirectoryToSync FindOrAddDirectoryTree(Dictionary<Utf8String, DirectoryToSync> pathToDirectory, Utf8String directoryPath, Utf8StringComparer comparer)
-		{
-			DirectoryToSync? directory;
-			if (!pathToDirectory.TryGetValue(directoryPath, out directory))
-			{
-				// Add a new path
-				Utf8String normalizedPath = NormalizePathSeparators(directoryPath);
-				directory = new DirectoryToSync(normalizedPath, comparer);
-				pathToDirectory.Add(directoryPath, directory);
-
-				// Also add all the parent directories. This makes the logic for combining wildcards simpler.
-				for (int idx = normalizedPath.Length - 2; idx > 0; idx--)
-				{
-					if (normalizedPath[idx] == '/')
-					{
-						Utf8String parentDirectoryPath = directoryPath.Substring(0, idx + 1);
-						if (pathToDirectory.ContainsKey(parentDirectoryPath))
-						{
-							break;
-						}
-						else
-						{
-							pathToDirectory.Add(parentDirectoryPath, new DirectoryToSync(normalizedPath.Substring(0, idx + 1), comparer));
-						}
-					}
-				}
-			}
-			return directory;
-		}
-
-		static Utf8String NormalizePathSeparators(Utf8String path)
-		{
-			byte[] newPath = new byte[path.Length];
-			for (int idx = 0; idx < path.Length; idx++)
-			{
-				newPath[idx] = (path[idx] == '\\') ? (byte)'/' : path[idx];
-			}
-			return new Utf8String(newPath);
-		}
-
-		async Task FlushWorkspaceAsync(ReplicationClient clientInfo, IPerforceConnection perforce, int change)
-		{
-			if (clientInfo.Change != change)
-			{
-				clientInfo.Change = -1;
-				if (change == 0)
-				{
-					_logger.LogInformation("Flushing have table for {Client}", clientInfo.Client.Name);
-					await perforce.SyncQuietAsync(SyncOptions.Force | SyncOptions.KeepWorkspaceFiles, -1, $"//{clientInfo.Client.Name}/...#0");
-				}
-				else
-				{
-					_logger.LogInformation("Flushing have table for {Client} to change {Change}", clientInfo.Client.Name, change);
-					await perforce.SyncQuietAsync(SyncOptions.Force | SyncOptions.KeepWorkspaceFiles, -1, $"//{clientInfo.Client.Name}/...@{change}");
-				}
-				clientInfo.Change = change;
-			}
-		}
-
-		static Utf8String GetClientRelativePath(ReadOnlyMemory<byte> data, Utf8String clientRoot)
-		{
-			int length = data.Span.IndexOf((byte)0);
-			if (length != -1)
-			{
-				data = data.Slice(0, length);
-			}
-
-			Utf8String path = new Utf8String(data);
-			if (!path.StartsWith(clientRoot, Utf8StringComparer.Ordinal))
-			{
-				throw new ArgumentException($"Unable to make path {path} relative to client root {clientRoot}");
-			}
-
-			return path.Substring(clientRoot.Length).Clone();
-		}
-
-		class ReplicationClient
-		{
-			public PerforceSettings Settings { get; }
-			public string ClusterName { get; }
-			public InfoRecord ServerInfo { get; }
-			public ClientRecord Client { get; }
-			public int Change { get; set; }
-
-			public ReplicationClient(PerforceSettings settings, string clusterName, InfoRecord serverInfo, ClientRecord client, int change)
-			{
-				Settings = settings;
-				ClusterName = clusterName;
-				ServerInfo = serverInfo;
-				Client = client;
-				Change = change;
-			}
-		}
-
-		readonly Dictionary<StreamId, ReplicationClient> _cachedPerforceClients = new Dictionary<StreamId, ReplicationClient>();
-
-		async Task<ReplicationClient?> FindReplicationClientAsync(IStream stream)
-		{
-			ReplicationClient? clientInfo;
-			if (_cachedPerforceClients.TryGetValue(stream.Id, out clientInfo))
-			{
-				if (!String.Equals(clientInfo.ClusterName, stream.Config.ClusterName, StringComparison.Ordinal) && String.Equals(clientInfo.Client.Stream, stream.Name, StringComparison.Ordinal))
-				{
-					PerforceSettings serverSettings = new PerforceSettings(clientInfo.Settings);
-					serverSettings.ClientName = null;
-
-					using IPerforceConnection perforce = await PerforceConnection.CreateAsync(_logger);
-					await perforce.DeleteClientAsync(DeleteClientOptions.None, clientInfo.Client.Name);
-
-					_cachedPerforceClients.Remove(stream.Id);
-					clientInfo = null;
-				}
-			}
-			return clientInfo;
-		}
-
-		async Task<ReplicationClient> FindOrAddReplicationClientAsync(IStream stream)
-		{
-			ReplicationClient? clientInfo = await FindReplicationClientAsync(stream);
-			if (clientInfo == null)
-			{
-				using IPerforceConnection? perforce = await _perforceService.ConnectAsync(stream.Config.ClusterName);
-				if (perforce == null)
-				{
-					throw new PerforceException($"Unable to create connection to Perforce server");
-				}
-
-				InfoRecord serverInfo = await perforce.GetInfoAsync(InfoOptions.ShortOutput);
-
-				ClientRecord newClient = new ClientRecord($"Horde.Build_Rep_{serverInfo.ClientHost}_{stream.Id}", perforce.Settings.UserName, "/p4/");
-				newClient.Description = "Created to mirror Perforce content to Horde Storage";
-				newClient.Owner = perforce.Settings.UserName;
-				newClient.Host = serverInfo.ClientHost;
-				newClient.Stream = stream.Config.ReplicationStream ?? stream.Name;
-				newClient.Type = "readonly";
-				await perforce.CreateClientAsync(newClient);
-				_logger.LogInformation("Created client {ClientName} for {StreamName}", newClient.Name, stream.Name);
-
-				PerforceSettings settings = new PerforceSettings(perforce.Settings);
-				settings.ClientName = newClient.Name;
-				settings.PreferNativeClient = true;
-
-				clientInfo = new ReplicationClient(settings, stream.Config.ClusterName, serverInfo, newClient, -1);
-				_cachedPerforceClients.Add(stream.Id, clientInfo);
-			}
-			return clientInfo;
-		}
+		public Task StopAsync(CancellationToken cancellationToken) => Task.CompletedTask;
 	}
 }
