Horde: Missing file.

diff --git a/Engine/Source/Programs/Horde/Horde.Build.Tests/Stubs/Services/PerforceServiceStub.cs b/Engine/Source/Programs/Horde/Horde.Build.Tests/Stubs/Services/PerforceServiceStub.cs
--- a/Engine/Source/Programs/Horde/Horde.Build.Tests/Stubs/Services/PerforceServiceStub.cs
+++ b/Engine/Source/Programs/Horde/Horde.Build.Tests/Stubs/Services/PerforceServiceStub.cs
@@ -249,6 +249,11 @@ namespace Horde.Build.Tests.Stubs.Services
 				List<ICommit> commits = await _owner.GetChangesAsync(_stream, null, null, 1, cancellationToken);
 				return commits[0].Number;
 			}
+
+			public IAsyncEnumerable<ICommit> SubscribeAsync(int minChange, IReadOnlyList<CommitTag>? tags = null, CancellationToken cancellationToken = default)
+			{
+				throw new NotImplementedException();
+			}
 		}
 
 		public ICommitCollection GetCommits(IStream stream)
