Horde: Add flag to workspace command for P4 native client preference

diff --git a/Engine/Source/Programs/Horde/Horde.Agent/Commands/Workspace/WorkspaceCommand.cs b/Engine/Source/Programs/Horde/Horde.Agent/Commands/Workspace/WorkspaceCommand.cs
--- a/Engine/Source/Programs/Horde/Horde.Agent/Commands/Workspace/WorkspaceCommand.cs
+++ b/Engine/Source/Programs/Horde/Horde.Agent/Commands/Workspace/WorkspaceCommand.cs
@@ -30,6 +30,10 @@ namespace Horde.Agent.Commands.Workspace
 		[CommandLine("-Overwrite")]
 		[Description("")]
 		protected bool Overwrite { get; set; } = false;
+		
+		[CommandLine("-PreferNativeClient")]
+		[Description("Prefer to use native Perforce client (instead of launching separate p4 process)")]
+		protected bool PreferNativeClient { get; set; } = false;
 
 		public override void Configure(CommandLineArguments arguments, ILogger logger)
 		{
@@ -55,15 +59,14 @@ namespace Horde.Agent.Commands.Workspace
 
 		public override async Task<int> ExecuteAsync(ILogger logger)
 		{
-			using (IPerforceConnection perforce = await PerforceConnection.CreateAsync(new PerforceSettings(ServerAndPort, UserName), logger))
-			{
-				InfoRecord info = await perforce.GetInfoAsync(InfoOptions.ShortOutput, CancellationToken.None);
+			PerforceSettings settings = new(ServerAndPort, UserName) { PreferNativeClient = PreferNativeClient };
+			using IPerforceConnection perforce = await PerforceConnection.CreateAsync(settings, logger);
+			InfoRecord info = await perforce.GetInfoAsync(InfoOptions.ShortOutput, CancellationToken.None);
 
-				using Logging.HordeLoggerProvider loggerProvider = new Logging.HordeLoggerProvider();
-				ILogger repoLogger = loggerProvider.CreateLogger("Repository");
-				ManagedWorkspace repo = await ManagedWorkspace.LoadOrCreateAsync(info.ClientHost!, BaseDir, Overwrite, repoLogger, CancellationToken.None);
-				await ExecuteAsync(perforce, repo, logger);
-			}
+			using Logging.HordeLoggerProvider loggerProvider = new ();
+			ILogger repoLogger = loggerProvider.CreateLogger("Repository");
+			ManagedWorkspace repo = await ManagedWorkspace.LoadOrCreateAsync(info.ClientHost!, BaseDir, Overwrite, repoLogger, CancellationToken.None);
+			await ExecuteAsync(perforce, repo, logger);
 			return 0;
 		}
 
