Horde: Use async disposable pattern for agent subscriptions rather than syncrhonous version.

diff --git a/Engine/Source/Programs/Horde/Horde.Build/Agents/AgentCollection.cs b/Engine/Source/Programs/Horde/Horde.Build/Agents/AgentCollection.cs
--- a/Engine/Source/Programs/Horde/Horde.Build/Agents/AgentCollection.cs
+++ b/Engine/Source/Programs/Horde/Horde.Build/Agents/AgentCollection.cs
@@ -558,9 +558,9 @@ namespace Horde.Build.Agents
 		}
 
 		/// <inheritdoc/>
-		public async Task<IDisposable> SubscribeToUpdateEventsAsync(Action<AgentId> onUpdate)
+		public async Task<IAsyncDisposable> SubscribeToUpdateEventsAsync(Action<AgentId> onUpdate)
 		{
-			return await _redisService.GetDatabase().Multiplexer.SubscribeAsync(_updateEventChannel, (channel, agentId) => onUpdate(agentId));
+			return await _redisService.GetDatabase().Multiplexer.SubscribeAsync(_updateEventChannel, onUpdate);
 		}
 	}
 }
diff --git a/Engine/Source/Programs/Horde/Horde.Build/Agents/AgentService.cs b/Engine/Source/Programs/Horde/Horde.Build/Agents/AgentService.cs
--- a/Engine/Source/Programs/Horde/Horde.Build/Agents/AgentService.cs
+++ b/Engine/Source/Programs/Horde/Horde.Build/Agents/AgentService.cs
@@ -96,7 +96,7 @@ namespace Horde.Build.Agents
 		readonly Dictionary<AgentId, CancellationTokenSource> _waitingAgents = new Dictionary<AgentId, CancellationTokenSource>();
 
 		// Subscription for update events
-		readonly IDisposable _subscription;
+		IAsyncDisposable? _subscription;
 
 		/// <summary>
 		/// Constructor
@@ -117,20 +117,29 @@ namespace Horde.Build.Agents
 			_clock = clock;
 			_ticker = clock.AddTicker<AgentService>(TimeSpan.FromSeconds(30.0), TickAsync, logger);
 			_logger = logger;
-
-			_subscription = agents.SubscribeToUpdateEventsAsync(OnAgentUpdate).Result;
 		}
 
 		/// <inheritdoc/>
-		public Task StartAsync(CancellationToken cancellationToken) => _ticker.StartAsync();
+		public async Task StartAsync(CancellationToken cancellationToken)
+		{
+			await _ticker.StartAsync();
+			_subscription = await Agents.SubscribeToUpdateEventsAsync(OnAgentUpdate);
+		}
 
 		/// <inheritdoc/>
-		public Task StopAsync(CancellationToken cancellationToken) => _ticker.StopAsync();
+		public async Task StopAsync(CancellationToken cancellationToken)
+		{
+			if (_subscription != null)
+			{
+				await _subscription.DisposeAsync();
+				_subscription = null;
+			}
+			await _ticker.StopAsync();
+		}
 
 		/// <inheritdoc/>
 		public void Dispose()
 		{
-			_subscription.Dispose();
 			_ticker.Dispose();
 		}
 
diff --git a/Engine/Source/Programs/Horde/Horde.Build/Agents/IAgentCollection.cs b/Engine/Source/Programs/Horde/Horde.Build/Agents/IAgentCollection.cs
--- a/Engine/Source/Programs/Horde/Horde.Build/Agents/IAgentCollection.cs
+++ b/Engine/Source/Programs/Horde/Horde.Build/Agents/IAgentCollection.cs
@@ -1,4 +1,4 @@
-﻿// Copyright Epic Games, Inc. All Rights Reserved.
+// Copyright Epic Games, Inc. All Rights Reserved.
 
 using System;
 using System.Collections.Generic;
@@ -157,6 +157,6 @@ namespace Horde.Build.Agents
 		/// </summary>
 		/// <param name="onUpdate">Callback for updates</param>
 		/// <returns>Disposable subscription object</returns>
-		Task<IDisposable> SubscribeToUpdateEventsAsync(Action<AgentId> onUpdate);
+		Task<IAsyncDisposable> SubscribeToUpdateEventsAsync(Action<AgentId> onUpdate);
 	}
 }
