Horde: Fix typo in argument order.

diff --git a/Engine/Source/Programs/Horde/Horde.Build/Dockerfile b/Engine/Source/Programs/Horde/Horde.Build/Dockerfile
--- a/Engine/Source/Programs/Horde/Horde.Build/Dockerfile
+++ b/Engine/Source/Programs/Horde/Horde.Build/Dockerfile
@@ -55,7 +55,7 @@ RUN redis-server --save "" --appendonly no --daemonize yes && \
 	dotnet test Source/Programs/Shared/EpicGames.Core.Tests/EpicGames.Core.Tests.csproj --blame-hang-timeout 2m --logger console;verbosity=normal && \
 	dotnet test Source/Programs/Shared/EpicGames.Horde.Tests/EpicGames.Horde.Tests.csproj --blame-hang-timeout 2m --logger console;verbosity=normal && \
 	dotnet test Source/Programs/Shared/EpicGames.Serialization.Tests/EpicGames.Serialization.Tests.csproj --blame-hang-timeout 2m --logger console;verbosity=normal && \
-	dotnet test Source/Programs/Horde/Horde.Build.Tests/Horde.Build.Tests.csproj --blame-hang-timeout --logger console;verbosity=normal 2m
+	dotnet test Source/Programs/Horde/Horde.Build.Tests/Horde.Build.Tests.csproj --blame-hang-timeout 2m --logger console;verbosity=normal
 
 # Create the runtime image
 FROM mcr.microsoft.com/dotnet/aspnet:6.0 AS runtime
