Horde: Allow pool scale-in during maintenance window

diff --git a/Engine/Source/Programs/Horde/Horde.Build.Tests/Fleet/FleetServiceTest.cs b/Engine/Source/Programs/Horde/Horde.Build.Tests/Fleet/FleetServiceTest.cs
--- a/Engine/Source/Programs/Horde/Horde.Build.Tests/Fleet/FleetServiceTest.cs
+++ b/Engine/Source/Programs/Horde/Horde.Build.Tests/Fleet/FleetServiceTest.cs
@@ -94,7 +94,7 @@ namespace Horde.Build.Tests.Fleet
 		[TestMethod]
 		public async Task OnlyEnabledAgentsAreAutoScaled()
 		{
-			using FleetService service = GetAutoscaleService(_fleetManagerSpy);
+			using FleetService service = GetFleetService(_fleetManagerSpy);
 			IPool pool = await PoolService.CreatePoolAsync("testPool", null, true, 0, 0, sizeStrategy: PoolSizeStrategy.LeaseUtilization);
 			await CreateAgentAsync(pool, true);
 			await CreateAgentAsync(pool, true);
@@ -109,7 +109,7 @@ namespace Horde.Build.Tests.Fleet
 		[TestMethod]
 		public async Task ExceptionsAreSwallowedAndLogged()
 		{
-			using FleetService service = GetAutoscaleService(_fleetManagerSpy);
+			using FleetService service = GetFleetService(_fleetManagerSpy);
 			IPool pool = await PoolService.CreatePoolAsync("testPool", null, true, 0, 0, sizeStrategy: PoolSizeStrategy.LeaseUtilization);
 			List<PoolSizeResult> poolSizeDatas = new()
 			{
@@ -124,7 +124,7 @@ namespace Horde.Build.Tests.Fleet
 		[TestMethod]
 		public async Task ScaleOutCooldown()
 		{
-			using FleetService service = GetAutoscaleService(_fleetManagerSpy);
+			using FleetService service = GetFleetService(_fleetManagerSpy);
 			IPool pool = await PoolService.CreatePoolAsync("testPool", null, true, 0, 0, sizeStrategy: PoolSizeStrategy.NoOp);
 
 			// First scale-out will succeed
@@ -144,7 +144,7 @@ namespace Horde.Build.Tests.Fleet
 		[TestMethod]
 		public async Task ScaleInCooldown()
 		{
-			using FleetService service = GetAutoscaleService(_fleetManagerSpy);
+			using FleetService service = GetFleetService(_fleetManagerSpy);
 			IPool pool = await PoolService.CreatePoolAsync("testPool", null, true, 0, 0, sizeStrategy: PoolSizeStrategy.NoOp);
 			IAgent agent1 = await CreateAgentAsync(pool);
 			IAgent agent2 = await CreateAgentAsync(pool);
@@ -162,15 +162,37 @@ namespace Horde.Build.Tests.Fleet
 			await service.ResizePoolsAsync(new() { new PoolSizeResult(pool, new () { agent1 }, 0) });
 			Assert.AreEqual(2, _fleetManagerSpy.ShrinkPoolAsyncCallCount);
 		}
+		
+		[TestMethod]
+		public async Task ScaleOutDuringDowntime()
+		{
+			using FleetService service = GetFleetService(_fleetManagerSpy, true);
+			IPool pool = await PoolService.CreatePoolAsync("testPool", null, true, 0, 0, sizeStrategy: PoolSizeStrategy.NoOp);
 
-		private FleetService GetAutoscaleService(IFleetManager fleetManager)
+			await service.ResizePoolsAsync(new() { new PoolSizeResult(pool, new List<IAgent>(), 1) });
+			Assert.AreEqual(0, _fleetManagerSpy.ExpandPoolAsyncCallCount);
+		}
+		
+		[TestMethod]
+		public async Task ScaleInDuringDowntime()
+		{
+			using FleetService service = GetFleetService(_fleetManagerSpy, true);
+			IPool pool = await PoolService.CreatePoolAsync("testPool", null, true, 0, 0, sizeStrategy: PoolSizeStrategy.NoOp);
+			IAgent agent1 = await CreateAgentAsync(pool);
+			IAgent agent2 = await CreateAgentAsync(pool);
+
+			await service.ResizePoolsAsync(new() { new PoolSizeResult(pool, new () { agent1, agent2 }, 1) });
+			Assert.AreEqual(1, _fleetManagerSpy.ShrinkPoolAsyncCallCount);
+		}
+
+		private FleetService GetFleetService(IFleetManager fleetManager, bool isDowntimeActive = false)
 		{
 			ILoggerFactory loggerFactory = ServiceProvider.GetRequiredService<ILoggerFactory>();
 			IOptions<ServerSettings> serverSettingsOpt = ServiceProvider.GetRequiredService<IOptions<ServerSettings>>();
 			serverSettingsOpt.Value.FleetManagerV2 = FleetManagerType.AwsReuse;
 			
 			FleetService service = new(
-				AgentCollection, GraphCollection, JobCollection, LeaseCollection, PoolCollection, DowntimeService, StreamService, _dogStatsD,
+				AgentCollection, GraphCollection, JobCollection, LeaseCollection, PoolCollection, new DowntimeServiceStub(isDowntimeActive), StreamService, _dogStatsD,
 				new StubFleetManagerFactory(fleetManager), Clock, Cache, serverSettingsOpt, loggerFactory.CreateLogger<FleetService>());
 				
 			return service;
@@ -325,6 +347,14 @@ namespace Horde.Build.Tests.Fleet
 			Assert.IsTrue(EvalCondition(new DateTime(2022, 9, 5, 7, 0, 15, DateTimeKind.Utc), "timeUtcSec > 5"));
 			Assert.IsFalse(EvalCondition(new DateTime(2022, 9, 5, 15, 0, 55, DateTimeKind.Utc), "timeUtcSec == 3"));
 		}
+		
+		[TestMethod]
+		public void ConditionCombination()
+		{
+			Assert.IsTrue(EvalCondition(new DateTime(2022, 9, 5, 15, 55, 0, DateTimeKind.Utc), "timeUtcDayOfWeek == 'monday' && "));
+			Assert.IsTrue(EvalCondition(new DateTime(2022, 9, 5, 7, 0, 0, DateTimeKind.Utc), "timeUtcMin == 0"));
+			Assert.IsFalse(EvalCondition(new DateTime(2022, 9, 5, 15, 1, 0, DateTimeKind.Utc), "timeUtcMin == 2"));
+		}
 
 		private bool EvalCondition(DateTime now, string condition)
 		{
diff --git a/Engine/Source/Programs/Horde/Horde.Build.Tests/TestSetup.cs b/Engine/Source/Programs/Horde/Horde.Build.Tests/TestSetup.cs
--- a/Engine/Source/Programs/Horde/Horde.Build.Tests/TestSetup.cs
+++ b/Engine/Source/Programs/Horde/Horde.Build.Tests/TestSetup.cs
@@ -385,10 +385,15 @@ namespace Horde.Build.Tests
 			
 			processExitInstance.TrySetResult(true);
 		}
-		
-		private class DowntimeServiceStub : IDowntimeService
+	}
+	
+	public class DowntimeServiceStub : IDowntimeService
+	{
+		public DowntimeServiceStub(bool isDowntimeActive = false)
 		{
-			public bool IsDowntimeActive { get; } = false;
+			IsDowntimeActive = isDowntimeActive;
 		}
+
+		public bool IsDowntimeActive { get; set; }
 	}
 }
\ No newline at end of file
diff --git a/Engine/Source/Programs/Horde/Horde.Build/Agents/Fleet/FleetService.cs b/Engine/Source/Programs/Horde/Horde.Build/Agents/Fleet/FleetService.cs
--- a/Engine/Source/Programs/Horde/Horde.Build/Agents/Fleet/FleetService.cs
+++ b/Engine/Source/Programs/Horde/Horde.Build/Agents/Fleet/FleetService.cs
@@ -230,18 +230,26 @@ namespace Horde.Build.Agents.Fleet
 			{
 				if (deltaAgentCount > 0)
 				{
-					TimeSpan scaleOutCooldown = pool.ScaleOutCooldown ?? _defaultScaleOutCooldown;
-					bool isCoolingDown = pool.LastScaleUpTime != null && pool.LastScaleUpTime + scaleOutCooldown > _clock.UtcNow;
-					scope.Span.SetTag("isCoolingDown", isCoolingDown);
-					if (!isCoolingDown)
+					if (!_downtimeService.IsDowntimeActive)
 					{
-						await fleetManager.ExpandPoolAsync(pool, poolSizeResult.Agents, deltaAgentCount, cancellationToken);
-						await _poolCollection.TryUpdateAsync(pool, lastScaleUpTime: _clock.UtcNow);
+						TimeSpan scaleOutCooldown = pool.ScaleOutCooldown ?? _defaultScaleOutCooldown;
+						bool isCoolingDown = pool.LastScaleUpTime != null && pool.LastScaleUpTime + scaleOutCooldown > _clock.UtcNow;
+						scope.Span.SetTag("isCoolingDown", isCoolingDown);
+						if (!isCoolingDown)
+						{
+							await fleetManager.ExpandPoolAsync(pool, poolSizeResult.Agents, deltaAgentCount, cancellationToken);
+							await _poolCollection.TryUpdateAsync(pool, lastScaleUpTime: _clock.UtcNow);
+						}
+						else
+						{
+							TimeSpan? cooldownTimeLeft = pool.LastScaleUpTime + _defaultScaleOutCooldown - _clock.UtcNow;
+							_logger.LogDebug("Cannot scale out {PoolName}, it's cooling down for another {TimeLeft} secs", pool.Name, cooldownTimeLeft?.TotalSeconds);
+						}
 					}
 					else
 					{
-						TimeSpan? cooldownTimeLeft = pool.LastScaleUpTime + _defaultScaleOutCooldown - _clock.UtcNow;
-						_logger.LogDebug("Cannot scale out {PoolName}, it's cooling down for another {TimeLeft} secs", pool.Name, cooldownTimeLeft?.TotalSeconds);
+						_logger.LogDebug("Cannot scale out {PoolName}, downtime is active", pool.Name);
+						scope.Span.SetTag("IsDowntimeActive", true);
 					}
 				}
 
