Horde: Remove downtime checks from FleetService ticking

Now checked inside ResizePoolAsync instead.

diff --git a/Engine/Source/Programs/Horde/Horde.Build/Agents/Fleet/FleetService.cs b/Engine/Source/Programs/Horde/Horde.Build/Agents/Fleet/FleetService.cs
--- a/Engine/Source/Programs/Horde/Horde.Build/Agents/Fleet/FleetService.cs
+++ b/Engine/Source/Programs/Horde/Horde.Build/Agents/Fleet/FleetService.cs
@@ -117,12 +117,6 @@ namespace Horde.Build.Agents.Fleet
 			ISpan span = GlobalTracer.Instance.BuildSpan("FleetService.Tick").Start();
 			try
 			{
-				if (_downtimeService.IsDowntimeActive)
-				{
-					span.SetTag("IsDowntimeActive", _downtimeService.IsDowntimeActive);
-					return;
-				}
-				
 				List<PoolSizeResult> input = await GetPoolSizeDataAsync();
 				List<PoolSizeResult> output = await CalculatePoolSizesAsync(input, stoppingToken);
 				await ResizePoolsAsync(output, stoppingToken);
@@ -138,12 +132,6 @@ namespace Horde.Build.Agents.Fleet
 			ISpan span = GlobalTracer.Instance.BuildSpan("FleetService.TickHighFrequency").Start();
 			try
 			{
-				if (_downtimeService.IsDowntimeActive)
-				{
-					span.SetTag("IsDowntimeActive", _downtimeService.IsDowntimeActive);
-					return;
-				}
-				
 				// TODO: Re-enable high frequency scaling (only used for experimental scaling of remote execution agents at the moment)
 				await Task.Delay(0, stoppingToken);
 			}
