Horde: Parallelize pool size calculation + scaling in one batch

diff --git a/Engine/Source/Programs/Horde/Horde.Build.Tests/Fleet/FleetServiceTest.cs b/Engine/Source/Programs/Horde/Horde.Build.Tests/Fleet/FleetServiceTest.cs
--- a/Engine/Source/Programs/Horde/Horde.Build.Tests/Fleet/FleetServiceTest.cs
+++ b/Engine/Source/Programs/Horde/Horde.Build.Tests/Fleet/FleetServiceTest.cs
@@ -26,16 +26,23 @@ namespace Horde.Build.Tests.Fleet
 	{
 		public int ExpandPoolAsyncCallCount { get; private set; }
 		public int ShrinkPoolAsyncCallCount { get; private set; }
+		public IReadOnlyList<IAgent> ExpandAgents => _expandAgents;
+		public IReadOnlyList<IAgent> ShrinkAgents => _shrinkAgents;
+		
+		private readonly List<IAgent> _expandAgents = new ();
+		private readonly List<IAgent> _shrinkAgents = new ();
 		
 		public Task ExpandPoolAsync(IPool pool, IReadOnlyList<IAgent> agents, int count, CancellationToken cancellationToken)
 		{
 			ExpandPoolAsyncCallCount++;
+			_expandAgents.AddRange(agents);
 			return Task.CompletedTask;
 		}
 
 		public Task ShrinkPoolAsync(IPool pool, IReadOnlyList<IAgent> agents, int count, CancellationToken cancellationToken)
 		{
 			ShrinkPoolAsyncCallCount++;
+			_shrinkAgents.AddRange(agents);
 			return Task.CompletedTask;
 		}
 
@@ -45,20 +52,6 @@ namespace Horde.Build.Tests.Fleet
 		}
 	}
 
-	public class PoolSizeStrategySpy : IPoolSizeStrategy
-	{
-		public int CallCount { get; private set; }
-		public HashSet<PoolId> PoolIdsSeen { get; } = new();
-		
-		public Task<PoolSizeResult> CalculatePoolSizeAsync(IPool pool, List<IAgent> agents)
-		{
-			CallCount++;
-			PoolIdsSeen.Add(pool.Id);
-			return Task.FromResult(new PoolSizeResult(pool, agents, agents.Count));
-		}
-
-		public string Name { get; } = "PoolSizeStrategySpy";
-	}
 	
 	public class StubFleetManagerFactory : IFleetManagerFactory
 	{
@@ -100,27 +93,13 @@ namespace Horde.Build.Tests.Fleet
 			await CreateAgentAsync(pool, true);
 			await CreateAgentAsync(pool, false);
 
-			List<PoolSizeResult> poolSizeDatas = await service.GetPoolSizeDataAsync();
-			poolSizeDatas = await service.CalculatePoolSizesAsync(poolSizeDatas, CancellationToken.None);
-			Assert.AreEqual(1, poolSizeDatas.Count);
-			Assert.AreEqual(2, poolSizeDatas[0].Agents.Count);
+			PoolWithAgents poolWithAgents = (await service.GetPoolsWithAgentsAsync())[0];
+			await service.ScalePoolAsync(pool, poolWithAgents.Agents, new PoolSizeResult(2, 5), CancellationToken.None);
+			Assert.AreEqual(2, _fleetManagerSpy.ExpandAgents.Count);
+			Assert.IsTrue(_fleetManagerSpy.ExpandAgents[0].Enabled);
+			Assert.IsTrue(_fleetManagerSpy.ExpandAgents[1].Enabled);
 		}
 		
-		[TestMethod]
-		public async Task ExceptionsAreSwallowedAndLogged()
-		{
-			using FleetService service = GetFleetService(_fleetManagerSpy);
-			IPool pool = await PoolService.CreatePoolAsync("testPool", null, true, 0, 0, sizeStrategy: PoolSizeStrategy.LeaseUtilization);
-			List<PoolSizeResult> poolSizeDatas = new()
-			{
-				new PoolSizeResult(null!, new(), null),
-				new PoolSizeResult(pool, new(), null),
-			};
-
-			poolSizeDatas = await service.CalculatePoolSizesAsync(poolSizeDatas, CancellationToken.None);
-			Assert.AreEqual(1, poolSizeDatas.Count);
-		}
-
 		[TestMethod]
 		public async Task ScaleOutCooldown()
 		{
@@ -128,16 +107,16 @@ namespace Horde.Build.Tests.Fleet
 			IPool pool = await PoolService.CreatePoolAsync("testPool", null, true, 0, 0, sizeStrategy: PoolSizeStrategy.NoOp);
 
 			// First scale-out will succeed
-			await service.ResizePoolsAsync(new() { new PoolSizeResult(pool, new List<IAgent>(), 1) });
+			await service.ScalePoolAsync(pool, new List<IAgent>(), new PoolSizeResult(0, 1), CancellationToken.None);
 			Assert.AreEqual(1, _fleetManagerSpy.ExpandPoolAsyncCallCount);
 			
 			// Cannot scale-out due to cool-down
-			await service.ResizePoolsAsync(new() { new PoolSizeResult(pool, new List<IAgent>(), 2) });
+			await service.ScalePoolAsync(pool, new List<IAgent>(), new PoolSizeResult(0, 2), CancellationToken.None);
 			Assert.AreEqual(1, _fleetManagerSpy.ExpandPoolAsyncCallCount);
 
 			// Wait some time and then try again
 			await Clock.AdvanceAsync(TimeSpan.FromHours(2));
-			await service.ResizePoolsAsync(new() { new PoolSizeResult(pool, new List<IAgent>(), 2) });
+			await service.ScalePoolAsync(pool, new List<IAgent>(), new PoolSizeResult(0, 2), CancellationToken.None);
 			Assert.AreEqual(2, _fleetManagerSpy.ExpandPoolAsyncCallCount);
 		}
 		
@@ -149,17 +128,17 @@ namespace Horde.Build.Tests.Fleet
 			IAgent agent1 = await CreateAgentAsync(pool);
 			IAgent agent2 = await CreateAgentAsync(pool);
 
-			// First scale-out will succeed
-			await service.ResizePoolsAsync(new() { new PoolSizeResult(pool, new () { agent1, agent2 }, 1) });
+			// First scale-in will succeed
+			await service.ScalePoolAsync(pool, new List<IAgent> { agent1, agent2 }, new PoolSizeResult(2, 1), CancellationToken.None);
 			Assert.AreEqual(1, _fleetManagerSpy.ShrinkPoolAsyncCallCount);
 			
-			// Cannot scale-out due to cool-down
-			await service.ResizePoolsAsync(new() { new PoolSizeResult(pool, new () { agent1 }, 0) });
+			// Cannot scale-in due to cool-down
+			await service.ScalePoolAsync(pool, new List<IAgent> { agent1 }, new PoolSizeResult(1, 0), CancellationToken.None);
 			Assert.AreEqual(1, _fleetManagerSpy.ShrinkPoolAsyncCallCount);
 
 			// Wait some time and then try again
 			await Clock.AdvanceAsync(TimeSpan.FromHours(2));
-			await service.ResizePoolsAsync(new() { new PoolSizeResult(pool, new () { agent1 }, 0) });
+			await service.ScalePoolAsync(pool, new List<IAgent> { agent1 }, new PoolSizeResult(1, 0), CancellationToken.None);
 			Assert.AreEqual(2, _fleetManagerSpy.ShrinkPoolAsyncCallCount);
 		}
 		
@@ -169,7 +148,7 @@ namespace Horde.Build.Tests.Fleet
 			using FleetService service = GetFleetService(_fleetManagerSpy, true);
 			IPool pool = await PoolService.CreatePoolAsync("testPool", null, true, 0, 0, sizeStrategy: PoolSizeStrategy.NoOp);
 
-			await service.ResizePoolsAsync(new() { new PoolSizeResult(pool, new List<IAgent>(), 1) });
+			await service.ScalePoolAsync(pool, new List<IAgent>(), new PoolSizeResult(0, 1), CancellationToken.None);
 			Assert.AreEqual(0, _fleetManagerSpy.ExpandPoolAsyncCallCount);
 		}
 		
@@ -181,7 +160,7 @@ namespace Horde.Build.Tests.Fleet
 			IAgent agent1 = await CreateAgentAsync(pool);
 			IAgent agent2 = await CreateAgentAsync(pool);
 
-			await service.ResizePoolsAsync(new() { new PoolSizeResult(pool, new () { agent1, agent2 }, 1) });
+			await service.ScalePoolAsync(pool, new List<IAgent> { agent1, agent2 }, new PoolSizeResult(2, 1), CancellationToken.None);
 			Assert.AreEqual(1, _fleetManagerSpy.ShrinkPoolAsyncCallCount);
 		}
 
@@ -347,14 +326,6 @@ namespace Horde.Build.Tests.Fleet
 			Assert.IsTrue(EvalCondition(new DateTime(2022, 9, 5, 7, 0, 15, DateTimeKind.Utc), "timeUtcSec > 5"));
 			Assert.IsFalse(EvalCondition(new DateTime(2022, 9, 5, 15, 0, 55, DateTimeKind.Utc), "timeUtcSec == 3"));
 		}
-		
-		[TestMethod]
-		public void ConditionCombination()
-		{
-			Assert.IsTrue(EvalCondition(new DateTime(2022, 9, 5, 15, 55, 0, DateTimeKind.Utc), "timeUtcDayOfWeek == 'monday' && "));
-			Assert.IsTrue(EvalCondition(new DateTime(2022, 9, 5, 7, 0, 0, DateTimeKind.Utc), "timeUtcMin == 0"));
-			Assert.IsFalse(EvalCondition(new DateTime(2022, 9, 5, 15, 1, 0, DateTimeKind.Utc), "timeUtcMin == 2"));
-		}
 
 		private bool EvalCondition(DateTime now, string condition)
 		{
diff --git a/Engine/Source/Programs/Horde/Horde.Build.Tests/Fleet/JobQueueStrategyTest.cs b/Engine/Source/Programs/Horde/Horde.Build.Tests/Fleet/JobQueueStrategyTest.cs
--- a/Engine/Source/Programs/Horde/Horde.Build.Tests/Fleet/JobQueueStrategyTest.cs
+++ b/Engine/Source/Programs/Horde/Horde.Build.Tests/Fleet/JobQueueStrategyTest.cs
@@ -27,11 +27,11 @@ namespace Horde.Build.Tests.Fleet
 		[TestMethod]
 		public async Task GetPoolQueueSizes()
 		{
-			(JobQueueStrategy strategy, PoolSizeResult poolSizeData) = await SetUpJobsAsync(1, 5);
+			(JobQueueStrategy strategy, PoolSizeResult poolSizeResult, IPool pool, List<IAgent> agents) = await SetUpJobsAsync(1, 5);
 			await Clock.AdvanceAsync(TimeSpan.FromSeconds(strategy.Settings.ReadyTimeThresholdSec) + TimeSpan.FromSeconds(5));
 			Dictionary<PoolId, int> poolQueueSizes = await strategy.GetPoolQueueSizesAsync(Clock.UtcNow - TimeSpan.FromHours(2));
 			Assert.AreEqual(1, poolQueueSizes.Count);
-			Assert.AreEqual(5, poolQueueSizes[poolSizeData.Pool.Id]);
+			Assert.AreEqual(5, poolQueueSizes[pool.Id]);
 		}
 		
 		[TestMethod]
@@ -84,15 +84,15 @@ namespace Horde.Build.Tests.Fleet
 
 		public async Task AssertAgentCount(int numBatchesReady, int expectedAgentDelta, bool waitedBeyondThreshold = true, int numAgents = 8)
 		{
-			(JobQueueStrategy strategy, PoolSizeResult poolSizeData) = await SetUpJobsAsync(1, numBatchesReady, numAgents);
+			(JobQueueStrategy strategy, PoolSizeResult poolSizeResult, IPool pool, List<IAgent> agents) = await SetUpJobsAsync(1, numBatchesReady, numAgents);
 			TimeSpan timeToWait = waitedBeyondThreshold
 				? TimeSpan.FromSeconds(strategy.Settings.ReadyTimeThresholdSec) + TimeSpan.FromSeconds(5)
 				: TimeSpan.FromSeconds(15);
 			
 			await Clock.AdvanceAsync(timeToWait);
 
-			PoolSizeResult result = await strategy.CalculatePoolSizeAsync(poolSizeData.Pool, poolSizeData.Agents);
-			Assert.AreEqual(result.Agents.Count + expectedAgentDelta, result.DesiredAgentCount);
+			PoolSizeResult result = await strategy.CalculatePoolSizeAsync(pool, agents);
+			Assert.AreEqual(agents.Count + expectedAgentDelta, result.DesiredAgentCount);
 		}
 	
 		/// <summary>
@@ -100,19 +100,19 @@ namespace Horde.Build.Tests.Fleet
 		/// </summary>
 		/// <param name="numBatchesRunning">Num of job batches that should be in state running</param>
 		/// <param name="numBatchesReady">Num of job batches that should be in state waiting</param>
-		private async Task<(JobQueueStrategy, PoolSizeResult)> SetUpJobsAsync(int numBatchesRunning, int numBatchesReady, int numAgents = 8)
+		private async Task<(JobQueueStrategy, PoolSizeResult, IPool, List<IAgent> agents)> SetUpJobsAsync(int numBatchesRunning, int numBatchesReady, int numAgents = 8)
 		{
-			IPool pool1 = await PoolService.CreatePoolAsync("bogusPool1", null, true, 0, 0);
+			IPool pool = await PoolService.CreatePoolAsync("bogusPool1", null, true, 0, 0);
 			List<IAgent> agents = new();
 			for (int i = 0; i < numAgents; i++)
 			{
-				agents.Add(await CreateAgentAsync(pool1));
+				agents.Add(await CreateAgentAsync(pool));
 			}
 			
-			PoolSizeResult poolSize = new (pool1, agents, null);
+			PoolSizeResult poolSize = new (agents.Count, agents.Count, null);
 			
 			string agentTypeName1 = "bogusAgentType1";
-			Dictionary<string, AgentConfig> agentTypes = new() { {agentTypeName1, new() { Pool = new PoolId(pool1.Name) } }, };
+			Dictionary<string, AgentConfig> agentTypes = new() { {agentTypeName1, new() { Pool = new PoolId(pool.Name) } }, };
 
 			IStream stream = (await CreateOrReplaceStreamAsync(
 				new StreamId("ue5-main"),
@@ -143,7 +143,7 @@ namespace Horde.Build.Tests.Fleet
 			}
 
 			JobQueueSettings settings = new ();
-			return (new (JobCollection, GraphCollection, StreamService, Clock, Cache, settings), poolSize);
+			return (new (JobCollection, GraphCollection, StreamService, Clock, Cache, settings), poolSize, pool, agents);
 		}
 		
 		private async Task<IJob> AddPlaceholderJob(IGraph graph, StreamId streamId, string nodeNameToExecute)
diff --git a/Engine/Source/Programs/Horde/Horde.Build/Agents/Fleet/ComputeQueueAwsMetricStrategy.cs b/Engine/Source/Programs/Horde/Horde.Build/Agents/Fleet/ComputeQueueAwsMetricStrategy.cs
--- a/Engine/Source/Programs/Horde/Horde.Build/Agents/Fleet/ComputeQueueAwsMetricStrategy.cs
+++ b/Engine/Source/Programs/Horde/Horde.Build/Agents/Fleet/ComputeQueueAwsMetricStrategy.cs
@@ -131,7 +131,7 @@ namespace Horde.Build.Agents.Fleet
 			}
 
 			// Return the input pool size data as-is since we are only observing and not modifying
-			return new PoolSizeResult(pool, agents, agents.Count);
+			return new PoolSizeResult(numAgents, numAgents);
 		}
 	}
 }
diff --git a/Engine/Source/Programs/Horde/Horde.Build/Agents/Fleet/FleetService.cs b/Engine/Source/Programs/Horde/Horde.Build/Agents/Fleet/FleetService.cs
--- a/Engine/Source/Programs/Horde/Horde.Build/Agents/Fleet/FleetService.cs
+++ b/Engine/Source/Programs/Horde/Horde.Build/Agents/Fleet/FleetService.cs
@@ -27,7 +27,34 @@ using StatsdClient;
 namespace Horde.Build.Agents.Fleet
 {
 	using PoolId = StringId<IPool>;
-
+	
+	/// <summary>
+	/// Parameters required for calculating pool size
+	/// </summary>
+	public class PoolWithAgents
+	{
+		/// <summary>
+		/// Pool being resized
+		/// </summary>
+		public IPool Pool { get; }
+		
+		/// <summary>
+		/// All agents currently associated with the pool
+		/// </summary>
+		public List<IAgent> Agents { get; }
+	
+		/// <summary>
+		/// Constructor
+		/// </summary>
+		/// <param name="pool"></param>
+		/// <param name="agents"></param>
+		public PoolWithAgents(IPool pool, List<IAgent> agents)
+		{
+			Pool = pool;
+			Agents = agents;
+		}
+	}
+	
 	/// <summary>
 	/// Service for managing the autoscaling of agent pools
 	/// </summary>
@@ -117,9 +144,20 @@ namespace Horde.Build.Agents.Fleet
 			ISpan span = GlobalTracer.Instance.BuildSpan("FleetService.Tick").Start();
 			try
 			{
-				List<PoolSizeResult> input = await GetPoolSizeDataAsync();
-				List<PoolSizeResult> output = await CalculatePoolSizesAsync(input, stoppingToken);
-				await ResizePoolsAsync(output, stoppingToken);
+				List<PoolWithAgents> poolsWithAgents = await GetPoolsWithAgentsAsync();
+			
+				ParallelOptions options = new () { MaxDegreeOfParallelism = MaxParallelTasks, CancellationToken = stoppingToken };
+				await Parallel.ForEachAsync(poolsWithAgents, options, async (input, innerCt) =>
+				{
+					try
+					{
+						await CalculateSizeAndScaleAsync(input.Pool, input.Agents, innerCt);
+					}
+					catch (Exception e)
+					{
+						_logger.LogError(e, "Failed to scale pool {PoolId}", input.Pool.Id);
+					}
+				});
 			}
 			finally
 			{
@@ -141,72 +179,47 @@ namespace Horde.Build.Agents.Fleet
 			}
 		}
 
-		internal async Task<List<PoolSizeResult>> CalculatePoolSizesAsync(List<PoolSizeResult> poolSizeDatas, CancellationToken cancellationToken)
+		internal async Task CalculateSizeAndScaleAsync(IPool pool, List<IAgent> agents, CancellationToken cancellationToken)
 		{
-			ISpan span = GlobalTracer.Instance.BuildSpan("FleetService.CalculatePoolSizes").WithTag("MaxParallelTasks", MaxParallelTasks).Start();
-			ConcurrentQueue<PoolSizeResult> results = new ();
-			
-			try
-			{
-				ParallelOptions options = new () { MaxDegreeOfParallelism = MaxParallelTasks, CancellationToken = cancellationToken };
-				await Parallel.ForEachAsync(poolSizeDatas, options, async (input, innerCt) =>
-				{
-					try
-					{
-						IPoolSizeStrategy sizeStrategy = CreatePoolSizeStrategy(input.Pool);
-						PoolSizeResult output = await sizeStrategy.CalculatePoolSizeAsync(input.Pool, input.Agents);
-						results.Enqueue(output);
-					}
-					catch (Exception e)
-					{
-						_logger.LogError(e, "Failed calculating pool size for pool {PoolId}", input.Pool?.Id);
-					}
-				});
-			}
-			finally
-			{
-				span.Finish();
-			}
-
-			return results.ToList();
+			IPoolSizeStrategy sizeStrategy = CreatePoolSizeStrategy(pool);
+			PoolSizeResult result = await sizeStrategy.CalculatePoolSizeAsync(pool, agents);
+			await ScalePoolAsync(pool, agents, result, cancellationToken);
 		}
 
-		internal async Task<List<PoolSizeResult>> GetPoolSizeDataAsync()
+		internal async Task<List<PoolWithAgents>> GetPoolsWithAgentsAsync()
 		{
 			List<IAgent> agents = await _agentCollection.FindAsync(status: AgentStatus.Ok, enabled: true);
 			List<IAgent> GetAgentsInPool(PoolId poolId) => agents.FindAll(a => a.GetPools().Any(p => p == poolId));
 			List<IPool> pools = await _poolCollection.GetAsync();
 
-			return pools.Select(pool => new PoolSizeResult(pool, GetAgentsInPool(pool.Id), null)).ToList();
+			return pools.Select(pool => new PoolWithAgents(pool, GetAgentsInPool(pool.Id))).ToList();
 		}
 
-		internal async Task ResizePoolAsync(PoolSizeResult poolSizeResult, CancellationToken cancellationToken)
+		internal async Task ScalePoolAsync(IPool pool, List<IAgent> agents, PoolSizeResult poolSizeResult, CancellationToken cancellationToken)
 		{
-			IPool pool = poolSizeResult.Pool;
-
-			if (!pool.EnableAutoscaling || poolSizeResult.DesiredAgentCount == null)
+			if (!pool.EnableAutoscaling)
 			{
 				return;
 			}
 
-			int currentAgentCount = poolSizeResult.Agents.Count;
-			int desiredAgentCount = poolSizeResult.DesiredAgentCount.Value;
+			int currentAgentCount = poolSizeResult.CurrentAgentCount;
+			int desiredAgentCount = poolSizeResult.DesiredAgentCount;
 			int deltaAgentCount = desiredAgentCount - currentAgentCount;
 
 			using IScope scope = GlobalTracer.Instance
-				.BuildSpan("FleetService.ResizePool")
+				.BuildSpan("FleetService.ScalePool")
 				.WithTag(Datadog.Trace.OpenTracing.DatadogTags.ResourceName, pool.Id.ToString())
 				.WithTag("CurrentAgentCount", currentAgentCount)
 				.WithTag("DesiredAgentCount", desiredAgentCount)
 				.WithTag("DeltaAgentCount", deltaAgentCount)
 				.StartActive();
 			
-			IFleetManager fleetManager = CreateFleetManager(poolSizeResult.Pool);
+			IFleetManager fleetManager = CreateFleetManager(pool);
 			Dictionary<string, object> logScopeMetadata = new()
 			{
 				["FleetManager"] = fleetManager.GetType().Name,
 				["PoolSizeStrategy"] = poolSizeResult.Status ?? new Dictionary<string, object>(),
-				["PoolId"] = poolSizeResult.Pool.Id,
+				["PoolId"] = pool.Id,
 			};
 
 			using IDisposable logScope = _logger.BeginScope(logScopeMetadata);
@@ -225,7 +238,7 @@ namespace Horde.Build.Agents.Fleet
 						scope.Span.SetTag("isCoolingDown", isCoolingDown);
 						if (!isCoolingDown)
 						{
-							await fleetManager.ExpandPoolAsync(pool, poolSizeResult.Agents, deltaAgentCount, cancellationToken);
+							await fleetManager.ExpandPoolAsync(pool, agents, deltaAgentCount, cancellationToken);
 							await _poolCollection.TryUpdateAsync(pool, lastScaleUpTime: _clock.UtcNow);
 						}
 						else
@@ -248,7 +261,7 @@ namespace Horde.Build.Agents.Fleet
 					scope.Span.SetTag("isCoolingDown", isCoolingDown);
 					if (!isCoolingDown)
 					{
-						await fleetManager.ShrinkPoolAsync(pool, poolSizeResult.Agents, -deltaAgentCount, cancellationToken);
+						await fleetManager.ShrinkPoolAsync(pool, agents, -deltaAgentCount, cancellationToken);
 						await _poolCollection.TryUpdateAsync(pool, lastScaleDownTime: _clock.UtcNow);
 					}
 					else
@@ -268,15 +281,6 @@ namespace Horde.Build.Agents.Fleet
 			_dogStatsd.Gauge("agentpools.autoscale.current", currentAgentCount, tags: new []{"pool:" + pool.Name});
 		}
 		
-		internal async Task ResizePoolsAsync(List<PoolSizeResult> poolSizeDatas, CancellationToken cancellationToken = default)
-		{
-			ParallelOptions options = new () { MaxDegreeOfParallelism = MaxParallelTasks, CancellationToken = cancellationToken };
-			await Parallel.ForEachAsync(poolSizeDatas.OrderByDescending(x => x.Agents.Count), options, async (poolSizeData, innerCt) =>
-			{
-				await ResizePoolAsync(poolSizeData, innerCt);
-			});
-		}
-		
 		internal IEnumerable<string> GetPropValues(string name)
 		{
 			DateTime now = _clock.UtcNow;
diff --git a/Engine/Source/Programs/Horde/Horde.Build/Agents/Fleet/IPoolSizeStrategy.cs b/Engine/Source/Programs/Horde/Horde.Build/Agents/Fleet/IPoolSizeStrategy.cs
--- a/Engine/Source/Programs/Horde/Horde.Build/Agents/Fleet/IPoolSizeStrategy.cs
+++ b/Engine/Source/Programs/Horde/Horde.Build/Agents/Fleet/IPoolSizeStrategy.cs
@@ -57,24 +57,19 @@ namespace Horde.Build.Agents.Fleet
 	}
 
 	/// <summary>
-	/// Class for specifying and grouping data together required for calculating pool size
+	/// Result from calculating pool size
 	/// </summary>
 	public class PoolSizeResult
 	{
 		/// <summary>
-		/// Pool being resized
+		/// The agent count as it appeared when pool sizing was calculated
 		/// </summary>
-		public IPool Pool { get; }
-		
-		/// <summary>
-		/// All agents currently associated with the pool
-		/// </summary>
-		public List<IAgent> Agents { get; }
+		public int CurrentAgentCount { get; }
 
 		/// <summary>
-		/// The desired agent count (calculated and updated once the strategy has been run, null otherwise)
+		/// The desired agent count as calculated by pool sizing strategy
 		/// </summary>
-		public int? DesiredAgentCount  { get; }
+		public int DesiredAgentCount { get; }
 		
 		/// <summary>
 		/// Log-friendly metadata object describing the output of the size calculation through key/values
@@ -84,31 +79,15 @@ namespace Horde.Build.Agents.Fleet
 		/// <summary>
 		/// Constructor
 		/// </summary>
-		/// <param name="pool"></param>
-		/// <param name="agents"></param>
+		/// <param name="currentAgentCount"></param>
 		/// <param name="desiredAgentCount"></param>
 		/// <param name="status"></param>
-		public PoolSizeResult(IPool pool, List<IAgent> agents, int? desiredAgentCount, IReadOnlyDictionary<string, object>? status = null)
+		public PoolSizeResult(int currentAgentCount, int desiredAgentCount, IReadOnlyDictionary<string, object>? status = null)
 		{
-			Pool = pool;
-			Agents = agents;
+			CurrentAgentCount = currentAgentCount;
 			DesiredAgentCount = desiredAgentCount;
 			Status = status;
 		}
-
-		/// <summary>
-		/// Copy the object, inheriting any unspecified values from current instance
-		/// Needed because the class is immutable 
-		/// </summary>
-		/// <param name="pool"></param>
-		/// <param name="agents"></param>
-		/// <param name="desiredAgentCount"></param>
-		/// <param name="status"></param>
-		/// <returns>A new copy</returns>
-		public PoolSizeResult Copy(IPool? pool = null, List<IAgent>? agents = null, int? desiredAgentCount = null, IReadOnlyDictionary<string, object>? status = null)
-		{
-			return new PoolSizeResult(pool ?? Pool, agents ?? Agents, desiredAgentCount ?? DesiredAgentCount, status);
-		}
 	}
 	
 	/// <summary>
@@ -139,7 +118,7 @@ namespace Horde.Build.Agents.Fleet
 		/// <inheritdoc/>
 		public Task<PoolSizeResult> CalculatePoolSizeAsync(IPool pool, List<IAgent> agents)
 		{
-			return Task.FromResult(new PoolSizeResult(pool, agents, agents.Count));
+			return Task.FromResult(new PoolSizeResult(agents.Count, agents.Count));
 		}
 
 		/// <inheritdoc/>
diff --git a/Engine/Source/Programs/Horde/Horde.Build/Agents/Fleet/JobQueueStrategy.cs b/Engine/Source/Programs/Horde/Horde.Build/Agents/Fleet/JobQueueStrategy.cs
--- a/Engine/Source/Programs/Horde/Horde.Build/Agents/Fleet/JobQueueStrategy.cs
+++ b/Engine/Source/Programs/Horde/Horde.Build/Agents/Fleet/JobQueueStrategy.cs
@@ -223,12 +223,12 @@ namespace Horde.Build.Agents.Fleet
 			{
 				int additionalAgentCount = (int)Math.Ceiling(queueSize * Settings.ScaleOutFactor);
 				int desiredAgentCount = agents.Count + additionalAgentCount;
-				return new PoolSizeResult(pool, agents, desiredAgentCount, status);
+				return new PoolSizeResult(agents.Count, desiredAgentCount, status);
 			}
 			else
 			{
 				int desiredAgentCount = (int)(agents.Count * Settings.ScaleInFactor);
-				return new PoolSizeResult(pool, agents, desiredAgentCount, status);
+				return new PoolSizeResult(agents.Count, desiredAgentCount, status);
 			}
 		}
 	}
diff --git a/Engine/Source/Programs/Horde/Horde.Build/Agents/Fleet/LeaseUtilizationStrategy.cs b/Engine/Source/Programs/Horde/Horde.Build/Agents/Fleet/LeaseUtilizationStrategy.cs
--- a/Engine/Source/Programs/Horde/Horde.Build/Agents/Fleet/LeaseUtilizationStrategy.cs
+++ b/Engine/Source/Programs/Horde/Horde.Build/Agents/Fleet/LeaseUtilizationStrategy.cs
@@ -276,7 +276,7 @@ namespace Horde.Build.Agents.Fleet
 				["NumReserveAgents"] = Settings.NumReserveAgents,
 			};
 
-			return new PoolSizeResult(pool, agents, desiredAgentCount, status); 
+			return new PoolSizeResult(agents.Count, desiredAgentCount, status); 
 		}
 		
 		/// <summary>
