Horde: Add support for supplying an extra agent count once pool size has been calculated

diff --git a/Engine/Source/Programs/Horde/Horde.Build.Tests/Fleet/FleetServiceTest.cs b/Engine/Source/Programs/Horde/Horde.Build.Tests/Fleet/FleetServiceTest.cs
--- a/Engine/Source/Programs/Horde/Horde.Build.Tests/Fleet/FleetServiceTest.cs
+++ b/Engine/Source/Programs/Horde/Horde.Build.Tests/Fleet/FleetServiceTest.cs
@@ -223,6 +223,14 @@ namespace Horde.Build.Tests.Fleet
 			Assert.AreEqual(30, ((LeaseUtilizationStrategy)s).Settings.NumSamplesForResult);
 		}
 		
+		[TestMethod]
+		public async Task CreateStrategyWithExtraAgentCount()
+		{
+			IPoolSizeStrategy s = await CreateStrategy(new PoolSizeStrategyInfo(PoolSizeStrategy.NoOp, null, "{}", 39));
+			PoolSizeResult result = await s.CalculatePoolSizeAsync(null!, new List<IAgent>());
+			Assert.AreEqual(39, result.DesiredAgentCount);
+		}
+		
 		[TestMethod]
 		public async Task EmptyOrInvalidJsonConfig()
 		{
diff --git a/Engine/Source/Programs/Horde/Horde.Build.Tests/PoolsControllerTest.cs b/Engine/Source/Programs/Horde/Horde.Build.Tests/PoolsControllerTest.cs
--- a/Engine/Source/Programs/Horde/Horde.Build.Tests/PoolsControllerTest.cs
+++ b/Engine/Source/Programs/Horde/Horde.Build.Tests/PoolsControllerTest.cs
@@ -33,7 +33,7 @@ namespace Horde.Build.Tests
 		        Name = "Pool1",
 		        ScaleOutCooldown = 111,
 		        ScaleInCooldown = 222,
-		        SizeStrategies = new List<PoolSizeStrategyMessage> { new () { Type = PoolSizeStrategy.JobQueue, Condition = "dayOfWeek == 'monday'", Config = @"{""ScaleOutFactor"": 22.0, ""ScaleInFactor"": 33.0}" }},
+		        SizeStrategies = new List<PoolSizeStrategyMessage> { new () { Type = PoolSizeStrategy.JobQueue, Condition = "dayOfWeek == 'monday'", Config = @"{""ScaleOutFactor"": 22.0, ""ScaleInFactor"": 33.0}", ExtraAgentCount = 567 }},
 		        SizeStrategy = PoolSizeStrategy.JobQueue,
 		        JobQueueSettings = new JobQueueSettingsMessage(new JobQueueSettings(0.35, 0.85)),
 		        FleetManagers = new List<FleetManagerMessage> { new () { Type = FleetManagerType.AwsReuse, Condition = "dayOfWeek == 'monday'", Config = "{}" }},
@@ -52,6 +52,7 @@ namespace Horde.Build.Tests
 	        Assert.AreEqual(request.SizeStrategies[0].Type, pool.SizeStrategies[0].Type);
 	        Assert.AreEqual(request.SizeStrategies[0].Condition!.Text, pool.SizeStrategies[0].Condition!.Text);
 	        Assert.AreEqual(request.SizeStrategies[0].Config, pool.SizeStrategies[0].Config);
+	        Assert.AreEqual(request.SizeStrategies[0].ExtraAgentCount, pool.SizeStrategies[0].ExtraAgentCount);
 	        Assert.AreEqual(1, request.FleetManagers.Count);
 	        Assert.AreEqual(request.FleetManagers[0].Type, pool.FleetManagers[0].Type);
 	        Assert.AreEqual(request.FleetManagers[0].Condition!.Text, pool.FleetManagers[0].Condition!.Text);
@@ -68,7 +69,7 @@ namespace Horde.Build.Tests
 		        Name = "Pool1Modified",
 		        SizeStrategies = new List<PoolSizeStrategyMessage>
 		        {
-			        new () { Type = PoolSizeStrategy.JobQueue, Condition = "true", Config = "{\"someConfig\": 123}"}
+			        new () { Type = PoolSizeStrategy.JobQueue, Condition = "true", Config = "{\"someConfig\": 123}", ExtraAgentCount = 123}
 		        },
 		        FleetManagers = new List<FleetManagerMessage>
 		        {
@@ -89,6 +90,7 @@ namespace Horde.Build.Tests
 	        Assert.AreEqual(PoolSizeStrategy.JobQueue, response.SizeStrategies[0].Type);
 	        Assert.AreEqual("true", response.SizeStrategies[0].Condition!.Text);
 	        Assert.AreEqual("{\"someConfig\": 123}", response.SizeStrategies[0].Config);
+	        Assert.AreEqual(123, response.SizeStrategies[0].ExtraAgentCount);
 	        
 	        Assert.AreEqual(1, response.FleetManagers.Count);
 	        Assert.AreEqual(FleetManagerType.AwsReuse, response.FleetManagers[0].Type);
diff --git a/Engine/Source/Programs/Horde/Horde.Build/Agents/Fleet/FleetService.cs b/Engine/Source/Programs/Horde/Horde.Build/Agents/Fleet/FleetService.cs
--- a/Engine/Source/Programs/Horde/Horde.Build/Agents/Fleet/FleetService.cs
+++ b/Engine/Source/Programs/Horde/Horde.Build/Agents/Fleet/FleetService.cs
@@ -337,11 +337,13 @@ namespace Horde.Build.Agents.Fleet
 						{
 							case PoolSizeStrategy.JobQueue:
 								JobQueueSettings jqSettings = DeserializeConfig<JobQueueSettings>(info.Config);
-								return new JobQueueStrategy(_jobCollection, _graphCollection, _streamService, _clock, _cache, jqSettings);
+								JobQueueStrategy jqStrategy = new (_jobCollection, _graphCollection, _streamService, _clock, _cache, jqSettings);
+								return info.ExtraAgentCount != 0 ? new ExtraAgentCountStrategy(jqStrategy, info.ExtraAgentCount) : jqStrategy;
 							
 							case PoolSizeStrategy.LeaseUtilization:
 								LeaseUtilizationSettings luSettings = DeserializeConfig<LeaseUtilizationSettings>(info.Config);
-								return new LeaseUtilizationStrategy(_agentCollection, _poolCollection, _leaseCollection, _clock, _cache, luSettings);
+								LeaseUtilizationStrategy luStrategy = new (_agentCollection, _poolCollection, _leaseCollection, _clock, _cache, luSettings);
+								return info.ExtraAgentCount != 0 ? new ExtraAgentCountStrategy(luStrategy, info.ExtraAgentCount) : luStrategy;
 							
 							// Disabled until moved to a separate factory class as FleetService should not contain AWS-specific classes
 							// case PoolSizeStrategy.ComputeQueueAwsMetric:
@@ -349,7 +351,8 @@ namespace Horde.Build.Agents.Fleet
 							// 	return new ComputeQueueAwsMetricStrategy(_awsCloudWatch, _computeService, cqamSettings);
 							
 							case PoolSizeStrategy.NoOp:
-								return new NoOpPoolSizeStrategy();
+								NoOpPoolSizeStrategy noStrategy = new ();
+								return info.ExtraAgentCount != 0 ? new ExtraAgentCountStrategy(noStrategy, info.ExtraAgentCount) : noStrategy;
 							
 							default:
 								throw new ArgumentException("Invalid pool size strategy type " + info.Type);
diff --git a/Engine/Source/Programs/Horde/Horde.Build/Agents/Fleet/IPoolSizeStrategy.cs b/Engine/Source/Programs/Horde/Horde.Build/Agents/Fleet/IPoolSizeStrategy.cs
--- a/Engine/Source/Programs/Horde/Horde.Build/Agents/Fleet/IPoolSizeStrategy.cs
+++ b/Engine/Source/Programs/Horde/Horde.Build/Agents/Fleet/IPoolSizeStrategy.cs
@@ -124,4 +124,35 @@ namespace Horde.Build.Agents.Fleet
 		/// <inheritdoc/>
 		public string Name { get; } = "NoOp";
 	}
+	
+	/// <summary>
+	/// Pool size strategy wrapping a normal strategy and
+	/// applies an extra agent count to the desired agent count.
+	/// </summary>
+	public class ExtraAgentCountStrategy : IPoolSizeStrategy
+	{
+		private readonly IPoolSizeStrategy _backingStrategy;
+		private readonly int _extraAgentCount;
+
+		/// <summary>
+		/// Constructor
+		/// </summary>
+		/// <param name="backingStrategy">Strategy to be wrapped</param>
+		/// <param name="extraAgentCount">Extra count to apply</param>
+		public ExtraAgentCountStrategy(IPoolSizeStrategy backingStrategy, int extraAgentCount)
+		{
+			_backingStrategy = backingStrategy;
+			_extraAgentCount = extraAgentCount;
+		}
+
+		/// <inheritdoc/>
+		public async Task<PoolSizeResult> CalculatePoolSizeAsync(IPool pool, List<IAgent> agents)
+		{
+			PoolSizeResult result = await _backingStrategy.CalculatePoolSizeAsync(pool, agents);
+			return new PoolSizeResult(result.CurrentAgentCount, result.DesiredAgentCount + _extraAgentCount, result.Status);
+		}
+
+		/// <inheritdoc/>
+		public string Name => _backingStrategy.Name;
+	}
 }
diff --git a/Engine/Source/Programs/Horde/Horde.Build/Agents/Pools/IPool.cs b/Engine/Source/Programs/Horde/Horde.Build/Agents/Pools/IPool.cs
--- a/Engine/Source/Programs/Horde/Horde.Build/Agents/Pools/IPool.cs
+++ b/Engine/Source/Programs/Horde/Horde.Build/Agents/Pools/IPool.cs
@@ -32,6 +32,11 @@ namespace Horde.Build.Agents.Pools
 		/// Configuration for the strategy, serialized as JSON
 		/// </summary>
 		public string Config { get; set; } = "";
+		
+		/// <summary>
+		/// Integer to add after pool size has been calculated. Can also be negative.
+		/// </summary>
+		public int ExtraAgentCount { get; set; } = 0;
 
 		/// <summary>
 		/// Empty constructor for JSON serialization
@@ -43,7 +48,7 @@ namespace Horde.Build.Agents.Pools
 		/// <summary>
 		/// Constructor
 		/// </summary>
-		public PoolSizeStrategyInfo(PoolSizeStrategy type, Condition? condition, string? config)
+		public PoolSizeStrategyInfo(PoolSizeStrategy type, Condition? condition, string? config, int extraAgentCount = 0)
 		{
 			config = String.IsNullOrWhiteSpace(config) ? "{}" : config;
 
@@ -54,6 +59,7 @@ namespace Horde.Build.Agents.Pools
 			Type = type;
 			Condition = condition;
 			Config = config;
+			ExtraAgentCount = extraAgentCount;
 		}
 	}
 	
diff --git a/Engine/Source/Programs/Horde/Horde.Build/Agents/Pools/PoolMessages.cs b/Engine/Source/Programs/Horde/Horde.Build/Agents/Pools/PoolMessages.cs
--- a/Engine/Source/Programs/Horde/Horde.Build/Agents/Pools/PoolMessages.cs
+++ b/Engine/Source/Programs/Horde/Horde.Build/Agents/Pools/PoolMessages.cs
@@ -26,6 +26,11 @@ namespace Horde.Build.Agents.Pools
 		/// </summary>
 		public string Config { get; set; } = "";
 		
+		/// <summary>
+		/// Integer to add after pool size has been calculated. Can also be negative.
+		/// </summary>
+		public int ExtraAgentCount { get; set; } = 0;
+		
 		/// <summary>
 		/// Empty constructor for JSON serialization
 		/// </summary>
@@ -42,6 +47,7 @@ namespace Horde.Build.Agents.Pools
 			Type = strategy.Type;
 			Condition = strategy.Condition;
 			Config = strategy.Config;
+			ExtraAgentCount = strategy.ExtraAgentCount;
 		}
 		
 		/// <summary>
@@ -50,7 +56,7 @@ namespace Horde.Build.Agents.Pools
 		/// <returns></returns>
 		public PoolSizeStrategyInfo Convert()
 		{
-			return new PoolSizeStrategyInfo(Type, Condition, Config);
+			return new PoolSizeStrategyInfo(Type, Condition, Config, ExtraAgentCount);
 		}
 	}
 	
