Horde: Make parsing of pool sizing config case-insensitive

diff --git a/Engine/Source/Programs/Horde/Horde.Build.Tests/Fleet/FleetServiceTest.cs b/Engine/Source/Programs/Horde/Horde.Build.Tests/Fleet/FleetServiceTest.cs
--- a/Engine/Source/Programs/Horde/Horde.Build.Tests/Fleet/FleetServiceTest.cs
+++ b/Engine/Source/Programs/Horde/Horde.Build.Tests/Fleet/FleetServiceTest.cs
@@ -257,6 +257,15 @@ namespace Horde.Build.Tests.Fleet
 			Assert.AreEqual(0.9, ((JobQueueStrategy)s).Settings.ScaleInFactor);
 		}
 		
+		[TestMethod]
+		public async Task ConfigHandlesMixedCase()
+		{
+			IPoolSizeStrategy s = await CreateStrategy(new PoolSizeStrategyInfo(PoolSizeStrategy.JobQueue, null, "{\"scaleOutFactor\": 100, \"SCALEINFACTOR\": 200}"));
+			Assert.AreEqual(typeof(JobQueueStrategy), s.GetType());
+			Assert.AreEqual(100.0, ((JobQueueStrategy)s).Settings.ScaleOutFactor);
+			Assert.AreEqual(200.0, ((JobQueueStrategy)s).Settings.ScaleInFactor);
+		}
+		
 		[TestMethod]
 		public async Task CreateFromEmptyStrategyList()
 		{
diff --git a/Engine/Source/Programs/Horde/Horde.Build/Agents/Fleet/FleetService.cs b/Engine/Source/Programs/Horde/Horde.Build/Agents/Fleet/FleetService.cs
--- a/Engine/Source/Programs/Horde/Horde.Build/Agents/Fleet/FleetService.cs
+++ b/Engine/Source/Programs/Horde/Horde.Build/Agents/Fleet/FleetService.cs
@@ -381,7 +381,7 @@ namespace Horde.Build.Agents.Fleet
 		private static T DeserializeConfig<T>(string json)
 		{
 			json = String.IsNullOrEmpty(json) ? "{}" : json;
-			T? config = JsonSerializer.Deserialize<T>(json);
+			T? config = JsonSerializer.Deserialize<T>(json, new JsonSerializerOptions { PropertyNameCaseInsensitive = true });
 			if (config == null) throw new ArgumentException("Unable to deserialize config: " + json);
 			return config;
 		}
