Horde: Remove unused feature flags

These have been fully tested and rolled out.

diff --git a/Engine/Source/Programs/Horde/Horde.Build/Logs/LogFileService.cs b/Engine/Source/Programs/Horde/Horde.Build/Logs/LogFileService.cs
--- a/Engine/Source/Programs/Horde/Horde.Build/Logs/LogFileService.cs
+++ b/Engine/Source/Programs/Horde/Horde.Build/Logs/LogFileService.cs
@@ -1318,16 +1318,8 @@ namespace Horde.Build.Logs
 				await _builder.CompleteChunkAsync(logId, offset);
 			}
 
-			if (_settings.Value.FeatureFlags.LimitConcurrentLogChunkWriting)
-			{
-				// Flush all the chunks and await completion instead of running them async
-				await WriteCompleteChunksV2Async(flushChunks, true, cancellationToken);
-			}
-			else
-			{
-				// Add tasks for flushing all the chunks
-				WriteCompleteChunks(flushChunks, true);
-			}
+			// Flush all the chunks and await completion instead of running them async
+			await WriteCompleteChunksV2Async(flushChunks, true, cancellationToken);
 		}
 
 		/// <summary>
diff --git a/Engine/Source/Programs/Horde/Horde.Build/ServerSettings.cs b/Engine/Source/Programs/Horde/Horde.Build/ServerSettings.cs
--- a/Engine/Source/Programs/Horde/Horde.Build/ServerSettings.cs
+++ b/Engine/Source/Programs/Horde/Horde.Build/ServerSettings.cs
@@ -170,16 +170,6 @@ namespace Horde.Build
 	/// </summary>
 	public class FeatureFlagSettings
 	{
-		/// <summary>
-		/// Use new auth config for custom auth servers in Okta
-		/// </summary>
-		public bool AuthSettingsV2 { get; set; } = false;
-		
-		/// <summary>
-		/// Limit concurrent log chunk writes and await them to reduce mem and I/O usage
-		/// </summary>
-		public bool LimitConcurrentLogChunkWriting { get; set; } = false;
-
 		/// <summary>
 		/// Whether to use the new log storage backend
 		/// </summary>
diff --git a/Engine/Source/Programs/Horde/Horde.Build/Startup.cs b/Engine/Source/Programs/Horde/Horde.Build/Startup.cs
--- a/Engine/Source/Programs/Horde/Horde.Build/Startup.cs
+++ b/Engine/Source/Programs/Horde/Horde.Build/Startup.cs
@@ -546,11 +546,7 @@ namespace Horde.Build
 						{
 							options.Authority = settings.OidcAuthority;
 							options.ClientId = settings.OidcClientId;
-
-							if (settings.FeatureFlags.AuthSettingsV2)
-							{
-								options.Scope.Remove("groups");								
-							}
+							options.Scope.Remove("groups");								
 
 							if (!String.IsNullOrEmpty(settings.OidcSigninRedirect))
 							{
