Horde: Fix accidental change to do conforms without have-table

This should be set to true until syncs without have-table are more battletested.

diff --git a/Engine/Source/Programs/Horde/Horde.Agent/Leases/Handlers/ConformHandler.cs b/Engine/Source/Programs/Horde/Horde.Agent/Leases/Handlers/ConformHandler.cs
--- a/Engine/Source/Programs/Horde/Horde.Agent/Leases/Handlers/ConformHandler.cs
+++ b/Engine/Source/Programs/Horde/Horde.Agent/Leases/Handlers/ConformHandler.cs
@@ -46,7 +46,7 @@ namespace Horde.Agent.Leases.Handlers
 				// Run the conform task
 				if (_settings.Executor.Equals(PerforceExecutor.Name, StringComparison.OrdinalIgnoreCase) && _settings.PerforceExecutor.RunConform)
 				{
-					await PerforceExecutor.ConformAsync(session.WorkingDir, pendingWorkspaces, removeUntrackedFiles, false, conformLogger, cancellationToken);
+					await PerforceExecutor.ConformAsync(session.WorkingDir, pendingWorkspaces, removeUntrackedFiles, true, conformLogger, cancellationToken);
 				}
 				else
 				{
