Horde: Fix submit failure retry loop never terminating.

diff --git a/Engine/Source/Programs/Horde/Horde.Build/Perforce/PerforceService.cs b/Engine/Source/Programs/Horde/Horde.Build/Perforce/PerforceService.cs
--- a/Engine/Source/Programs/Horde/Horde.Build/Perforce/PerforceService.cs
+++ b/Engine/Source/Programs/Horde/Horde.Build/Perforce/PerforceService.cs
@@ -921,11 +921,12 @@ namespace Horde.Build.Perforce
 				string workspaceFilePath = $"//{perforce.Client!.Name}/{filePath.TrimStart('/')}";
 				string diskFilePath = $"{perforce.Client!.Root + filePath.TrimStart('/')}";
 
-				int attempt = 1;
+				int attempt = 0;
 				const int MaxAttempts = 5;
 
 				for(; ;)
 				{
+					attempt++;
 					await ResetClientAsync(perforce, cancellationToken);
 
 					string? depotPath = null;
@@ -1011,7 +1012,7 @@ namespace Horde.Build.Perforce
 					{
 						if (attempt < MaxAttempts)
 						{
-							_logger.LogWarning(ex, "Unable to submit new changelist (file: {File}, depotPath: {DepotPath}, description: \"{Description}\", attempt: {Attempt}/{MaxAttempts}) - {Message}", filePath, depotPath, description, attempt, MaxAttempts, ex.Message);
+							_logger.LogWarning(ex, "Unable to submit new changelist (file: {File}, depotPath: {DepotPath}, description: \"{Description}\", attempt: {Attempt}/{MaxAttempts}, error: {Message}", filePath, depotPath, description, attempt, MaxAttempts, ex.Message);
 						}
 						else
 						{
