Horde: Allow cancellation to break out of the retry loop for submitting a new change at the start of each job.

diff --git a/Engine/Source/Programs/Horde/Horde.Build/Perforce/PerforceService.cs b/Engine/Source/Programs/Horde/Horde.Build/Perforce/PerforceService.cs
--- a/Engine/Source/Programs/Horde/Horde.Build/Perforce/PerforceService.cs
+++ b/Engine/Source/Programs/Horde/Horde.Build/Perforce/PerforceService.cs
@@ -920,6 +920,7 @@ namespace Horde.Build.Perforce
 			{
 				string workspaceFilePath = $"//{perforce.Client!.Name}/{filePath.TrimStart('/')}";
 				string diskFilePath = $"{perforce.Client!.Root + filePath.TrimStart('/')}";
+				_logger.LogInformation("Attempting to create new changelist by submitting {WorkspaceFilePath} (local: {DiskFilePath})", workspaceFilePath, diskFilePath);
 
 				int attempt = 0;
 				const int MaxAttempts = 5;
@@ -1008,6 +1009,10 @@ namespace Horde.Build.Perforce
 						_logger.LogInformation("Submitted new changelist with {DepotPath}: CL {Change}", depotPath, submit.SubmittedChangeNumber);
 						return submit.SubmittedChangeNumber;
 					}
+					catch (OperationCanceledException)
+					{
+						throw;
+					}
 					catch (Exception ex)
 					{
 						if (attempt < MaxAttempts)
