Horde: Add method setting for workspaces

Allows for specifying which method to use when materializing files for jobs on agents. Currently defaults to ManagedWorkspace when null.

diff --git a/Engine/Source/Programs/Horde/Horde.Build/Agents/IAgent.cs b/Engine/Source/Programs/Horde/Horde.Build/Agents/IAgent.cs
--- a/Engine/Source/Programs/Horde/Horde.Build/Agents/IAgent.cs
+++ b/Engine/Source/Programs/Horde/Horde.Build/Agents/IAgent.cs
@@ -66,6 +66,11 @@ namespace Horde.Build.Agents
 		/// Whether to use an incremental workspace
 		/// </summary>
 		public bool Incremental { get; set; }
+		
+		/// <summary>
+		/// Method to use when syncing/materializing data from Perforce
+		/// </summary>
+		public string? Method { get; set; }
 
 		/// <summary>
 		/// Constructor
@@ -76,7 +81,8 @@ namespace Horde.Build.Agents
 		/// <param name="stream">The stream to sync</param>
 		/// <param name="view">Custom view for the workspace</param>
 		/// <param name="incremental">Whether to use an incremental workspace</param>
-		public AgentWorkspace(string? cluster, string? userName, string identifier, string stream, List<string>? view, bool incremental)
+		/// <param name="method">Method to use when syncing/materializing data from Perforce</param>
+		public AgentWorkspace(string? cluster, string? userName, string identifier, string stream, List<string>? view, bool incremental, string? method)
 		{
 			if (!String.IsNullOrEmpty(cluster))
 			{
@@ -90,6 +96,7 @@ namespace Horde.Build.Agents
 			Stream = stream;
 			View = view;
 			Incremental = incremental;
+			Method = method;
 		}
 
 		/// <summary>
@@ -97,7 +104,7 @@ namespace Horde.Build.Agents
 		/// </summary>
 		/// <param name="workspace">RPC message to construct from</param>
 		public AgentWorkspace(HordeCommon.Rpc.Messages.AgentWorkspace workspace)
-			: this(workspace.ConfiguredCluster, workspace.ConfiguredUserName, workspace.Identifier, workspace.Stream, (workspace.View.Count > 0) ? workspace.View.ToList() : null, workspace.Incremental)
+			: this(workspace.ConfiguredCluster, workspace.ConfiguredUserName, workspace.Identifier, workspace.Stream, (workspace.View.Count > 0) ? workspace.View.ToList() : null, workspace.Incremental, workspace.Method)
 		{
 		}
 
@@ -687,7 +694,7 @@ namespace Horde.Build.Agents
 			{
 				if (autoSdk.Stream != null && autoSdk.Properties.All(x => agent.Properties.Contains(x)))
 				{
-					return new AgentWorkspace(cluster.Name, autoSdk.UserName, autoSdk.Name ?? "AutoSDK", autoSdk.Stream!, null, true);
+					return new AgentWorkspace(cluster.Name, autoSdk.UserName, autoSdk.Name ?? "AutoSDK", autoSdk.Stream!, null, true, null);
 				}
 			}
 			return null;
diff --git a/Engine/Source/Programs/Horde/Horde.Build/Streams/IStream.cs b/Engine/Source/Programs/Horde/Horde.Build/Streams/IStream.cs
--- a/Engine/Source/Programs/Horde/Horde.Build/Streams/IStream.cs
+++ b/Engine/Source/Programs/Horde/Horde.Build/Streams/IStream.cs
@@ -211,7 +211,7 @@ namespace Horde.Build.Streams
 			if (agentType.Workspace == null)
 			{
 				// Use the default settings (fast switching workspace, clean 
-				workspace = (new AgentWorkspace(null, null, GetDefaultWorkspaceIdentifier(stream), stream.Name, null, false), true);
+				workspace = (new AgentWorkspace(null, null, GetDefaultWorkspaceIdentifier(stream), stream.Name, null, false, null), true);
 				return true;
 			}
 			else
@@ -240,7 +240,7 @@ namespace Horde.Build.Streams
 				}
 
 				// Create the new workspace
-				workspace = (new AgentWorkspace(workspaceType.Cluster, workspaceType.UserName, identifier, workspaceType.Stream ?? stream.Name, workspaceType.View, workspaceType.Incremental), workspaceType.UseAutoSdk);
+				workspace = (new AgentWorkspace(workspaceType.Cluster, workspaceType.UserName, identifier, workspaceType.Stream ?? stream.Name, workspaceType.View, workspaceType.Incremental, workspaceType.Method), workspaceType.UseAutoSdk);
 				return true;
 			}
 		}
diff --git a/Engine/Source/Programs/Horde/Horde.Build/Streams/StreamConfig.cs b/Engine/Source/Programs/Horde/Horde.Build/Streams/StreamConfig.cs
--- a/Engine/Source/Programs/Horde/Horde.Build/Streams/StreamConfig.cs
+++ b/Engine/Source/Programs/Horde/Horde.Build/Streams/StreamConfig.cs
@@ -502,6 +502,11 @@ namespace Horde.Build.Streams
 		/// Whether to use the AutoSDK
 		/// </summary>
 		public bool UseAutoSdk { get; set; } = true;
+
+		/// <summary>
+		/// Method to use when syncing/materializing data from Perforce
+		/// </summary>
+		public string? Method { get; set; } = null;
 	}
 
 	/// <summary>
