Horde: Add a service to copy Perforce data into bundles. Unlike ReplicationService, this does not run as a hosted task, but provides explicit methods for creating snapshots.

diff --git a/Engine/Source/Programs/Horde/Horde.Build/Commands/Bundles/ExtractCommand.cs b/Engine/Source/Programs/Horde/Horde.Build/Commands/Bundles/ExtractCommand.cs
--- a/Engine/Source/Programs/Horde/Horde.Build/Commands/Bundles/ExtractCommand.cs
+++ b/Engine/Source/Programs/Horde/Horde.Build/Commands/Bundles/ExtractCommand.cs
@@ -6,7 +6,7 @@ using System.Threading.Tasks;
 using EpicGames.Core;
 using EpicGames.Horde.Storage;
 using EpicGames.Horde.Storage.Nodes;
-using Horde.Build.Perforce;
+using Horde.Build.Storage;
 using Microsoft.Extensions.Caching.Memory;
 using Microsoft.Extensions.Configuration;
 using Microsoft.Extensions.DependencyInjection;
@@ -17,8 +17,11 @@ namespace Horde.Build.Commands.Bundles
 	[Command("bundle", "extract", "Extracts data from a bundle to the local hard drive")]
 	internal class ExtractCommand : Command
 	{
+		[CommandLine("-Ns=")]
+		public NamespaceId NamespaceId { get; set; } = Namespace.Perforce;
+
 		[CommandLine("-Ref=")]
-		public RefName RefName { get; set; } = new RefName("default-ref");
+		public RefName RefName { get; set; } = new RefName("ue5-main");
 
 		[CommandLine("-OutputDir=", Required = true)]
 		public DirectoryReference OutputDir { get; set; } = null!;
@@ -32,14 +35,22 @@ namespace Horde.Build.Commands.Bundles
 			_loggerProvider = loggerProvider;
 		}
 
-		public override Task<int> ExecuteAsync(ILogger logger)
+		public override async Task<int> ExecuteAsync(ILogger logger)
 		{
 			using ServiceProvider serviceProvider = Startup.CreateServiceProvider(_configuration, _loggerProvider);
 
-			IStorageClient store = serviceProvider.GetRequiredService<IStorageClient<ReplicationService>>();
+			StorageService storageService = serviceProvider.GetRequiredService<StorageService>();
+
+			IStorageClient store = await storageService.GetClientAsync(NamespaceId, default);
+
 			TreeReader reader = new TreeReader(store, serviceProvider.GetRequiredService<IMemoryCache>(), serviceProvider.GetRequiredService<ILogger<ExtractCommand>>());
+			CommitNode commit = await reader.ReadNodeAsync<CommitNode>(RefName);
+			logger.LogInformation("Extracting {Number}: {Description} to {OutputDir}", commit.Number, (commit.Message ?? String.Empty).Replace("\n", "\\n", StringComparison.Ordinal), OutputDir);
+
+			DirectoryNode contents = await commit.Contents.ExpandAsync(reader);
+			await contents.CopyToDirectoryAsync(reader, OutputDir.ToDirectoryInfo(), logger, CancellationToken.None);
 
-			throw new NotImplementedException();
+			return 0;
 		}
 	}
 }
diff --git a/Engine/Source/Programs/Horde/Horde.Build/Commands/Bundles/PerforceCommand.cs b/Engine/Source/Programs/Horde/Horde.Build/Commands/Bundles/PerforceCommand.cs
--- a/Engine/Source/Programs/Horde/Horde.Build/Commands/Bundles/PerforceCommand.cs
+++ b/Engine/Source/Programs/Horde/Horde.Build/Commands/Bundles/PerforceCommand.cs
@@ -1,18 +1,13 @@
 // Copyright Epic Games, Inc. All Rights Reserved.
 
 using System;
-using System.Collections.Generic;
-using System.Linq;
 using System.Threading;
 using System.Threading.Tasks;
 using EpicGames.Core;
-using EpicGames.Horde.Storage;
-using EpicGames.Horde.Storage.Nodes;
+using Horde.Build.Configuration;
 using Horde.Build.Perforce;
-using Horde.Build.Storage;
 using Horde.Build.Streams;
 using Horde.Build.Utilities;
-using Microsoft.Extensions.Caching.Memory;
 using Microsoft.Extensions.Configuration;
 using Microsoft.Extensions.DependencyInjection;
 using Microsoft.Extensions.Logging;
@@ -30,27 +25,6 @@ namespace Horde.Build.Commands.Bundles
 		[CommandLine(Required = true)]
 		public int Change { get; set; }
 
-		[CommandLine]
-		public int BaseChange { get; set; }
-
-		[CommandLine]
-		public int Count { get; set; } = 1;
-
-		[CommandLine]
-		public bool Content { get; set; }
-
-		[CommandLine]
-		public bool Compact { get; set; } = true;
-
-		[CommandLine]
-		public string Filter { get; set; } = "...";
-
-		[CommandLine]
-		public bool RevisionsOnly { get; set; } = false;
-
-		[CommandLine]
-		public DirectoryReference? OutputDir { get; set; }
-
 		readonly IConfiguration _configuration;
 		readonly ILoggerProvider _loggerProvider;
 
@@ -64,13 +38,11 @@ namespace Horde.Build.Commands.Bundles
 		{
 			using ServiceProvider serviceProvider = Startup.CreateServiceProvider(_configuration, _loggerProvider);
 
-			ICommitService commitService = serviceProvider.GetRequiredService<ICommitService>();
-			ReplicationService replicationService = serviceProvider.GetRequiredService<ReplicationService>();
-			IStreamCollection streamCollection = serviceProvider.GetRequiredService<IStreamCollection>();
-			StorageService storageService = serviceProvider.GetRequiredService<StorageService>();
+			ConfigUpdateService configUpdateService = serviceProvider.GetRequiredService<ConfigUpdateService>();
+			await configUpdateService.ForceUpdateAsync(CancellationToken.None);
 
-			IStorageClient storage = await storageService.GetClientAsync(Namespace.Perforce, CancellationToken.None);
-			TreeReader reader = new TreeReader(storage, serviceProvider.GetRequiredService<IMemoryCache>(), serviceProvider.GetRequiredService<ILogger<PerforceCommand>>());
+			PerforceReplicator replicator = serviceProvider.GetRequiredService<PerforceReplicator>();
+			IStreamCollection streamCollection = serviceProvider.GetRequiredService<IStreamCollection>();
 
 			IStream? stream = await streamCollection.GetAsync(new StreamId(StreamId));
 			if (stream == null)
@@ -78,7 +50,10 @@ namespace Horde.Build.Commands.Bundles
 				throw new FatalErrorException($"Stream '{stream}' not found");
 			}
 
-			throw new NotImplementedException();
+			PerforceReplicationOptions options = new PerforceReplicationOptions();
+			await replicator.WriteAsync(stream, Change, options, default);
+
+			return 0;
 		}
 	}
 }
diff --git a/Engine/Source/Programs/Horde/Horde.Build/Configuration/ConfigUpdateService.cs b/Engine/Source/Programs/Horde/Horde.Build/Configuration/ConfigUpdateService.cs
--- a/Engine/Source/Programs/Horde/Horde.Build/Configuration/ConfigUpdateService.cs
+++ b/Engine/Source/Programs/Horde/Horde.Build/Configuration/ConfigUpdateService.cs
@@ -84,6 +84,15 @@ namespace Horde.Build.Configuration
 			_logger = logger;
 		}
 
+		/// <summary>
+		/// Triggers an immediate update of the config service
+		/// </summary>
+		/// <param name="cancellationToken"></param>
+		public async ValueTask ForceUpdateAsync(CancellationToken cancellationToken)
+		{
+			await TickLeaderAsync(cancellationToken);
+		}
+
 		/// <inheritdoc/>
 		public async Task StartAsync(CancellationToken cancellationToken)
 		{
diff --git a/Engine/Source/Programs/Horde/Horde.Build/Perforce/PerforceReplicator.cs b/Engine/Source/Programs/Horde/Horde.Build/Perforce/PerforceReplicator.cs
new file mode 100644
--- /dev/null
+++ b/Engine/Source/Programs/Horde/Horde.Build/Perforce/PerforceReplicator.cs
@@ -0,0 +1,580 @@
+// Copyright Epic Games, Inc. All Rights Reserved.
+
+using System;
+using System.Collections.Generic;
+using System.Diagnostics;
+using System.Linq;
+using System.Threading;
+using System.Threading.Tasks;
+using EpicGames.Core;
+using EpicGames.Horde.Storage;
+using EpicGames.Horde.Storage.Nodes;
+using EpicGames.Perforce;
+using Horde.Build.Storage;
+using Horde.Build.Streams;
+using Horde.Build.Utilities;
+using Microsoft.Extensions.Caching.Memory;
+using Microsoft.Extensions.Logging;
+
+namespace Horde.Build.Perforce
+{
+	using StreamId = StringId<IStream>;
+
+	/// <summary>
+	/// Options for replicating commits
+	/// </summary>
+	class PerforceReplicationOptions
+	{
+		public bool IncludeContent { get; set; }
+		public TreeOptions TreeOptions { get; set; } = new TreeOptions();
+		public ChunkingOptions ChunkingOptions { get; set; } = new ChunkingOptions();
+		public RefOptions RefOptions { get; set; } = new RefOptions();
+	}
+
+	/// <summary>
+	/// Replicates commits from Perforce into Horde's internal storage
+	/// </summary>
+	class PerforceReplicator
+	{
+		[TreeNode("{8C874966-2E89-4273-A4AB-AC9F5491C86D}")]
+		class SyncNode : TreeNode
+		{
+			public int Change { get; }
+			public int ParentChange { get; }
+			public TreeNodeRef<DirectoryNode> Contents { get; }
+			public List<Utf8String> Paths { get; }
+
+			public SyncNode(int number, int parentNumber, TreeNodeRef<DirectoryNode> contents)
+			{
+				Change = number;
+				ParentChange = parentNumber;
+				Contents = contents;
+				Paths = new List<Utf8String>();
+			}
+
+			public SyncNode(ITreeNodeReader reader)
+			{
+				Change = (int)reader.ReadUnsignedVarInt();
+				ParentChange = (int)reader.ReadUnsignedVarInt();
+				Contents = reader.ReadRef<DirectoryNode>();
+				Paths = reader.ReadList(() => reader.ReadUtf8String());
+			}
+
+			/// <inheritdoc/>
+			public override void Serialize(ITreeNodeWriter writer)
+			{
+				writer.WriteUnsignedVarInt(Change);
+				writer.WriteUnsignedVarInt(ParentChange);
+				writer.WriteRef(Contents);
+				writer.WriteList(Paths, x => writer.WriteUtf8String(x));
+			}
+
+			/// <inheritdoc/>
+			public override IEnumerable<TreeNodeRef> EnumerateRefs()
+			{
+				yield return Contents;
+			}
+		}
+
+		[DebuggerDisplay("{_path}")]
+		class DirectoryToSync
+		{
+			public readonly Utf8String Path;
+			public readonly Dictionary<Utf8String, long> FileNameToSize;
+			public long _size;
+
+			public DirectoryToSync(Utf8String path, Utf8StringComparer comparer)
+			{
+				Path = path;
+				FileNameToSize = new Dictionary<Utf8String, long>(comparer);
+			}
+		}
+
+		class FileWriter
+		{
+			readonly Utf8String _path;
+			FileNode _node;
+			readonly long _size;
+
+			public Utf8String Path => _path;
+			public FileNode Node => _node;
+			public long Size => _size;
+
+			public FileWriter(Utf8String path, long size)
+			{
+				_path = path;
+				_node = new LeafFileNode();
+				_size = size;
+			}
+
+			public async Task AppendAsync(ReadOnlyMemory<byte> data, ChunkingOptions options, TreeWriter writer, CancellationToken cancellationToken)
+			{
+				_node = await _node.AppendAsync(data, options, writer, cancellationToken);
+			}
+		}
+
+		class ReplicationClient
+		{
+			public PerforceSettings Settings { get; }
+			public string ClusterName { get; }
+			public InfoRecord ServerInfo { get; }
+			public ClientRecord Client { get; }
+			public int Change { get; set; }
+
+			public ReplicationClient(PerforceSettings settings, string clusterName, InfoRecord serverInfo, ClientRecord client, int change)
+			{
+				Settings = settings;
+				ClusterName = clusterName;
+				ServerInfo = serverInfo;
+				Client = client;
+				Change = change;
+			}
+		}
+
+		readonly Dictionary<StreamId, ReplicationClient> _cachedPerforceClients = new Dictionary<StreamId, ReplicationClient>();
+
+		readonly IPerforceService _perforceService;
+		readonly StorageService _storageService;
+		readonly IMemoryCache _memoryCache;
+		readonly ILogger _logger;
+
+		/// <summary>
+		/// Constructor
+		/// </summary>
+		public PerforceReplicator(IPerforceService perforceService, StorageService storageService, IMemoryCache memoryCache, ILogger<PerforceReplicator> logger)
+		{
+			_perforceService = perforceService;
+			_storageService = storageService;
+			_memoryCache = memoryCache;
+			_logger = logger;
+		}
+
+		/// <summary>
+		/// Gets the ref name for a given stream
+		/// </summary>
+		/// <param name="streamId">The stream to get a ref for</param>
+		/// <returns>Ref name for the stream</returns>
+		public static RefName GetRefName(StreamId streamId) => new RefName(streamId.ToString());
+
+		/// <summary>
+		/// Gets the ref name for a given stream
+		/// </summary>
+		/// <param name="streamId">The stream to get a ref for</param>
+		/// <returns>Ref name for the stream</returns>
+		static RefName GetIncrementalRefName(StreamId streamId) => new RefName($"{streamId}/incremental");
+
+		/// <summary>
+		/// Replicates a change to storage
+		/// </summary>
+		/// <param name="stream">Stream to replicate data from</param>
+		/// <param name="change">Changelist to replicate</param>
+		/// <param name="options">Options for replication</param>
+		/// <param name="cancellationToken">Cancellation token for the operation</param>
+		public async Task WriteAsync(IStream stream, int change, PerforceReplicationOptions options, CancellationToken cancellationToken)
+		{
+			IStorageClient store = await _storageService.GetClientAsync(Namespace.Perforce, cancellationToken);
+
+			TreeReaderOptions readerOptions = new TreeReaderOptions();
+			readerOptions.Types.Add(typeof(SyncNode));
+
+			TreeReader reader = new TreeReader(store, _memoryCache, readerOptions, _logger);
+
+			// Find the parent node
+			RefName refName = GetRefName(stream.Id);
+			CommitNode? parent = await FindParentAsync(reader, refName, change, cancellationToken);
+			TreeNodeRef<CommitNode>? parentRef = (parent == null) ? null : new TreeNodeRef<CommitNode>(parent);
+			int parentChange = parent?.Number ?? 0;
+
+			// Read the current incremental state or create a new node to track the incremental state
+			RefName incRefName = GetIncrementalRefName(stream.Id);
+			SyncNode? syncNode = await reader.TryReadNodeAsync<SyncNode>(incRefName, cancellationToken: cancellationToken);
+			if (syncNode == null || syncNode.Change != change || syncNode.ParentChange != parentChange)
+			{
+				syncNode = new SyncNode(change, parentChange, parent?.Contents ?? new TreeNodeRef<DirectoryNode>(new DirectoryNode()));
+			}
+			DirectoryNode root = await syncNode.Contents.ExpandAsync(reader, cancellationToken);
+
+			// Create a client to replicate from this stream
+			ReplicationClient clientInfo = await FindOrAddReplicationClientAsync(stream);
+
+			// Connect to the server and flush the workspace
+			using IPerforceConnection perforce = await PerforceConnection.CreateAsync(clientInfo.Settings, _logger);
+
+			// Apply all the updates
+			_logger.LogInformation("Syncing client {Client} from changelist {BaseChange} to {Change}", clientInfo.Client.Name, parentChange, change);
+			await FlushWorkspaceAsync(clientInfo, perforce, parentChange);
+			clientInfo.Change = -1;
+
+			Utf8String clientRoot = new Utf8String(clientInfo.Client.Root);
+			string queryPath = $"//{clientInfo.Client.Name}/...";
+
+			// Replay the files that have already been synced
+			foreach (Utf8String path in syncNode.Paths)
+			{
+				string flushPath = $"//{clientInfo.Client.Name}/{path}@{change}";
+				_logger.LogInformation("Flushing {FlushPath}", flushPath);
+				await perforce.SyncQuietAsync(SyncOptions.Force | SyncOptions.KeepWorkspaceFiles, -1, flushPath, cancellationToken);
+			}
+
+			// Add the root directory from the filter to the list of files to sync. This prevents traversing above it.
+			Dictionary<Utf8String, DirectoryToSync> pathToDirectory = new Dictionary<Utf8String, DirectoryToSync>(clientInfo.ServerInfo.Utf8PathComparer);
+
+			// Do a sync preview to find everything that's left, and sort the remaining list of paths
+			await foreach (PerforceResponse<SyncRecord> response in perforce.StreamCommandAsync<SyncRecord>("sync", new[] { "-n" }, new string[] { $"{queryPath}@{change}" }, null, cancellationToken))
+			{
+				PerforceError? error = response.Error;
+				if (error != null)
+				{
+					_logger.LogWarning("Perforce: {Message}", error.Data);
+					continue;
+				}
+
+				Utf8String path = response.Data.Path;
+				if (!path.StartsWith(clientRoot, Utf8StringComparer.Ordinal))
+				{
+					throw new ArgumentException($"Unable to make path {path} relative to client root {clientRoot}");
+				}
+
+				int fileIdx = path.Length;
+				while (fileIdx > 0 && path[fileIdx - 1] != '/' && path[fileIdx - 1] != '\\')
+				{
+					fileIdx--;
+				}
+
+				Utf8String directoryPath = path.Slice(clientRoot.Length, fileIdx - clientRoot.Length);
+
+				DirectoryToSync directory = FindOrAddDirectoryTree(pathToDirectory, directoryPath, clientInfo.ServerInfo.Utf8PathComparer);
+				directory.FileNameToSize.Add(path.Slice(fileIdx).Clone(), response.Data.FileSize);
+				directory._size += response.Data.FileSize;
+			}
+
+			// Sort the directories by name to ensure that they are consistent between runs
+			List<DirectoryToSync> directories = pathToDirectory.Values.OrderBy(x => x.Path, clientInfo.ServerInfo.Utf8PathComparer).ToList();
+
+			// Output some stats for the sync
+			long totalSize = directories.Sum(x => x._size);
+			_logger.LogInformation("Total sync size: {Size:n1}mb", totalSize / (1024.0 * 1024.0));
+
+			// Create the tree writer
+			TreeWriter writer = new TreeWriter(store, refName, options.TreeOptions);
+
+			// Sync incrementally
+			long syncedSize = 0;
+			while (directories.Count > 0)
+			{
+				// Save the incremental state
+				if (syncedSize > 0)
+				{
+					Stopwatch flushTimer = Stopwatch.StartNew();
+					await writer.WriteAsync(incRefName, syncNode, cancellationToken: cancellationToken);
+					flushTimer.Stop();
+				}
+
+				// Find the next paths to sync
+				const long MaxBatchSize = 1L * 1024 * 1024 * 1024;
+
+				int dirIdx = directories.Count - 1;
+				long size = directories[dirIdx]._size;
+
+				for (; dirIdx > 0; dirIdx--)
+				{
+					long nextSize = size + directories[dirIdx - 1]._size;
+					if (size > 0 && nextSize > MaxBatchSize)
+					{
+						break;
+					}
+					size = nextSize;
+				}
+
+				syncedSize += size;
+				double syncPct = (totalSize == 0) ? 100.0 : (syncedSize * 100.0) / totalSize;
+				_logger.LogInformation("Syncing {StreamId} to {Change} [{SyncPct:n1}%] ({Size:n1}mb)", stream.Id, change, syncPct, size / (1024.0 * 1024.0));
+
+				// Copy them to a separate list and remove any redundant paths
+				List<string> syncPaths = new List<string>();
+				for (int idx = dirIdx; idx < directories.Count; idx++)
+				{
+					Utf8String basePath = directories[idx].Path;
+					syncPaths.Add($"//{clientInfo.Client.Name}/{basePath}...@{change}");
+
+					long dirSize = directories[idx]._size;
+					while (idx + 1 < directories.Count && directories[idx + 1].Path.StartsWith(basePath, clientInfo.ServerInfo.Utf8PathComparer))
+					{
+						dirSize += directories[idx + 1]._size;
+						idx++;
+					}
+
+					_logger.LogInformation("  {Directory} ({Size:n1}mb)", directories[idx].Path + "...", dirSize / (1024.0 * 1024.0));
+				}
+
+				Stopwatch syncTimer = Stopwatch.StartNew();
+				Stopwatch processTimer = new Stopwatch();
+				Stopwatch gcTimer = new Stopwatch();
+
+				Dictionary<int, FileWriter> handles = new Dictionary<int, FileWriter>();
+				await foreach (PerforceResponse response in perforce.StreamCommandAsync("sync", Array.Empty<string>(), syncPaths, null, typeof(SyncRecord), true, default))
+				{
+					PerforceError? error = response.Error;
+					if (error != null)
+					{
+						if (error.Generic == PerforceGenericCode.Empty)
+						{
+							continue;
+						}
+						else
+						{
+							throw new ReplicationException($"Perforce error while replicating content - {error}");
+						}
+					}
+
+					processTimer.Start();
+
+					PerforceIo? io = response.Io;
+					if (io != null)
+					{
+						if (io.Command == PerforceIoCommand.Open)
+						{
+							Utf8String file = GetClientRelativePath(io.Payload, clientInfo.Client.Root);
+							int offset = GetFileOffset(file);
+
+							long fileSize = 0;
+							if (!pathToDirectory.TryGetValue(file.Substring(0, offset), out DirectoryToSync? directory))
+							{
+								throw new ReplicationException($"Unable to find directory for {file}");
+							}
+							if (!directory.FileNameToSize.TryGetValue(file.Substring(offset), out fileSize))
+							{
+								throw new ReplicationException($"Unable to find file entry for {file}");
+							}
+
+							handles[io.File] = new FileWriter(file, fileSize);
+						}
+						else if (io.Command == PerforceIoCommand.Write)
+						{
+							FileWriter file = handles[io.File];
+							await file.AppendAsync(io.Payload, options.ChunkingOptions, writer, cancellationToken);
+						}
+						else if (io.Command == PerforceIoCommand.Close)
+						{
+							FileWriter? file = null;
+							if (handles.Remove(io.File, out file))
+							{
+								FileNode node = file.Node;
+								if (node.Length != file.Size)
+								{
+									throw new ReplicationException($"Invalid size for replicated file '{file.Path}'. Expected {file.Size}, got {node.Length}.");
+								}
+
+								FileEntry entry = await root.AddFileByPathAsync(reader, file.Path, FileEntryFlags.None, node, cancellationToken);
+								await writer.WriteAsync(entry, cancellationToken);
+							}
+						}
+						else if (io.Command == PerforceIoCommand.Unlink)
+						{
+							Utf8String file = GetClientRelativePath(io.Payload, clientInfo.Client.Root);
+							await root.DeleteFileByPathAsync(reader, file, cancellationToken);
+						}
+						else
+						{
+							_logger.LogWarning("Unhandled command code {Code}", io.Command);
+						}
+					}
+
+					processTimer.Stop();
+				}
+
+				TimeSpan stallTime = (perforce as NativePerforceConnection)?.StallTime ?? TimeSpan.Zero;
+				_logger.LogInformation("Completed batch in {TimeSeconds:n1}s ({ProcessTimeSeconds:n1}s processing, {StallTimeSeconds:n1}s stalled, {Throughput:n1}mb/s)", syncTimer.Elapsed.TotalSeconds, processTimer.Elapsed.TotalSeconds, stallTime.TotalSeconds, size / (1024.0 * 1024.0 * syncTimer.Elapsed.TotalSeconds));
+
+				// Combine all the existing sync paths together with a new wildcard.
+				while (directories.Count > dirIdx)
+				{
+					Utf8String nextPath = directories[^1].Path;
+					if (syncNode.Paths.Count > 0)
+					{
+						Utf8String lastPath = syncNode.Paths[^1];
+						for (int endIdx = 0; endIdx < lastPath.Length; endIdx++)
+						{
+							if (lastPath[endIdx] == '/')
+							{
+								Utf8String prefix = lastPath.Substring(0, endIdx + 1);
+								if (!nextPath.StartsWith(prefix, clientInfo.ServerInfo.Utf8PathComparer))
+								{
+									// Remove any paths that start with this prefix
+									while (syncNode.Paths.Count > 0 && syncNode.Paths[^1].StartsWith(prefix, clientInfo.ServerInfo.Utf8PathComparer))
+									{
+										syncNode.Paths.RemoveAt(syncNode.Paths.Count - 1);
+									}
+
+									// Replace it with a wildcard
+									syncNode.Paths.Add(prefix + "...");
+									break;
+								}
+							}
+						}
+					}
+					syncNode.Paths.Add(nextPath + "...");
+					directories.RemoveAt(directories.Count - 1);
+				}
+			}
+
+			// Create the commit node
+			ChangeRecord changeRecord = await perforce.GetChangeAsync(GetChangeOptions.None, change, cancellationToken);
+			CommitNode commitNode = new CommitNode(change, parentRef, changeRecord.User ?? "Unknown", changeRecord.Description ?? String.Empty, changeRecord.Date, new TreeNodeRef<DirectoryNode>(root));
+			await writer.WriteAsync(refName, commitNode, options.RefOptions, cancellationToken); 
+		}
+
+		static async Task<CommitNode?> FindParentAsync(TreeReader reader, RefName refName, int change, CancellationToken cancellationToken)
+		{
+			TreeNodeRef<CommitNode>? parentRef = await reader.TryReadNodeRefAsync<CommitNode>(refName, cancellationToken: cancellationToken);
+			while (parentRef != null)
+			{
+				CommitNode parent = await parentRef.ExpandCopyAsync(reader, cancellationToken);
+				if (parent.Number < change)
+				{
+					return parent;
+				}
+				parentRef = parent.Parent;
+			}
+			return null;
+		}
+
+		static int GetFileOffset(Utf8String path)
+		{
+			int fileIdx = path.Length;
+			while (fileIdx > 0 && path[fileIdx - 1] != '/' && path[fileIdx - 1] != '\\')
+			{
+				fileIdx--;
+			}
+			return fileIdx;
+		}
+
+		static DirectoryToSync FindOrAddDirectoryTree(Dictionary<Utf8String, DirectoryToSync> pathToDirectory, Utf8String directoryPath, Utf8StringComparer comparer)
+		{
+			DirectoryToSync? directory;
+			if (!pathToDirectory.TryGetValue(directoryPath, out directory))
+			{
+				// Add a new path
+				Utf8String normalizedPath = NormalizePathSeparators(directoryPath);
+				directory = new DirectoryToSync(normalizedPath, comparer);
+				pathToDirectory.Add(directoryPath, directory);
+
+				// Also add all the parent directories. This makes the logic for combining wildcards simpler.
+				for (int idx = normalizedPath.Length - 2; idx > 0; idx--)
+				{
+					if (normalizedPath[idx] == '/')
+					{
+						Utf8String parentDirectoryPath = directoryPath.Substring(0, idx + 1);
+						if (pathToDirectory.ContainsKey(parentDirectoryPath))
+						{
+							break;
+						}
+						else
+						{
+							pathToDirectory.Add(parentDirectoryPath, new DirectoryToSync(normalizedPath.Substring(0, idx + 1), comparer));
+						}
+					}
+				}
+			}
+			return directory;
+		}
+
+		static Utf8String NormalizePathSeparators(Utf8String path)
+		{
+			byte[] newPath = new byte[path.Length];
+			for (int idx = 0; idx < path.Length; idx++)
+			{
+				newPath[idx] = (path[idx] == '\\') ? (byte)'/' : path[idx];
+			}
+			return new Utf8String(newPath);
+		}
+
+		async Task FlushWorkspaceAsync(ReplicationClient clientInfo, IPerforceConnection perforce, int change)
+		{
+			if (clientInfo.Change != change)
+			{
+				clientInfo.Change = -1;
+				if (change == 0)
+				{
+					_logger.LogInformation("Flushing have table for {Client}", clientInfo.Client.Name);
+					await perforce.SyncQuietAsync(SyncOptions.Force | SyncOptions.KeepWorkspaceFiles, -1, $"//{clientInfo.Client.Name}/...#0");
+				}
+				else
+				{
+					_logger.LogInformation("Flushing have table for {Client} to change {Change}", clientInfo.Client.Name, change);
+					await perforce.SyncQuietAsync(SyncOptions.Force | SyncOptions.KeepWorkspaceFiles, -1, $"//{clientInfo.Client.Name}/...@{change}");
+				}
+				clientInfo.Change = change;
+			}
+		}
+
+		static Utf8String GetClientRelativePath(ReadOnlyMemory<byte> data, Utf8String clientRoot)
+		{
+			int length = data.Span.IndexOf((byte)0);
+			if (length != -1)
+			{
+				data = data.Slice(0, length);
+			}
+
+			Utf8String path = new Utf8String(data);
+			if (!path.StartsWith(clientRoot, Utf8StringComparer.Ordinal))
+			{
+				throw new ArgumentException($"Unable to make path {path} relative to client root {clientRoot}");
+			}
+
+			return path.Substring(clientRoot.Length).Clone();
+		}
+
+		async Task<ReplicationClient?> FindReplicationClientAsync(IStream stream)
+		{
+			ReplicationClient? clientInfo;
+			if (_cachedPerforceClients.TryGetValue(stream.Id, out clientInfo))
+			{
+				if (!String.Equals(clientInfo.ClusterName, stream.Config.ClusterName, StringComparison.Ordinal) && String.Equals(clientInfo.Client.Stream, stream.Name, StringComparison.Ordinal))
+				{
+					PerforceSettings serverSettings = new PerforceSettings(clientInfo.Settings);
+					serverSettings.ClientName = null;
+
+					using IPerforceConnection perforce = await PerforceConnection.CreateAsync(_logger);
+					await perforce.DeleteClientAsync(DeleteClientOptions.None, clientInfo.Client.Name);
+
+					_cachedPerforceClients.Remove(stream.Id);
+					clientInfo = null;
+				}
+			}
+			return clientInfo;
+		}
+
+		async Task<ReplicationClient> FindOrAddReplicationClientAsync(IStream stream)
+		{
+			ReplicationClient? clientInfo = await FindReplicationClientAsync(stream);
+			if (clientInfo == null)
+			{
+				using IPerforceConnection? perforce = await _perforceService.ConnectAsync(stream.Config.ClusterName);
+				if (perforce == null)
+				{
+					throw new PerforceException($"Unable to create connection to Perforce server");
+				}
+
+				InfoRecord serverInfo = await perforce.GetInfoAsync(InfoOptions.ShortOutput);
+
+				ClientRecord newClient = new ClientRecord($"Horde.Build_Rep_{serverInfo.ClientHost}_{stream.Id}", perforce.Settings.UserName, "/p4/");
+				newClient.Description = "Created to mirror Perforce content to Horde Storage";
+				newClient.Owner = perforce.Settings.UserName;
+				newClient.Host = serverInfo.ClientHost;
+				newClient.Stream = stream.Config.ReplicationStream ?? stream.Name;
+				newClient.Type = "readonly";
+				await perforce.CreateClientAsync(newClient);
+				_logger.LogInformation("Created client {ClientName} for {StreamName}", newClient.Name, stream.Name);
+
+				PerforceSettings settings = new PerforceSettings(perforce.Settings);
+				settings.ClientName = newClient.Name;
+				settings.PreferNativeClient = true;
+
+				clientInfo = new ReplicationClient(settings, stream.Config.ClusterName, serverInfo, newClient, -1);
+				_cachedPerforceClients.Add(stream.Id, clientInfo);
+			}
+			return clientInfo;
+		}
+	}
+}
\ No newline at end of file
diff --git a/Engine/Source/Programs/Horde/Horde.Build/Startup.cs b/Engine/Source/Programs/Horde/Horde.Build/Startup.cs
--- a/Engine/Source/Programs/Horde/Horde.Build/Startup.cs
+++ b/Engine/Source/Programs/Horde/Horde.Build/Startup.cs
@@ -403,6 +403,7 @@ namespace Horde.Build
 			});
 			services.AddSingleton<ICommitService, CommitService>();
 			services.AddSingleton<IClock, Clock>();
+			services.AddSingleton<ConfigUpdateService>();
 			services.AddSingleton<IDowntimeService, DowntimeService>();
 			services.AddSingleton<IssueService>();
 			services.AddSingleton<JobService>();
@@ -421,6 +422,7 @@ namespace Horde.Build
 				services.AddSingleton<PerforceService>();
 				services.AddSingleton<IPerforceService>(sp => sp.GetRequiredService<PerforceService>());
 			}
+			services.AddSingleton<PerforceReplicator>();
 
 			services.AddSingleton<PerforceLoadBalancer>();
 			services.AddSingleton<PoolService>();
@@ -546,7 +548,7 @@ namespace Horde.Build
 						{
 							options.Authority = settings.OidcAuthority;
 							options.ClientId = settings.OidcClientId;
-							options.Scope.Remove("groups");								
+								options.Scope.Remove("groups");								
 
 							if (!String.IsNullOrEmpty(settings.OidcSigninRedirect))
 							{
@@ -638,7 +640,7 @@ namespace Horde.Build
 				{
 					services.AddHostedService(provider => provider.GetRequiredService<SlackNotificationSink>());
 				}
-				services.AddHostedService<ConfigUpdateService>();
+				services.AddHostedService(provider => provider.GetRequiredService<ConfigUpdateService>());
 				services.AddHostedService<TelemetryService>();
 				services.AddHostedService(provider => provider.GetRequiredService<DeviceService>());				
 			}
@@ -1079,10 +1081,20 @@ namespace Horde.Build
 			return LogEventLevel.Information;
 		}
 
+		class HostApplicationLifetime : IHostApplicationLifetime
+		{
+			public CancellationToken ApplicationStarted => CancellationToken.None;
+			public CancellationToken ApplicationStopped => CancellationToken.None;
+			public CancellationToken ApplicationStopping => CancellationToken.None;
+
+			public void StopApplication() { }
+		}
+
 		public static void AddServices(IServiceCollection serviceCollection, IConfiguration configuration)
 		{
 			Startup startup = new Startup(configuration);
 			startup.ConfigureServices(serviceCollection);
+			serviceCollection.AddSingleton<IHostApplicationLifetime, HostApplicationLifetime>();
 		}
 
 		public static IServiceProvider CreateServiceProvider(IConfiguration configuration)
diff --git a/Engine/Source/Programs/Shared/EpicGames.Horde/Storage/TreeReader.cs b/Engine/Source/Programs/Shared/EpicGames.Horde/Storage/TreeReader.cs
--- a/Engine/Source/Programs/Shared/EpicGames.Horde/Storage/TreeReader.cs
+++ b/Engine/Source/Programs/Shared/EpicGames.Horde/Storage/TreeReader.cs
@@ -52,6 +52,7 @@ namespace EpicGames.Horde.Storage
 		/// </summary>
 		public List<Type> Types { get; } = new List<Type>
 		{
+			typeof(Nodes.CommitNode),
 			typeof(Nodes.DirectoryNode),
 			typeof(Nodes.LeafFileNode),
 			typeof(Nodes.InteriorFileNode),
