Horde: Fix incorrect default value of feature flag for workspace-less syncs

diff --git a/Engine/Source/Programs/Horde/Horde.Agent/Execution/PerforceExecutor.cs b/Engine/Source/Programs/Horde/Horde.Agent/Execution/PerforceExecutor.cs
--- a/Engine/Source/Programs/Horde/Horde.Agent/Execution/PerforceExecutor.cs
+++ b/Engine/Source/Programs/Horde/Horde.Agent/Execution/PerforceExecutor.cs
@@ -34,7 +34,7 @@ namespace Horde.Agent.Execution
 
 		protected WorkspaceInfo? _autoSdkWorkspace;
 		protected WorkspaceInfo _workspace;
-		private bool _syncWithoutHaveTable = false;
+		private bool _useHaveTable = true;
 
 		public PerforceExecutor(ISession session, string jobId, string batchId, string agentTypeName, AgentWorkspace? autoSdkWorkspaceInfo, AgentWorkspace workspaceInfo, DirectoryReference rootDir, IHttpClientFactory httpClientFactory)
 			: base(session, jobId, batchId, agentTypeName, httpClientFactory)
@@ -55,7 +55,7 @@ namespace Horde.Agent.Execution
 			{
 				logger.LogInformation("Syncing without have table");
 				_additionalArguments.RemoveAt(index);
-				_syncWithoutHaveTable = true;
+				_useHaveTable = false;
 			}
 
 			// Setup and sync the autosdk workspace
@@ -115,7 +115,7 @@ namespace Horde.Agent.Execution
 
 				// Sync the workspace
 				int syncPreflightChange = (_job.ClonedPreflightChange != 0) ? _job.ClonedPreflightChange : _job.PreflightChange;
-				await _workspace.SyncAsync(_job.Change, syncPreflightChange, null, _syncWithoutHaveTable, cancellationToken);
+				await _workspace.SyncAsync(_job.Change, syncPreflightChange, null, _useHaveTable, cancellationToken);
 
 				// Remove any cached BuildGraph manifests
 				DirectoryReference manifestDir = DirectoryReference.Combine(_workspace.WorkspaceDir, "Engine", "Saved", "BuildGraph");
