Horde: Fix missing agent rate and add test

diff --git a/Engine/Source/Programs/Horde/Horde.Build.Tests/AgentServiceTest.cs b/Engine/Source/Programs/Horde/Horde.Build.Tests/AgentServiceTest.cs
--- a/Engine/Source/Programs/Horde/Horde.Build.Tests/AgentServiceTest.cs
+++ b/Engine/Source/Programs/Horde/Horde.Build.Tests/AgentServiceTest.cs
@@ -8,6 +8,7 @@ using System.Threading.Tasks;
 using EpicGames.Core;
 using Horde.Build.Agents;
 using Horde.Build.Jobs;
+using Horde.Build.Server;
 using Horde.Build.Utilities;
 using HordeCommon;
 using Microsoft.AspNetCore.Mvc;
@@ -53,6 +54,29 @@ public class AgentServiceTest : TestSetup
 		await Clock.AdvanceAsync(TimeSpan.FromMinutes(20));
 		Assert.IsFalse(AgentService.AuthorizeSession(agent, GetUser(agent)));
 	}
+	
+	[TestMethod]
+	public async Task GetAgentRateTest()
+	{
+		IAgent agent1 = await AgentService.CreateAgentAsync("agent1", true, null);
+		IAgent agent2 = await AgentService.CreateAgentAsync("agent2", true, null);
+		await AgentService.CreateSessionAsync(agent1, AgentStatus.Ok, new List<string>() { "aws-instance-type=c5.24xlarge", "osfamily=windows" }, new Dictionary<string, int>(), "test");
+		await AgentService.CreateSessionAsync(agent2, AgentStatus.Ok, new List<string>() { "aws-instance-type=c4.4xLARge", "osfamily=WinDowS" }, new Dictionary<string, int>(), "test");
+		
+		List<AgentRateConfig> agentRateConfigs = new()
+		{
+			new AgentRateConfig() { Condition = "aws-instance-type == 'c5.24xlarge' && osfamily == 'windows'", Rate = 200, },
+			new AgentRateConfig() { Condition = "aws-instance-type == 'c4.4xlarge' && osfamily == 'windows'", Rate = 300 }
+			
+		};
+		await AgentService.UpdateRateTableAsync(agentRateConfigs);
+		
+		double? rate1 = await AgentService.GetRateAsync(agent1.Id);
+		Assert.AreEqual(200, rate1!.Value, 0.1);
+		
+		double? rate2 = await AgentService.GetRateAsync(agent2.Id);
+		Assert.AreEqual(300, rate2!.Value, 0.1);
+	}
 
 	private static ClaimsPrincipal GetUser(IAgent agent)
 	{
diff --git a/Engine/Source/Programs/Horde/Horde.Build/Agents/AgentService.cs b/Engine/Source/Programs/Horde/Horde.Build/Agents/AgentService.cs
--- a/Engine/Source/Programs/Horde/Horde.Build/Agents/AgentService.cs
+++ b/Engine/Source/Programs/Horde/Horde.Build/Agents/AgentService.cs
@@ -917,7 +917,7 @@ namespace Horde.Build.Agents
 		/// <returns>Hourly rate of running the given agent</returns>
 		public async ValueTask<double?> GetRateAsync(AgentId agentId)
 		{
-			RedisStringKey<double> key = $"agent-rate/{agentId}";
+			RedisKey key = $"agent-rate/{agentId}";
 
 			// Try to get the current value
 			RedisValue value = await _redisService.GetDatabase().StringGetAsync(key);
