Horde: Allow jobs to be scheduled even if no agents are enabled or online

Still requires agents to be assigned to the pool that job batch is scheduled to run in.

diff --git a/Engine/Source/Programs/Horde/Horde.Build.Tests/JobTaskSourceTests.cs b/Engine/Source/Programs/Horde/Horde.Build.Tests/JobTaskSourceTests.cs
--- a/Engine/Source/Programs/Horde/Horde.Build.Tests/JobTaskSourceTests.cs
+++ b/Engine/Source/Programs/Horde/Horde.Build.Tests/JobTaskSourceTests.cs
@@ -81,12 +81,13 @@ namespace Horde.Build.Tests
 			
 			Assert.AreEqual(0, JobTaskSource.GetQueueForTesting().Count);
 			await JobTaskSource.TickAsync(CancellationToken.None);
-			Assert.AreEqual(0, JobTaskSource.GetQueueForTesting().Count);
+			Assert.AreEqual(1, JobTaskSource.GetQueueForTesting().Count);
 
-			IJob job = (await JobService.GetJobAsync(fixture.Job1.Id))!;
-			Assert.AreEqual(JobStepBatchError.NoAgentsOnline, job.Batches[0].Error);
-			
-			Assert.IsFalse(_eventReceived);
+			Assert.AreEqual(fixture.Job1.Id, JobTaskSource.GetQueueForTesting().Min!.Id.Item1);
+			Assert.AreEqual(JobStepBatchState.Ready, JobTaskSource.GetQueueForTesting().Min!.Batch.State);
+
+			Assert.IsTrue(_eventReceived);
+			Assert.IsFalse(_eventPoolHasAgentsOnline!.Value);
 		}
 		
 		[TestMethod]
diff --git a/Engine/Source/Programs/Horde/Horde.Build/Jobs/JobTaskSource.cs b/Engine/Source/Programs/Horde/Horde.Build/Jobs/JobTaskSource.cs
--- a/Engine/Source/Programs/Horde/Horde.Build/Jobs/JobTaskSource.cs
+++ b/Engine/Source/Programs/Horde/Horde.Build/Jobs/JobTaskSource.cs
@@ -10,6 +10,7 @@ using System.Text;
 using System.Threading;
 using System.Threading.Tasks;
 using EpicGames.Core;
+using EpicGames.Horde.Common;
 using Google.Protobuf;
 using Google.Protobuf.WellKnownTypes;
 using Horde.Build.Agents;
@@ -195,6 +196,7 @@ namespace Horde.Build.Jobs
 		readonly IPoolCollection _poolCollection;
 		readonly IUgsMetadataCollection _ugsMetadataCollection;
 		readonly PerforceLoadBalancer _perforceLoadBalancer;
+		readonly IClock _clock;
 		readonly IOptionsMonitor<ServerSettings> _settings;
 		readonly ILogger<JobTaskSource> _logger;
 		readonly ITicker _ticker;
@@ -214,9 +216,6 @@ namespace Horde.Build.Jobs
 		// During a background queue refresh operation, any updated batches are added to this dictionary for merging into the updated queue.
 		List<QueueItem>? _newQueueItemsDuringUpdate;
 
-		// Cache of pools
-		Dictionary<PoolId, IPool> _cachedPoolIdToInstance = new Dictionary<PoolId, IPool>();
-
 		// Cache of stream objects. Used to resolve agent types.
 		private Dictionary<StreamId, IStream> _streams = new Dictionary<StreamId, IStream>();
 
@@ -248,6 +247,7 @@ namespace Horde.Build.Jobs
 			_streamService = streamService;
 			_logFileService = logFileService;
 			_perforceLoadBalancer = perforceLoadBalancer;
+			_clock = clock;
 			_ticker = clock.AddTicker<JobTaskSource>(s_refreshInterval, TickAsync, logger);
 			_settings = settings;
 			_logger = logger;
@@ -305,6 +305,49 @@ namespace Horde.Build.Jobs
 			}
 		}
 
+		internal class PoolStatus
+		{
+			public readonly IPool Pool;
+			public readonly bool HasAgents;
+			public readonly bool HasEnabledAgents;
+			public readonly bool HasOnlineAgents;
+			public bool IsAutoScaled => Pool.EnableAutoscaling;
+
+			public PoolStatus(IPool pool, bool hasAgents, bool hasEnabledAgents, bool hasOnlineAgents)
+			{
+				Pool = pool;
+				HasAgents = hasAgents;
+				HasEnabledAgents = hasEnabledAgents;
+				HasOnlineAgents = hasOnlineAgents;
+			}
+		}
+
+		/// <summary>
+		/// Calculate status of agents for a pool
+		/// </summary>
+		/// <param name="utcNow">Current time</param>
+		/// <param name="pools">List of all available pools</param>
+		/// <param name="agents">List of all available agents</param>
+		/// <returns></returns>
+		internal static Dictionary<PoolId, PoolStatus> GetPoolStatus(DateTime utcNow, List<IPool> pools, List<IAgent> agents)
+		{
+			Dictionary<StringId<IPool>,PoolStatus> poolStatus = new ();
+
+			foreach (IPool pool in pools)
+			{
+				Condition poolCondition = pool.Condition ?? Condition.Parse("false");
+				List<IAgent> poolAgents = agents.Where(x => x.ExplicitPools.Contains(pool.Id) || x.SatisfiesCondition(poolCondition)).ToList();
+
+				bool hasAgents = poolAgents.Count > 0;
+				bool hasEnabledAgents = poolAgents.Any(x => x.Enabled);
+				bool hasOnlineAgents = poolAgents.Any(x => x.Enabled && x.IsSessionValid(utcNow));
+
+				poolStatus[pool.Id] = new PoolStatus(pool, hasAgents, hasEnabledAgents, hasOnlineAgents);
+			}
+
+			return poolStatus;
+		}
+		
 		/// <summary>
 		/// Background task
 		/// </summary>
@@ -323,55 +366,13 @@ namespace Horde.Build.Jobs
 			_streams = streamsList.ToDictionary(x => x.Id, x => x);
 
 			// Find all the pools which are valid (ie. have at least one online agent)
-			DateTime utcNow = DateTime.UtcNow;
 			List<IAgent> agents = await _agentsCollection.FindAsync();
 			List<IPool> pools = await _poolCollection.GetAsync();
-
-			// Find all the pools which are currently online
-			HashSet<PoolId> onlinePools = new HashSet<PoolId>(agents.Where(x => x.IsSessionValid(utcNow) && x.Enabled).SelectMany(x => x.ExplicitPools));
-			foreach (IPool pool in pools)
-			{
-				if (pool.Condition != null && !onlinePools.Contains(pool.Id) && agents.Any(x => x.IsSessionValid(utcNow) && x.SatisfiesCondition(pool.Condition) && x.Enabled))
-				{
-					onlinePools.Add(pool.Id);
-				}
-			}
-
-			// Find lists of valid pools and online pools
-			HashSet<PoolId> validPools = new HashSet<PoolId>(onlinePools.Union(agents.Where(x => !x.IsSessionValid(utcNow) || !x.Enabled).SelectMany(x => x.ExplicitPools)));
-			foreach (IPool pool in pools)
-			{
-				if (pool.Condition != null && !validPools.Contains(pool.Id) && agents.Any(x => !x.IsSessionValid(utcNow) && x.SatisfiesCondition(pool.Condition) && x.Enabled))
-				{
-					validPools.Add(pool.Id);
-				}
-			}
-
-			// Query all the current pools
-			_cachedPoolIdToInstance = pools.ToDictionary(x => x.Id, x => x);
+			Dictionary<PoolId, PoolStatus> poolStatus = GetPoolStatus(_clock.UtcNow, pools, agents);
 
 			// New list of queue items
 			SortedSet<QueueItem> newQueue = new SortedSet<QueueItem>(_queue.Comparer);
 			Dictionary<(JobId, SubResourceId), QueueItem> newBatchIdToQueueItem = new Dictionary<(JobId, SubResourceId), QueueItem>();
-			
-			// Returns true if agents are online and available for scheduling for a pool
-			bool IsPoolOnline(PoolId poolId)
-			{
-				return onlinePools.Contains(poolId);
-			}
-			
-			// Returns true if a pool can be auto-scaled
-			bool IsPoolAutoScaled(PoolId poolId)
-			{
-				IPool? pool = pools.Find(p => p.Id == poolId);
-				return validPools.Contains(poolId) && pool != null && pool.EnableAutoscaling;
-			}
-			
-			bool HasAgentsOnlineOrIsAutoScaled(PoolId poolId)
-			{
-				// If pool is auto-scaled, it will be considered online even if it has no agents online
-				return IsPoolOnline(poolId) || IsPoolAutoScaled(poolId);
-			}
 
 			// Query for a new list of jobs for the queue
 			List<IJob> newJobs = await _jobs.GetDispatchQueueAsync();
@@ -420,18 +421,14 @@ namespace Horde.Build.Jobs
 					{
 						newJob = await SkipBatchAsync(newJob, batch.Id, graph, JobStepBatchError.UnknownAgentType);
 					}
-					else if (!_cachedPoolIdToInstance.TryGetValue(agentType.Pool, out IPool? pool))
+					else if (!poolStatus.ContainsKey(agentType.Pool))
 					{
 						newJob = await SkipBatchAsync(newJob, batch.Id, graph, JobStepBatchError.UnknownPool);
 					}
-					else if (!validPools.Contains(agentType.Pool))
+					else if (!poolStatus[agentType.Pool].HasAgents)
 					{
 						newJob = await SkipBatchAsync(newJob, batch.Id, graph, JobStepBatchError.NoAgentsInPool);
 					}
-					else if (!HasAgentsOnlineOrIsAutoScaled(agentType.Pool))
-					{
-						newJob = await SkipBatchAsync(newJob, batch.Id, graph, JobStepBatchError.NoAgentsOnline);
-					}
 					else if (!stream.TryGetAgentWorkspace(agentType, out (AgentWorkspace, bool)? workspaceResult))
 					{
 						newJob = await SkipBatchAsync(newJob, batch.Id, graph, JobStepBatchError.UnknownWorkspace);
@@ -476,7 +473,7 @@ namespace Horde.Build.Jobs
 							IPool? newJobPool = pools.Find(p => p.Id == agentType.Pool);
 							if (newJobPool != null)
 							{
-								OnJobScheduled?.Invoke(newJobPool, IsPoolOnline(agentType.Pool), newJob, graph, batch.Id);
+								OnJobScheduled?.Invoke(newJobPool, poolStatus[agentType.Pool].HasOnlineAgents, newJob, graph, batch.Id);
 							}
 						}
 					}
