Horde: Re-enable replication service using new PerforceReplicator implementation.

diff --git a/Engine/Source/Programs/Horde/Horde.Build/Perforce/ReplicationService.cs b/Engine/Source/Programs/Horde/Horde.Build/Perforce/ReplicationService.cs
--- a/Engine/Source/Programs/Horde/Horde.Build/Perforce/ReplicationService.cs
+++ b/Engine/Source/Programs/Horde/Horde.Build/Perforce/ReplicationService.cs
@@ -1,11 +1,17 @@
 // Copyright Epic Games, Inc. All Rights Reserved.
 
 using System;
+using System.Collections.Generic;
+using System.Linq;
 using System.Threading;
 using System.Threading.Tasks;
-using Horde.Build.Server;
+using EpicGames.Horde.Storage;
+using EpicGames.Horde.Storage.Nodes;
+using Horde.Build.Storage;
 using Horde.Build.Streams;
 using Horde.Build.Utilities;
+using HordeCommon;
+using Microsoft.Extensions.Caching.Memory;
 using Microsoft.Extensions.Hosting;
 using Microsoft.Extensions.Logging;
 
@@ -28,22 +34,112 @@ namespace Horde.Build.Perforce
 	/// </summary>
 	sealed class ReplicationService : IHostedService
 	{
-		readonly RedisService _redisService;
+		readonly IStreamCollection _streamCollection;
+		readonly IPerforceService _perforceService;
+		readonly PerforceReplicator _replicator;
+		readonly StorageService _storageService;
+		readonly IMemoryCache _memoryCache;
+		readonly ITicker _ticker;
 		readonly ILogger _logger;
 
 		/// <summary>
 		/// Constructor
 		/// </summary>
-		public ReplicationService(RedisService redisService, ILogger<ReplicationService> logger)
+		public ReplicationService(IStreamCollection streamCollection, IPerforceService perforceService, PerforceReplicator replicator, StorageService storageService, IMemoryCache memoryCache, IClock clock, ILogger<ReplicationService> logger)
 		{
-			_redisService = redisService;
+			_streamCollection = streamCollection;
+			_perforceService = perforceService;
+			_replicator = replicator;
+			_storageService = storageService;
+			_memoryCache = memoryCache;
+			_ticker = clock.AddSharedTicker<ReplicationService>(TimeSpan.FromSeconds(20.0), TickSharedAsync, logger);
 			_logger = logger;
 		}
 
 		/// <inheritdoc/>
-		public Task StartAsync(CancellationToken cancellationToken) => Task.CompletedTask;
+		public async Task StartAsync(CancellationToken cancellationToken)
+		{
+			await _ticker.StartAsync();
+		}
 
 		/// <inheritdoc/>
-		public Task StopAsync(CancellationToken cancellationToken) => Task.CompletedTask;
+		public async Task StopAsync(CancellationToken cancellationToken)
+		{
+			await _ticker.StopAsync();
+		}
+
+		async ValueTask TickSharedAsync(CancellationToken cancellationToken)
+		{
+			Dictionary<StreamId, BackgroundTask> streamIdToTask = new Dictionary<StreamId, BackgroundTask>();
+			try
+			{
+				for (; ; )
+				{
+					List<IStream> streams = await _streamCollection.FindAllAsync();
+
+					HashSet<StreamId> removeStreams = new HashSet<StreamId>(streamIdToTask.Keys);
+					foreach (IStream stream in streams)
+					{
+						if (stream.Config.ReplicationMode == ContentReplicationMode.Full)
+						{
+							removeStreams.Remove(stream.Id);
+							if (!streamIdToTask.ContainsKey(stream.Id))
+							{
+								streamIdToTask.Add(stream.Id, BackgroundTask.StartNew(ctx => RunReplicationAsync(stream.Id, ctx)));
+							}
+						}
+					}
+
+					foreach (StreamId removeStreamId in removeStreams)
+					{
+						if (streamIdToTask.Remove(removeStreamId, out BackgroundTask? task))
+						{
+							await task.DisposeAsync();
+						}
+					}
+
+					await Task.Delay(TimeSpan.FromMinutes(1.0), cancellationToken);
+				}
+			}
+			finally
+			{
+				await Parallel.ForEachAsync(streamIdToTask.Values, (x, ctx) => x.DisposeAsync());
+			}
+		}
+
+		async Task RunReplicationAsync(StreamId streamId, CancellationToken cancellationToken)
+		{
+			IStream? stream = await _streamCollection.GetAsync(streamId);
+			if (stream == null)
+			{
+				return;
+			}
+
+			RefName refName = new RefName(stream.Id.ToString());
+
+			IStorageClientImpl store = await _storageService.GetClientAsync(Namespace.Perforce, cancellationToken);
+			TreeReader reader = new TreeReader(store, _memoryCache, _logger);
+
+			CommitNode? lastCommitNode = await reader.TryReadNodeAsync<CommitNode>(refName, cancellationToken: cancellationToken);
+			ICommitCollection commits = _perforceService.GetCommits(stream);
+
+			PerforceReplicationOptions options = new PerforceReplicationOptions();
+
+			ICommit commit;
+			if (lastCommitNode == null)
+			{
+				commit = await commits.GetLatestAsync(cancellationToken);
+			}
+			else
+			{
+				commit = await commits.SubscribeAsync(lastCommitNode.Number, cancellationToken: cancellationToken).FirstAsync(cancellationToken);
+			}
+
+			for (; ; )
+			{
+				await _replicator.WriteAsync(stream, commit.Number, options, cancellationToken);
+				commit = await commits.SubscribeAsync(commit.Number, cancellationToken: cancellationToken).FirstAsync(cancellationToken);
+			}
+		}
 	}
 }
diff --git a/Engine/Source/Programs/Horde/Horde.Build/Utilities/BackgroundTask.cs b/Engine/Source/Programs/Horde/Horde.Build/Utilities/BackgroundTask.cs
--- a/Engine/Source/Programs/Horde/Horde.Build/Utilities/BackgroundTask.cs
+++ b/Engine/Source/Programs/Horde/Horde.Build/Utilities/BackgroundTask.cs
@@ -9,14 +9,14 @@ namespace Horde.Build.Utilities
 	/// <summary>
 	/// Runs a task in the background and allows stopping it on demand
 	/// </summary>
-	public sealed class BackgroundTask : IDisposable
+	public sealed class BackgroundTask : IDisposable, IAsyncDisposable
 	{
 		readonly Func<CancellationToken, Task> _runTask;
 		Task _task = Task.CompletedTask;
 		readonly CancellationTokenSource _cancellationTokenSource = new CancellationTokenSource();
 
 		/// <summary>
-		/// Constructor
+		/// Constructor. Note that the task does not start until <see cref="Start"/> is called.
 		/// </summary>
 		/// <param name="runTask"></param>
 		public BackgroundTask(Func<CancellationToken, Task> runTask)
@@ -35,6 +35,25 @@ namespace Horde.Build.Utilities
 			_cancellationTokenSource.Dispose();
 		}
 
+		/// <inheritdoc/>
+		public async ValueTask DisposeAsync()
+		{
+			await StopAsync();
+			_cancellationTokenSource.Dispose();
+		}
+
+		/// <summary>
+		/// Creates and starts a new background task instance
+		/// </summary>
+		/// <param name="runTask"></param>
+		/// <returns></returns>
+		public static BackgroundTask StartNew(Func<CancellationToken, Task> runTask)
+		{
+			BackgroundTask task = new BackgroundTask(runTask);
+			task.Start();
+			return task;
+		}
+
 		/// <summary>
 		/// Starts the task
 		/// </summary>
@@ -53,13 +72,16 @@ namespace Horde.Build.Utilities
 		/// <returns></returns>
 		public async Task StopAsync()
 		{
-			try
-			{
-				_cancellationTokenSource.Cancel();
-				await _task;
-			}
-			catch (OperationCanceledException)
+			if (!_task.IsCompleted)
 			{
+				try
+				{
+					_cancellationTokenSource.Cancel();
+					await _task;
+				}
+				catch (OperationCanceledException)
+				{
+				}
 			}
 		}
 	}
