Horde: Missing files.

diff --git a/Engine/Source/Programs/Horde/Horde.Build/Storage/StorageService.cs b/Engine/Source/Programs/Horde/Horde.Build/Storage/StorageService.cs
--- a/Engine/Source/Programs/Horde/Horde.Build/Storage/StorageService.cs
+++ b/Engine/Source/Programs/Horde/Horde.Build/Storage/StorageService.cs
@@ -212,10 +212,10 @@ namespace Horde.Build.Storage
 			#region Refs
 
 			/// <inheritdoc/>
-			public override Task<RefTarget?> TryReadRefTargetAsync(RefName name, DateTime cacheTime = default, CancellationToken cancellationToken = default) => _outer.TryReadRefTargetAsync(NamespaceId, name, cacheTime, cancellationToken);
+			public override Task<NodeLocator> TryReadRefTargetAsync(RefName name, DateTime cacheTime = default, CancellationToken cancellationToken = default) => _outer.TryReadRefTargetAsync(NamespaceId, name, cacheTime, cancellationToken);
 
 			/// <inheritdoc/>
-			public override Task WriteRefTargetAsync(RefName name, RefTarget target, RefOptions? options = null, CancellationToken cancellationToken = default) => _outer.WriteRefTargetAsync(NamespaceId, name, target, options, cancellationToken);
+			public override Task WriteRefTargetAsync(RefName name, NodeLocator target, RefOptions? options = null, CancellationToken cancellationToken = default) => _outer.WriteRefTargetAsync(NamespaceId, name, target, options, cancellationToken);
 
 			/// <inheritdoc/>
 			public override Task DeleteRefAsync(RefName name, CancellationToken cancellationToken = default) => _outer.DeleteRefAsync(NamespaceId, name, cancellationToken);
@@ -310,9 +310,6 @@ namespace Horde.Build.Storage
 			[BsonElement("idx")]
 			public int ExportIdx { get; set; }
 
-			[BsonElement("hash")]
-			public IoHash Hash { get; set; }
-
 			[BsonElement("xa"), BsonIgnoreIfDefault]
 			public DateTime? ExpiresAtUtc { get; set; }
 
@@ -320,7 +317,7 @@ namespace Horde.Build.Storage
 			public TimeSpan? Lifetime { get; set; }
 
 			[BsonIgnore]
-			public RefTarget Target => new RefTarget(Hash, new NodeLocator(BlobLocator, ExportIdx));
+			public NodeLocator Target => new NodeLocator(BlobLocator, ExportIdx);
 
 			public RefInfo()
 			{
@@ -328,14 +325,13 @@ namespace Horde.Build.Storage
 				BlobLocator = BlobLocator.Empty;
 			}
 
-			public RefInfo(NamespaceId namespaceId, RefName name, RefTarget target, ObjectId blobInfoId)
+			public RefInfo(NamespaceId namespaceId, RefName name, NodeLocator target, ObjectId blobInfoId)
 			{
 				NamespaceId = namespaceId;
 				Name = name;
-				BlobLocator = target.Locator.Blob;
+				BlobLocator = target.Blob;
 				BlobInfoId = blobInfoId;
-				ExportIdx = target.Locator.ExportIdx;
-				Hash = target.Hash;
+				ExportIdx = target.ExportIdx;
 			}
 
 			public bool HasExpired(DateTime utcNow) => ExpiresAtUtc.HasValue && utcNow >= ExpiresAtUtc.Value;
@@ -856,7 +852,7 @@ namespace Horde.Build.Storage
 		}
 
 		/// <inheritdoc/>
-		async Task<RefTarget?> TryReadRefTargetAsync(NamespaceId namespaceId, RefName name, DateTime cacheTime = default, CancellationToken cancellationToken = default)
+		async Task<NodeLocator> TryReadRefTargetAsync(NamespaceId namespaceId, RefName name, DateTime cacheTime = default, CancellationToken cancellationToken = default)
 		{
 			RefCacheValue entry;
 			if (!_cache.TryGetValue(name, out entry) || entry.Time < cacheTime)
@@ -867,7 +863,7 @@ namespace Horde.Build.Storage
 
 			if (entry.Value == null)
 			{
-				return null;
+				return default;
 			}
 
 			if (entry.Value.ExpiresAtUtc != null)
@@ -888,12 +884,12 @@ namespace Horde.Build.Storage
 		}
 
 		/// <inheritdoc/>
-		async Task WriteRefTargetAsync(NamespaceId namespaceId, RefName name, RefTarget target, RefOptions? options = null, CancellationToken cancellationToken = default)
+		async Task WriteRefTargetAsync(NamespaceId namespaceId, RefName name, NodeLocator target, RefOptions? options = null, CancellationToken cancellationToken = default)
 		{
-			BlobInfo? newBlobInfo = await _blobCollection.Find(x => x.NamespaceId == namespaceId && x.BlobId == target.Locator.Blob.BlobId).FirstOrDefaultAsync(cancellationToken);
+			BlobInfo? newBlobInfo = await _blobCollection.Find(x => x.NamespaceId == namespaceId && x.BlobId == target.Blob.BlobId).FirstOrDefaultAsync(cancellationToken);
 			if (newBlobInfo == null)
 			{
-				throw new Exception($"Invalid/unknown blob identifier '{target.Locator.Blob.BlobId}' in namespace {namespaceId}");
+				throw new Exception($"Invalid/unknown blob identifier '{target.Blob.BlobId}' in namespace {namespaceId}");
 			}
 
 			RefInfo newRefInfo = new RefInfo(namespaceId, name, target, newBlobInfo.Id);
diff --git a/Engine/Source/Programs/Shared/EpicGames.Horde/Storage/Backends/HttpStorageClient.cs b/Engine/Source/Programs/Shared/EpicGames.Horde/Storage/Backends/HttpStorageClient.cs
--- a/Engine/Source/Programs/Shared/EpicGames.Horde/Storage/Backends/HttpStorageClient.cs
+++ b/Engine/Source/Programs/Shared/EpicGames.Horde/Storage/Backends/HttpStorageClient.cs
@@ -37,7 +37,6 @@ namespace EpicGames.Horde.Storage.Backends
 		{
 			public BlobLocator Blob { get; set; }
 			public int ExportIdx { get; set; }
-			public IoHash Hash { get; set; }
 		}
 
 		readonly Func<HttpClient> _createClient;
@@ -196,7 +195,7 @@ namespace EpicGames.Horde.Storage.Backends
 		}
 
 		/// <inheritdoc/>
-		public override async Task<RefTarget?> TryReadRefTargetAsync(RefName name, DateTime cacheTime = default, CancellationToken cancellationToken = default)
+		public override async Task<NodeLocator> TryReadRefTargetAsync(RefName name, DateTime cacheTime = default, CancellationToken cancellationToken = default)
 		{
 			using (HttpClient httpClient = _createClient())
 			{
@@ -207,7 +206,7 @@ namespace EpicGames.Horde.Storage.Backends
 						if (response.StatusCode == HttpStatusCode.NotFound)
 						{
 							_logger.LogDebug("Read ref {RefName} -> None", name);
-							return null;
+							return default;
 						}
 						else if (!response.IsSuccessStatusCode)
 						{
@@ -219,7 +218,7 @@ namespace EpicGames.Horde.Storage.Backends
 							response.EnsureSuccessStatusCode();
 							ReadRefResponse? data = await response.Content.ReadFromJsonAsync<ReadRefResponse>(cancellationToken: cancellationToken);
 							_logger.LogDebug("Read ref {RefName} -> {Blob}#{ExportIdx}", name, data!.Blob, data!.ExportIdx);
-							return new(data.Hash, new NodeLocator(data!.Blob, data!.ExportIdx));
+							return new NodeLocator(data!.Blob, data!.ExportIdx);
 						}
 					}
 				}
@@ -227,12 +226,12 @@ namespace EpicGames.Horde.Storage.Backends
 		}
 
 		/// <inheritdoc/>
-		public override async Task WriteRefTargetAsync(RefName name, RefTarget target, RefOptions? options = null, CancellationToken cancellationToken = default)
+		public override async Task WriteRefTargetAsync(RefName name, NodeLocator target, RefOptions? options = null, CancellationToken cancellationToken = default)
 		{
 			_logger.LogDebug("Writing ref {RefName} -> {RefTarget}", name, target);
 			using (HttpClient httpClient = _createClient())
 			{
-				using (HttpResponseMessage response = await httpClient.PutAsync($"refs/{name}", new { blob = target.Locator.Blob, exportIdx = target.Locator.ExportIdx, hash = target.Hash, options }, cancellationToken))
+				using (HttpResponseMessage response = await httpClient.PutAsync($"refs/{name}", new { blob = target.Blob, exportIdx = target.ExportIdx, options }, cancellationToken))
 				{
 					response.EnsureSuccessStatusCode();
 				}
diff --git a/Engine/Source/Programs/Shared/EpicGames.Horde/Storage/StorageClientBase.cs b/Engine/Source/Programs/Shared/EpicGames.Horde/Storage/StorageClientBase.cs
--- a/Engine/Source/Programs/Shared/EpicGames.Horde/Storage/StorageClientBase.cs
+++ b/Engine/Source/Programs/Shared/EpicGames.Horde/Storage/StorageClientBase.cs
@@ -39,19 +39,19 @@ namespace EpicGames.Horde.Storage
 		public abstract Task DeleteRefAsync(RefName name, CancellationToken cancellationToken = default);
 
 		/// <inheritdoc/>
-		public abstract Task<RefTarget?> TryReadRefTargetAsync(RefName name, DateTime cacheTime = default, CancellationToken cancellationToken = default);
+		public abstract Task<NodeLocator> TryReadRefTargetAsync(RefName name, DateTime cacheTime = default, CancellationToken cancellationToken = default);
 
 		/// <inheritdoc/>
-		public virtual async Task<RefTarget> WriteRefAsync(RefName name, Bundle bundle, int exportIdx = 0, Utf8String prefix = default, RefOptions? options = null, CancellationToken cancellationToken = default)
+		public virtual async Task<NodeLocator> WriteRefAsync(RefName name, Bundle bundle, int exportIdx = 0, Utf8String prefix = default, RefOptions? options = null, CancellationToken cancellationToken = default)
 		{
 			BlobLocator locator = await this.WriteBundleAsync(bundle, prefix, cancellationToken);
-			RefTarget target = new RefTarget(bundle.Header.Exports[exportIdx].Hash, new NodeLocator(locator, exportIdx));
+			NodeLocator target = new NodeLocator(locator, exportIdx);
 			await WriteRefTargetAsync(name, target, options, cancellationToken);
 			return target;
 		}
 
 		/// <inheritdoc/>
-		public abstract Task WriteRefTargetAsync(RefName name, RefTarget target, RefOptions? options = null, CancellationToken cancellationToken = default);
+		public abstract Task WriteRefTargetAsync(RefName name, NodeLocator target, RefOptions? options = null, CancellationToken cancellationToken = default);
 
 		#endregion
 	}
diff --git a/Engine/Source/Programs/Shared/EpicGames.Horde/Storage/TreeNode.cs b/Engine/Source/Programs/Shared/EpicGames.Horde/Storage/TreeNode.cs
--- a/Engine/Source/Programs/Shared/EpicGames.Horde/Storage/TreeNode.cs
+++ b/Engine/Source/Programs/Shared/EpicGames.Horde/Storage/TreeNode.cs
@@ -21,6 +21,11 @@ namespace EpicGames.Horde.Storage
 		/// </summary>
 		public uint Revision { get; private set; }
 
+		/// <summary>
+		/// Hash when deserialized
+		/// </summary>
+		public IoHash Hash { get; internal set; }
+
 		/// <summary>
 		/// Default constructor
 		/// </summary>
@@ -32,9 +37,9 @@ namespace EpicGames.Horde.Storage
 		/// Serialization constructor. Leaves the revision number zeroed by default.
 		/// </summary>
 		/// <param name="reader"></param>
-		protected TreeNode(IMemoryReader reader)
+		protected TreeNode(ITreeNodeReader reader)
 		{
-			_ = reader;
+			Hash = reader.Hash;
 		}
 
 		/// <summary>
@@ -42,6 +47,7 @@ namespace EpicGames.Horde.Storage
 		/// </summary>
 		protected void MarkAsDirty()
 		{
+			Hash = IoHash.Zero;
 			Revision++;
 		}
 
