Horde: Catch and log errors thrown via replication.

diff --git a/Engine/Source/Programs/Horde/Horde.Build/Perforce/ReplicationService.cs b/Engine/Source/Programs/Horde/Horde.Build/Perforce/ReplicationService.cs
--- a/Engine/Source/Programs/Horde/Horde.Build/Perforce/ReplicationService.cs
+++ b/Engine/Source/Programs/Horde/Horde.Build/Perforce/ReplicationService.cs
@@ -85,7 +85,7 @@ namespace Horde.Build.Perforce
 							removeStreams.Remove(stream.Id);
 							if (!streamIdToTask.ContainsKey(stream.Id))
 							{
-								streamIdToTask.Add(stream.Id, BackgroundTask.StartNew(ctx => RunReplicationAsync(stream.Id, ctx)));
+								streamIdToTask.Add(stream.Id, BackgroundTask.StartNew(ctx => RunReplicationGuardedAsync(stream, ctx)));
 							}
 						}
 					}
@@ -107,14 +107,20 @@ namespace Horde.Build.Perforce
 			}
 		}
 
-		async Task RunReplicationAsync(StreamId streamId, CancellationToken cancellationToken)
+		async Task RunReplicationGuardedAsync(IStream stream, CancellationToken cancellationToken)
 		{
-			IStream? stream = await _streamCollection.GetAsync(streamId);
-			if (stream == null)
+			try
+			{
+				await RunReplicationAsync(stream, cancellationToken);
+			}
+			catch (Exception ex)
 			{
-				return;
+				_logger.LogError(ex, "Exception running replication for {StreamId}: {Message}", stream.Id, ex.Message);
 			}
+		}
 
+		async Task RunReplicationAsync(IStream stream, CancellationToken cancellationToken)
+		{
 			RefName refName = new RefName(stream.Id.ToString());
 
 			IStorageClientImpl store = await _storageService.GetClientAsync(Namespace.Perforce, cancellationToken);
@@ -137,6 +143,7 @@ namespace Horde.Build.Perforce
 
 			for (; ; )
 			{
+				_logger.LogInformation("Replicating {StreamId} change {Change}", stream.Id, commit.Number);
 				await _replicator.WriteAsync(stream, commit.Number, options, cancellationToken);
 				commit = await commits.SubscribeAsync(commit.Number, cancellationToken: cancellationToken).FirstAsync(cancellationToken);
 			}
