Horde: Additional logging for stream replication.

diff --git a/Engine/Source/Programs/Horde/Horde.Build/Perforce/ReplicationService.cs b/Engine/Source/Programs/Horde/Horde.Build/Perforce/ReplicationService.cs
--- a/Engine/Source/Programs/Horde/Horde.Build/Perforce/ReplicationService.cs
+++ b/Engine/Source/Programs/Horde/Horde.Build/Perforce/ReplicationService.cs
@@ -2,15 +2,18 @@
 
 using System;
 using System.Collections.Generic;
+using System.IO;
 using System.Linq;
 using System.Threading;
 using System.Threading.Tasks;
+using EpicGames.Core;
 using EpicGames.Horde.Storage;
 using EpicGames.Horde.Storage.Nodes;
 using Horde.Build.Storage;
 using Horde.Build.Streams;
 using Horde.Build.Utilities;
 using HordeCommon;
+using Microsoft.AspNetCore.SignalR.Protocol;
 using Microsoft.Extensions.Caching.Memory;
 using Microsoft.Extensions.Hosting;
 using Microsoft.Extensions.Logging;
@@ -86,6 +89,7 @@ namespace Horde.Build.Perforce
 							if (!streamIdToTask.ContainsKey(stream.Id))
 							{
 								streamIdToTask.Add(stream.Id, BackgroundTask.StartNew(ctx => RunReplicationGuardedAsync(stream, ctx)));
+								_logger.LogInformation("Started replication of {StreamId}", stream.Id);
 							}
 						}
 					}
@@ -95,9 +99,11 @@ namespace Horde.Build.Perforce
 						if (streamIdToTask.Remove(removeStreamId, out BackgroundTask? task))
 						{
 							await task.DisposeAsync();
+							_logger.LogInformation("Stopped replication of {StreamId}", removeStreamId);
 						}
 					}
 
+					_logger.LogInformation("Replicating {NumStreams} streams", streamIdToTask.Count);
 					await Task.Delay(TimeSpan.FromMinutes(1.0), cancellationToken);
 				}
 			}
