Horde: Retry replication after a delay if it fails with an exception.

diff --git a/Engine/Source/Programs/Horde/Horde.Build/Perforce/ReplicationService.cs b/Engine/Source/Programs/Horde/Horde.Build/Perforce/ReplicationService.cs
--- a/Engine/Source/Programs/Horde/Horde.Build/Perforce/ReplicationService.cs
+++ b/Engine/Source/Programs/Horde/Horde.Build/Perforce/ReplicationService.cs
@@ -13,7 +13,6 @@ using Horde.Build.Storage;
 using Horde.Build.Streams;
 using Horde.Build.Utilities;
 using HordeCommon;
-using Microsoft.AspNetCore.SignalR.Protocol;
 using Microsoft.Extensions.Caching.Memory;
 using Microsoft.Extensions.Hosting;
 using Microsoft.Extensions.Logging;
@@ -103,7 +102,6 @@ namespace Horde.Build.Perforce
 						}
 					}
 
-					_logger.LogInformation("Replicating {NumStreams} streams", streamIdToTask.Count);
 					await Task.Delay(TimeSpan.FromMinutes(1.0), cancellationToken);
 				}
 			}
@@ -115,13 +113,22 @@ namespace Horde.Build.Perforce
 
 		async Task RunReplicationGuardedAsync(IStream stream, CancellationToken cancellationToken)
 		{
-			try
-			{
-				await RunReplicationAsync(stream, cancellationToken);
-			}
-			catch (Exception ex)
+			for (; ; )
 			{
-				_logger.LogError(ex, "Exception running replication for {StreamId}: {Message}", stream.Id, ex.Message);
+				try
+				{
+					await RunReplicationAsync(stream, cancellationToken);
+				}
+				catch (OperationCanceledException) when (cancellationToken.IsCancellationRequested)
+				{
+					break;
+				}
+				catch (Exception ex)
+				{
+					_logger.LogError(ex, "Exception running replication for {StreamId}: {Message}", stream.Id, ex.Message);
+				}
+
+				await Task.Delay(TimeSpan.FromSeconds(20.0), cancellationToken);
 			}
 		}
 
diff --git a/Engine/Source/Programs/Horde/Horde.Build/Utilities/BackgroundTask.cs b/Engine/Source/Programs/Horde/Horde.Build/Utilities/BackgroundTask.cs
--- a/Engine/Source/Programs/Horde/Horde.Build/Utilities/BackgroundTask.cs
+++ b/Engine/Source/Programs/Horde/Horde.Build/Utilities/BackgroundTask.cs
@@ -15,6 +15,11 @@ namespace Horde.Build.Utilities
 		Task _task = Task.CompletedTask;
 		readonly CancellationTokenSource _cancellationTokenSource = new CancellationTokenSource();
 
+		/// <summary>
+		/// Accessor for the inner task
+		/// </summary>
+		public Task Task => _task;
+
 		/// <summary>
 		/// Constructor. Note that the task does not start until <see cref="Start"/> is called.
 		/// </summary>
