Horde: Use server profiles from config files for creating/extracting bundles. Allows extracting replicated Perforce data from Horde server.

diff --git a/Engine/Source/Programs/Horde/Horde.Agent/Commands/Bundles/CreateCommand.cs b/Engine/Source/Programs/Horde/Horde.Agent/Commands/Bundles/CreateCommand.cs
--- a/Engine/Source/Programs/Horde/Horde.Agent/Commands/Bundles/CreateCommand.cs
+++ b/Engine/Source/Programs/Horde/Horde.Agent/Commands/Bundles/CreateCommand.cs
@@ -3,6 +3,7 @@
 using System;
 using System.Diagnostics;
 using System.Net.Http;
+using System.Net.Http.Headers;
 using System.Threading;
 using System.Threading.Tasks;
 using EpicGames.Core;
@@ -11,6 +12,7 @@ using EpicGames.Horde.Storage.Backends;
 using EpicGames.Horde.Storage.Nodes;
 using Microsoft.Extensions.Caching.Memory;
 using Microsoft.Extensions.Logging;
+using Microsoft.Extensions.Options;
 
 namespace Horde.Agent.Commands.Bundles
 {
@@ -43,11 +45,29 @@ namespace Horde.Agent.Commands.Bundles
 		{
 			public IMemoryCache Cache { get; }
 			public IStorageClient Store { get; }
-
-			public HttpStorageClientOwner(Uri baseUri, ILogger logger)
+			readonly IOptions<AgentSettings> _settings;
+			readonly NamespaceId _namespaceId;
+		
+			public HttpStorageClientOwner(IOptions<AgentSettings> settings, NamespaceId namespaceId, ILogger logger)
 			{
 				Cache = new MemoryCache(new MemoryCacheOptions());
-				Store = new HttpStorageClient(() => new HttpClient { BaseAddress = baseUri }, () => new HttpClient(), logger);
+				Store = new HttpStorageClient(CreateDefaultClient, () => new HttpClient(), logger);
+				_settings = settings;
+				_namespaceId = namespaceId;
+			}
+
+			HttpClient CreateDefaultClient()
+			{
+				ServerProfile profile = _settings.Value.GetCurrentServerProfile();
+
+				HttpClient client = new HttpClient();
+				client.BaseAddress = new Uri(profile.Url, $"api/v1/storage/{_namespaceId}/");
+				if (!String.IsNullOrEmpty(profile.Token))
+				{
+					client.DefaultRequestHeaders.Authorization = new AuthenticationHeaderValue("Bearer", profile.Token);
+				}
+
+				return client;
 			}
 
 			public void Dispose()
@@ -61,20 +81,24 @@ namespace Horde.Agent.Commands.Bundles
 		[CommandLine("-Http")]
 		public bool Http { get; set; }
 
-		[CommandLine("-Server=", Description = "Server to read from")]
-		public string Server { get; set; } = "https://localhost:5001";
-
 		[CommandLine("-Namespace=", Description = "Namespace to use for storage")]
 		public NamespaceId NamespaceId { get; set; } = new NamespaceId("default");
 
 		[CommandLine("-StorageDir=", Description = "Overrides the default storage server with a local directory")]
 		public DirectoryReference StorageDir { get; set; } = DirectoryReference.Combine(Program.AppDir, "bundles");
 
+		readonly IOptions<AgentSettings> _settings;
+
+		public BundleCommandBase(IOptions<AgentSettings> settings)
+		{
+			_settings = settings;
+		}
+
 		protected IStorageClientOwner CreateStorageClient(ILogger logger)
 		{
 			if (Http)
 			{
-				return new HttpStorageClientOwner(new Uri($"{Server}/api/v1/storage/{NamespaceId}/"), logger);
+				return new HttpStorageClientOwner(_settings, NamespaceId, logger);
 			}
 			else
 			{
@@ -92,6 +116,11 @@ namespace Horde.Agent.Commands.Bundles
 		[CommandLine("-InputDir=", Required = true)]
 		public DirectoryReference InputDir { get; set; } = null!;
 
+		public CreateCommand(IOptions<AgentSettings> settings)
+			: base(settings)
+		{
+		}
+
 		public override async Task<int> ExecuteAsync(ILogger logger)
 		{
 			using IStorageClientOwner storeOwner = CreateStorageClient(logger);
diff --git a/Engine/Source/Programs/Horde/Horde.Agent/Commands/Bundles/DumpCommand.cs b/Engine/Source/Programs/Horde/Horde.Agent/Commands/Bundles/DumpCommand.cs
--- a/Engine/Source/Programs/Horde/Horde.Agent/Commands/Bundles/DumpCommand.cs
+++ b/Engine/Source/Programs/Horde/Horde.Agent/Commands/Bundles/DumpCommand.cs
@@ -7,6 +7,7 @@ using System.Threading.Tasks;
 using EpicGames.Core;
 using EpicGames.Horde.Storage;
 using Microsoft.Extensions.Logging;
+using Microsoft.Extensions.Options;
 
 namespace Horde.Agent.Commands.Bundles
 {
@@ -19,6 +20,11 @@ namespace Horde.Agent.Commands.Bundles
 		[CommandLine("-Blob=")]
 		public BlobLocator? BlobId { get; set; }
 
+		public DumpCommand(IOptions<AgentSettings> settings)
+			: base(settings)
+		{
+		}
+
 		public override async Task<int> ExecuteAsync(ILogger logger)
 		{
 			using IStorageClientOwner storeOwner = CreateStorageClient(logger);
diff --git a/Engine/Source/Programs/Horde/Horde.Agent/Commands/Bundles/ExtractCommand.cs b/Engine/Source/Programs/Horde/Horde.Agent/Commands/Bundles/ExtractCommand.cs
--- a/Engine/Source/Programs/Horde/Horde.Agent/Commands/Bundles/ExtractCommand.cs
+++ b/Engine/Source/Programs/Horde/Horde.Agent/Commands/Bundles/ExtractCommand.cs
@@ -7,6 +7,7 @@ using EpicGames.Core;
 using EpicGames.Horde.Storage;
 using EpicGames.Horde.Storage.Nodes;
 using Microsoft.Extensions.Logging;
+using Microsoft.Extensions.Options;
 
 namespace Horde.Agent.Commands.Bundles
 {
@@ -19,6 +20,11 @@ namespace Horde.Agent.Commands.Bundles
 		[CommandLine("-OutputDir=", Required = true)]
 		public DirectoryReference OutputDir { get; set; } = null!;
 
+		public ExtractCommand(IOptions<AgentSettings> settings)
+			: base(settings)
+		{
+		}
+
 		public override async Task<int> ExecuteAsync(ILogger logger)
 		{
 			using IStorageClientOwner owner = CreateStorageClient(logger);
diff --git a/Engine/Source/Programs/Horde/Horde.Agent/Commands/Vcs/VcsCommand.cs b/Engine/Source/Programs/Horde/Horde.Agent/Commands/Vcs/VcsCommand.cs
--- a/Engine/Source/Programs/Horde/Horde.Agent/Commands/Vcs/VcsCommand.cs
+++ b/Engine/Source/Programs/Horde/Horde.Agent/Commands/Vcs/VcsCommand.cs
@@ -12,6 +12,7 @@ using EpicGames.Horde.Storage;
 using EpicGames.Horde.Storage.Nodes;
 using Microsoft.CodeAnalysis;
 using Microsoft.Extensions.Logging;
+using Microsoft.Extensions.Options;
 
 namespace Horde.Agent.Commands.Bundles
 {
@@ -185,6 +186,12 @@ namespace Horde.Agent.Commands.Bundles
 		const string DataDir = ".horde";
 		const string StateFileName = "state.dat";
 
+		protected VcsCommand(IOptions<AgentSettings> settings)
+			: base(settings)
+		{
+			NamespaceId = new NamespaceId("horde-perforce");
+		}
+
 		protected IStorageClientOwner CreateStorageClient(DirectoryReference rootDir, ILogger logger)
 		{
 			StorageDir = DirectoryReference.Combine(rootDir, DataDir, "bundles");
@@ -372,11 +379,33 @@ namespace Horde.Agent.Commands.Bundles
 				}
 			}
 		}
+
+		protected static async Task<CommitNode?> GetCommitAsync(TreeReader reader, RefName branchName, int change = 0)
+		{
+			CommitNode tip = await reader.ReadNodeAsync<CommitNode>(branchName);
+			if (change != 0)
+			{
+				while (tip.Number != change)
+				{
+					if (tip.Parent == null || tip.Number < change)
+					{
+						return null;
+					}
+					tip = await tip.Parent.ExpandAsync(reader);
+				}
+			}
+			return tip;
+		}
 	}
 
 	[Command("Vcs", "Init", "Initialize a directory for VCS-like operations")]
 	class VcsInitCommand : VcsCommand
 	{
+		public VcsInitCommand(IOptions<AgentSettings> settings)
+			: base(settings)
+		{
+		}
+
 		public override async Task<int> ExecuteAsync(ILogger logger)
 		{
 			BaseDir ??= DirectoryReference.GetCurrentDirectory();
@@ -392,6 +421,11 @@ namespace Horde.Agent.Commands.Bundles
 		[CommandLine(Prefix = "-Name=", Required = true)]
 		public string Name { get; set; } = "";
 
+		public VcsBranchCommand(IOptions<AgentSettings> settings)
+			: base(settings)
+		{
+		}
+
 		public override async Task<int> ExecuteAsync(ILogger logger)
 		{
 			DirectoryReference rootDir = FindRootDir();
@@ -433,6 +467,11 @@ namespace Horde.Agent.Commands.Bundles
 		[CommandLine("-Force")]
 		public bool Force { get; set; }
 
+		public VcsCheckoutCommand(IOptions<AgentSettings> settings)
+			: base(settings)
+		{
+		}
+
 		public override async Task<int> ExecuteAsync(ILogger logger)
 		{
 			DirectoryReference rootDir = FindRootDir();
@@ -461,18 +500,11 @@ namespace Horde.Agent.Commands.Bundles
 			IStorageClient store = owner.Store;
 			TreeReader reader = new TreeReader(owner.Store, owner.Cache, logger);
 
-			CommitNode tip = await reader.ReadNodeAsync<CommitNode>(branchName);
-			if (Change != 0)
+			CommitNode? tip = await GetCommitAsync(reader, branchName, Change);
+			if (tip == null)
 			{
-				while (tip.Number != Change)
-				{
-					if (tip.Parent == null || tip.Number < Change)
-					{
-						logger.LogError("Unable to find change {Change}", Change);
-						return 1;
-					}
-					tip = await tip.Parent.ExpandAsync(reader);
-				}
+				logger.LogError("Unable to find change {Change}", Change);
+				return 1;
 			}
 
 			workspaceState.Tree = await RealizeAsync(reader, tip.Contents, rootDir, newState, Clean, logger);
@@ -482,7 +514,7 @@ namespace Horde.Agent.Commands.Bundles
 			return 0;
 		}
 
-		static async Task<DirectoryState> RealizeAsync(TreeReader reader, DirectoryNodeRef directoryRef, DirectoryReference dirPath, DirectoryState? directoryState, bool clean, ILogger logger)
+		protected static async Task<DirectoryState> RealizeAsync(TreeReader reader, DirectoryNodeRef directoryRef, DirectoryReference dirPath, DirectoryState? directoryState, bool clean, ILogger logger)
 		{
 			DirectoryReference.CreateDirectory(dirPath);
 
@@ -539,6 +571,42 @@ namespace Horde.Agent.Commands.Bundles
 		}
 	}
 
+	[Command("vcs", "clone", "Initialize a directory for VCS-like operations")]
+	class VcsCloneCommand : VcsCheckoutCommand
+	{
+		public VcsCloneCommand(IOptions<AgentSettings> settings)
+			: base(settings)
+		{
+		}
+
+		public override async Task<int> ExecuteAsync(ILogger logger)
+		{
+			BaseDir ??= DirectoryReference.GetCurrentDirectory();
+			await InitAsync(BaseDir);
+			logger.LogInformation("Initialized in {RootDir}", BaseDir);
+
+			WorkspaceState workspaceState = new WorkspaceState();
+
+			RefName branchName = new RefName(Branch ?? "ue5-main");
+
+			using IStorageClientOwner owner = CreateStorageClient(BaseDir, logger);
+			IStorageClient store = owner.Store;
+			TreeReader reader = new TreeReader(owner.Store, owner.Cache, logger);
+
+			CommitNode? tip = await GetCommitAsync(reader, branchName, Change);
+			if (tip == null)
+			{
+				logger.LogError("Unable to find change {Change}", Change);
+				return 1;
+			}
+
+			workspaceState.Tree = await RealizeAsync(reader, tip.Contents, BaseDir, null, true, logger);
+			await WriteStateAsync(BaseDir, workspaceState);
+
+			return 0;
+		}
+	}
+
 	[Command("vcs", "commit", "Commits data to the VCS store")]
 	class VcsCommitCommand : VcsCommand
 	{
@@ -548,6 +616,11 @@ namespace Horde.Agent.Commands.Bundles
 		[CommandLine("-Message=")]
 		public string? Message { get; set; }
 
+		public VcsCommitCommand(IOptions<AgentSettings> settings)
+			: base(settings)
+		{
+		}
+
 		public override async Task<int> ExecuteAsync(ILogger logger)
 		{
 			DirectoryReference rootDir = FindRootDir();
@@ -664,6 +737,11 @@ namespace Horde.Agent.Commands.Bundles
 		[CommandLine("-Count")]
 		public int Count { get; set; } = 20;
 
+		public VcsLogCommand(IOptions<AgentSettings> settings)
+			: base(settings)
+		{
+		}
+
 		public override async Task<int> ExecuteAsync(ILogger logger)
 		{
 			DirectoryReference rootDir = FindRootDir();
@@ -704,6 +782,11 @@ namespace Horde.Agent.Commands.Bundles
 	[Command("Vcs", "Status", "Find status of local files")]
 	class VcsStatusCommand : VcsCommand
 	{
+		public VcsStatusCommand(IOptions<AgentSettings> settings)
+			: base(settings)
+		{
+		}
+
 		public override async Task<int> ExecuteAsync(ILogger logger)
 		{
 			DirectoryReference rootDir = FindRootDir();
