Horde: Rename command implementation files to match their commands.

diff --git a/Engine/Source/Programs/Horde/Horde.Agent/Commands/Bundles/BundleCreate.cs b/Engine/Source/Programs/Horde/Horde.Agent/Commands/Bundles/BundleCreate.cs
new file mode 100644
--- /dev/null
+++ b/Engine/Source/Programs/Horde/Horde.Agent/Commands/Bundles/BundleCreate.cs
@@ -0,0 +1,144 @@
+// Copyright Epic Games, Inc. All Rights Reserved.
+
+using System;
+using System.Diagnostics;
+using System.Net.Http;
+using System.Net.Http.Headers;
+using System.Threading;
+using System.Threading.Tasks;
+using EpicGames.Core;
+using EpicGames.Horde.Storage;
+using EpicGames.Horde.Storage.Backends;
+using EpicGames.Horde.Storage.Nodes;
+using Microsoft.Extensions.Caching.Memory;
+using Microsoft.Extensions.Logging;
+using Microsoft.Extensions.Options;
+
+namespace Horde.Agent.Commands.Bundles
+{
+	abstract class BundleCommandBase : Command
+	{
+		protected interface IStorageClientOwner : IDisposable
+		{
+			public IMemoryCache Cache { get; }
+			public IStorageClient Store { get; }
+		}
+
+		class FileStorageClientOwner : IStorageClientOwner
+		{
+			public IMemoryCache Cache { get; }
+			public IStorageClient Store { get; }
+
+			public FileStorageClientOwner(DirectoryReference storageDir, ILogger logger)
+			{
+				Cache = new MemoryCache(new MemoryCacheOptions());
+				Store = new FileStorageClient(storageDir, logger);
+			}
+
+			public void Dispose()
+			{
+				Cache.Dispose();
+			}
+		}
+
+		class HttpStorageClientOwner : IStorageClientOwner
+		{
+			public IMemoryCache Cache { get; }
+			public IStorageClient Store { get; }
+			readonly IOptions<AgentSettings> _settings;
+			readonly NamespaceId _namespaceId;
+		
+			public HttpStorageClientOwner(IOptions<AgentSettings> settings, NamespaceId namespaceId, ILogger logger)
+			{
+				Cache = new MemoryCache(new MemoryCacheOptions());
+				Store = new HttpStorageClient(CreateDefaultClient, () => new HttpClient(), logger);
+				_settings = settings;
+				_namespaceId = namespaceId;
+			}
+
+			HttpClient CreateDefaultClient()
+			{
+				ServerProfile profile = _settings.Value.GetCurrentServerProfile();
+
+				HttpClient client = new HttpClient();
+				client.BaseAddress = new Uri(profile.Url, $"api/v1/storage/{_namespaceId}/");
+				if (!String.IsNullOrEmpty(profile.Token))
+				{
+					client.DefaultRequestHeaders.Authorization = new AuthenticationHeaderValue("Bearer", profile.Token);
+				}
+
+				return client;
+			}
+
+			public void Dispose()
+			{
+				Cache.Dispose();
+			}
+		}
+
+		public static RefName DefaultRefName { get; } = new RefName("default-ref");
+
+		[CommandLine("-Http")]
+		public bool Http { get; set; }
+
+		[CommandLine("-Namespace=", Description = "Namespace to use for storage")]
+		public NamespaceId NamespaceId { get; set; } = new NamespaceId("default");
+
+		[CommandLine("-StorageDir=", Description = "Overrides the default storage server with a local directory")]
+		public DirectoryReference StorageDir { get; set; } = DirectoryReference.Combine(Program.AppDir, "bundles");
+
+		readonly IOptions<AgentSettings> _settings;
+
+		public BundleCommandBase(IOptions<AgentSettings> settings)
+		{
+			_settings = settings;
+		}
+
+		protected IStorageClientOwner CreateStorageClient(ILogger logger)
+		{
+			if (Http)
+			{
+				return new HttpStorageClientOwner(_settings, NamespaceId, logger);
+			}
+			else
+			{
+				return new FileStorageClientOwner(DirectoryReference.Combine(StorageDir, NamespaceId.ToString()), logger);
+			}
+		}
+	}
+
+	[Command("bundle", "create", "Creates a bundle from a folder on the local hard drive")]
+	class CreateCommand : BundleCommandBase
+	{
+		[CommandLine("-Ref=")]
+		public RefName RefName { get; set; } = DefaultRefName;
+
+		[CommandLine("-InputDir=", Required = true)]
+		public DirectoryReference InputDir { get; set; } = null!;
+
+		public CreateCommand(IOptions<AgentSettings> settings)
+			: base(settings)
+		{
+		}
+
+		public override async Task<int> ExecuteAsync(ILogger logger)
+		{
+			using IStorageClientOwner storeOwner = CreateStorageClient(logger);
+			IStorageClient store = storeOwner.Store;
+
+			TreeWriter writer = new TreeWriter(store, prefix: RefName.Text);
+
+			DirectoryNode node = new DirectoryNode(DirectoryFlags.None);
+
+			Stopwatch timer = Stopwatch.StartNew();
+
+			ChunkingOptions options = new ChunkingOptions();
+			await node.CopyFromDirectoryAsync(InputDir.ToDirectoryInfo(), options, writer, CancellationToken.None);
+
+			await writer.WriteAsync(RefName, node);
+
+			logger.LogInformation("Time: {Time}", timer.Elapsed.TotalSeconds);
+			return 0;
+		}
+	}
+}
diff --git a/Engine/Source/Programs/Horde/Horde.Agent/Commands/Bundles/BundleDump.cs b/Engine/Source/Programs/Horde/Horde.Agent/Commands/Bundles/BundleDump.cs
new file mode 100644
--- /dev/null
+++ b/Engine/Source/Programs/Horde/Horde.Agent/Commands/Bundles/BundleDump.cs
@@ -0,0 +1,90 @@
+// Copyright Epic Games, Inc. All Rights Reserved.
+
+using System;
+using System.Collections.Generic;
+using System.Linq;
+using System.Threading.Tasks;
+using EpicGames.Core;
+using EpicGames.Horde.Storage;
+using Microsoft.Extensions.Logging;
+using Microsoft.Extensions.Options;
+
+namespace Horde.Agent.Commands.Bundles
+{
+	[Command("bundle", "dump", "Dumps the contents of a bundle")]
+	internal class DumpCommand : BundleCommandBase
+	{
+		[CommandLine("-Ref=")]
+		public RefName RefName { get; set; } = DefaultRefName;
+
+		[CommandLine("-Blob=")]
+		public BlobLocator? BlobId { get; set; }
+
+		public DumpCommand(IOptions<AgentSettings> settings)
+			: base(settings)
+		{
+		}
+
+		public override async Task<int> ExecuteAsync(ILogger logger)
+		{
+			using IStorageClientOwner storeOwner = CreateStorageClient(logger);
+			IStorageClient store = storeOwner.Store;
+
+			if (BlobId == null)
+			{
+				NodeLocator locator = await store.ReadRefTargetAsync(RefName);
+				BlobId = locator.Blob;
+			}
+
+			logger.LogInformation("Summary for blob {BlobId}", BlobId.Value);
+			Bundle bundle = await store.ReadBundleAsync(BlobId.Value);
+
+			BundleHeader header = bundle.Header;
+			int packetStart = 0;
+
+			logger.LogInformation("");
+			logger.LogInformation("Imports: {NumImports}", header.Imports.Count);
+
+			int refIdx = 0;
+			foreach (BundleImport import in header.Imports)
+			{
+				logger.LogInformation("  From blob {BlobId} ({NumExports} nodes)", import.Locator, import.Exports.Count);
+				foreach (int exportIdx in import.Exports)
+				{
+					logger.LogInformation("    [{Index}] IMP {BlobId}:{ExportIdx}", refIdx, import.Locator, exportIdx);
+					refIdx++;
+				}
+			}
+
+			logger.LogInformation("");
+			logger.LogInformation("Exports: {NumExports}", header.Exports.Count);
+
+			int packetIdx = 0;
+			int packetOffset = 0;
+			foreach (BundleExport export in header.Exports)
+			{
+				string refs = (export.References.Count == 0) ? "[ ]" : $"[ {String.Join(", ", export.References.Select(x => x.ToString()))} ]";
+				logger.LogInformation("  [{Index}] EXP {ExportHash} (length: {NumBytes:n0}, packet: {PacketIdx}, refs: {Refs})", refIdx, export.Hash, export.Length, packetIdx, refs);
+				refIdx++;
+
+				packetOffset += export.Length;
+				if(packetOffset >= header.Packets[packetIdx].DecodedLength)
+				{
+					packetIdx++;
+					packetOffset = 0;
+				}
+			}
+
+			logger.LogInformation("");
+			logger.LogInformation("Packets: {NumPackets}", header.Packets.Count);
+			for (int idx = 0; idx < header.Packets.Count; idx++)
+			{
+				BundlePacket packet = header.Packets[idx];
+				logger.LogInformation("  PKT {Idx} (file offset: {Offset:n0}, encoded: {EncodedLength:n0}, decoded: {DecodedLength:n0}, ratio: {Ratio}%)", idx, packetStart, packet.EncodedLength, packet.DecodedLength, (int)(packet.EncodedLength * 100) / packet.DecodedLength);
+				packetStart += packet.EncodedLength;
+			}
+
+			return 0;
+		}
+	}
+}
diff --git a/Engine/Source/Programs/Horde/Horde.Agent/Commands/Bundles/BundleExtract.cs b/Engine/Source/Programs/Horde/Horde.Agent/Commands/Bundles/BundleExtract.cs
new file mode 100644
--- /dev/null
+++ b/Engine/Source/Programs/Horde/Horde.Agent/Commands/Bundles/BundleExtract.cs
@@ -0,0 +1,43 @@
+// Copyright Epic Games, Inc. All Rights Reserved.
+
+using System.Diagnostics;
+using System.Threading;
+using System.Threading.Tasks;
+using EpicGames.Core;
+using EpicGames.Horde.Storage;
+using EpicGames.Horde.Storage.Nodes;
+using Microsoft.Extensions.Logging;
+using Microsoft.Extensions.Options;
+
+namespace Horde.Agent.Commands.Bundles
+{
+	[Command("bundle", "extract", "Extracts data from a bundle to the local hard drive")]
+	internal class ExtractCommand : BundleCommandBase
+	{
+		[CommandLine("-Ref=")]
+		public RefName RefName { get; set; } = DefaultRefName;
+
+		[CommandLine("-OutputDir=", Required = true)]
+		public DirectoryReference OutputDir { get; set; } = null!;
+
+		public ExtractCommand(IOptions<AgentSettings> settings)
+			: base(settings)
+		{
+		}
+
+		public override async Task<int> ExecuteAsync(ILogger logger)
+		{
+			using IStorageClientOwner owner = CreateStorageClient(logger);
+			IStorageClient store = owner.Store;
+			TreeReader reader = new TreeReader(owner.Store, owner.Cache, logger);
+
+			Stopwatch timer = Stopwatch.StartNew();
+
+			DirectoryNode node = await reader.ReadNodeAsync<DirectoryNode>(RefName);
+			await node.CopyToDirectoryAsync(reader, OutputDir.ToDirectoryInfo(), logger, CancellationToken.None);
+
+			logger.LogInformation("Elapsed: {Time}s", timer.Elapsed.TotalSeconds);
+			return 0;
+		}
+	}
+}
diff --git a/Engine/Source/Programs/Horde/Horde.Agent/Commands/Bundles/CreateCommand.cs b/Engine/Source/Programs/Horde/Horde.Agent/Commands/Bundles/CreateCommand.cs
deleted file mode 100644
diff --git a/Engine/Source/Programs/Horde/Horde.Agent/Commands/Bundles/DumpCommand.cs b/Engine/Source/Programs/Horde/Horde.Agent/Commands/Bundles/DumpCommand.cs
deleted file mode 100644
diff --git a/Engine/Source/Programs/Horde/Horde.Agent/Commands/Bundles/ExtractCommand.cs b/Engine/Source/Programs/Horde/Horde.Agent/Commands/Bundles/ExtractCommand.cs
deleted file mode 100644
diff --git a/Engine/Source/Programs/Horde/Horde.Agent/Commands/Service/InstallCommand.cs b/Engine/Source/Programs/Horde/Horde.Agent/Commands/Service/InstallCommand.cs
deleted file mode 100644
diff --git a/Engine/Source/Programs/Horde/Horde.Agent/Commands/Service/RunCommand.cs b/Engine/Source/Programs/Horde/Horde.Agent/Commands/Service/RunCommand.cs
deleted file mode 100644
diff --git a/Engine/Source/Programs/Horde/Horde.Agent/Commands/Service/ServiceInstall.cs b/Engine/Source/Programs/Horde/Horde.Agent/Commands/Service/ServiceInstall.cs
new file mode 100644
--- /dev/null
+++ b/Engine/Source/Programs/Horde/Horde.Agent/Commands/Service/ServiceInstall.cs
@@ -0,0 +1,106 @@
+// Copyright Epic Games, Inc. All Rights Reserved.
+
+using System;
+using System.Reflection;
+using System.Text;
+using System.Threading.Tasks;
+using EpicGames.Core;
+using Horde.Agent.Utility;
+using Microsoft.Extensions.Logging;
+
+namespace Horde.Agent.Commands.Service
+{
+	/// <summary>
+	/// Installs the agent as a service
+	/// </summary>
+	[Command("Service", "Install", "Installs the agent as a service")]
+	class InstallCommand : Command
+	{
+		/// <summary>
+		/// Name of the service
+		/// </summary>
+		public const string ServiceName = "HordeAgent";
+
+		/// <summary>
+		/// Specifies the username for the service to run under
+		/// </summary>
+		[CommandLine("-UserName=")]
+		public string? UserName { get; set; } = null;
+
+		/// <summary>
+		/// Password for the username
+		/// </summary>
+		[CommandLine("-Password=")]
+		public string? Password { get; set; } = null;
+
+		/// <summary>
+		/// The server profile to use
+		/// </summary>
+		[CommandLine("-Server=")]
+		public string? Server { get; set; } = null;
+		
+		/// <summary>
+		/// Path to dotnet executable (dotnet.exe on Windows)
+		/// When left empty, the value of "dotnet" will be used.
+		/// </summary>
+		[CommandLine("-DotNetExecutable=")]
+		public string DotNetExecutable { get; set; } = "dotnet";
+
+		/// <summary>
+		/// Runs the service indefinitely
+		/// </summary>
+		/// <param name="logger">Logger to use</param>
+		/// <returns>Exit code</returns>
+		public override Task<int> ExecuteAsync(ILogger logger)
+		{
+			using (WindowsServiceManager serviceManager = new WindowsServiceManager())
+			{
+				using (WindowsService service = serviceManager.Open(ServiceName))
+				{
+					if (service.IsValid)
+					{
+						logger.LogInformation("Stopping existing service...");
+						service.Stop();
+
+						WindowsServiceStatus status = service.WaitForStatusChange(WindowsServiceStatus.Stopping, TimeSpan.FromSeconds(30.0));
+						if (status != WindowsServiceStatus.Stopped)
+						{
+							logger.LogError("Unable to stop service (status = {Status})", status);
+							return Task.FromResult(1);
+						}
+
+						logger.LogInformation("Deleting service");
+						service.Delete();
+					}
+				}
+
+				logger.LogInformation("Registering {ServiceName} service", ServiceName);
+
+				StringBuilder commandLine = new StringBuilder();
+				commandLine.AppendFormat("{0} \"{1}\" service run", DotNetExecutable, Assembly.GetEntryAssembly()!.Location);
+				if(Server != null)
+				{
+					commandLine.Append($" -server={Server}");
+				}
+
+				using (WindowsService service = serviceManager.Create(ServiceName, "Horde Agent", commandLine.ToString(), UserName, Password))
+				{
+					service.SetDescription("Allows this machine to participate in a Horde farm.");
+
+					logger.LogInformation("Starting...");
+					service.Start();
+
+					WindowsServiceStatus status = service.WaitForStatusChange(WindowsServiceStatus.Starting, TimeSpan.FromSeconds(30.0));
+					if (status != WindowsServiceStatus.Running)
+					{
+						logger.LogError("Unable to start service (status = {Status})", status);
+						return Task.FromResult(1);
+					}
+
+					logger.LogInformation("Done.");
+				}
+			}
+			return Task.FromResult(0);
+		}
+	}
+}
diff --git a/Engine/Source/Programs/Horde/Horde.Agent/Commands/Service/ServiceRun.cs b/Engine/Source/Programs/Horde/Horde.Agent/Commands/Service/ServiceRun.cs
new file mode 100644
--- /dev/null
+++ b/Engine/Source/Programs/Horde/Horde.Agent/Commands/Service/ServiceRun.cs
@@ -0,0 +1,101 @@
+// Copyright Epic Games, Inc. All Rights Reserved.
+
+using System;
+using System.Threading.Tasks;
+using EpicGames.Core;
+using Microsoft.Extensions.Configuration;
+using Microsoft.Extensions.DependencyInjection;
+using Microsoft.Extensions.Hosting;
+using Microsoft.Extensions.Logging;
+using Serilog.Events;
+
+namespace Horde.Agent.Modes.Service
+{
+
+	/// <summary>
+	/// 
+	/// </summary>
+	[Command("Service", "Run", "Runs the service in listen mode")]
+	class RunCommand : Command
+	{
+		/// <summary>
+		/// Log verbosity level (use normal Serilog levels such as debug, warning or info)
+		/// </summary>
+		[CommandLine("-LogLevel")]
+		public string LogLevelStr { get; set; } = "information";
+
+		readonly DefaultServices _defaultServices;
+
+		/// <summary>
+		/// Constructor
+		/// </summary>
+		public RunCommand(DefaultServices defaultServices)
+		{
+			_defaultServices = defaultServices;
+		}
+
+		/// <summary>
+		/// Runs the service indefinitely
+		/// </summary>
+		/// <returns>Exit code</returns>
+		public override async Task<int> ExecuteAsync(ILogger logger)
+		{
+			if (Enum.TryParse(LogLevelStr, true, out LogEventLevel logEventLevel))
+			{
+				Logging.LogLevelSwitch.MinimumLevel = logEventLevel;
+			}
+			else
+			{
+				Console.WriteLine($"Unable to parse log level: {LogLevelStr}");
+				return 0;
+			}
+			
+			IHostBuilder hostBuilder = Host.CreateDefaultBuilder();
+
+			// Attempt to setup this process as a Windows service. A race condition inside Microsoft.Extensions.Hosting.WindowsServices.WindowsServiceHelpers.IsWindowsService
+			// can result in accessing the parent process after it's terminated, so catch any exceptions that it throws.
+			try
+			{
+				hostBuilder = hostBuilder.UseWindowsService();
+			}
+			catch (InvalidOperationException)
+			{
+			}
+
+			hostBuilder = hostBuilder
+				.ConfigureAppConfiguration(builder =>
+				{
+					builder.AddConfiguration(_defaultServices.Configuration);
+				})
+				.ConfigureLogging(builder =>
+				{
+					// We add our logger through ConfigureServices, inherited from _defaultServices
+					builder.ClearProviders();
+				})
+				.ConfigureServices((hostContext, services) =>
+				{
+					services.Configure<HostOptions>(options =>
+					{
+						// Allow the worker service to terminate any before shutting down
+						options.ShutdownTimeout = TimeSpan.FromSeconds(30);
+					});
+
+					foreach (ServiceDescriptor descriptor in _defaultServices.Descriptors)
+					{
+						services.Add(descriptor);
+					}
+				});
+
+			try
+			{
+				await hostBuilder.Build().RunAsync();
+			}
+			catch (OperationCanceledException)
+			{
+				// Need to catch this to prevent it propagating to the awaiter when pressing ctrl-c
+			}
+
+			return 0;
+		}
+	}
+}
diff --git a/Engine/Source/Programs/Horde/Horde.Agent/Commands/Service/ServiceUninstall.cs b/Engine/Source/Programs/Horde/Horde.Agent/Commands/Service/ServiceUninstall.cs
new file mode 100644
--- /dev/null
+++ b/Engine/Source/Programs/Horde/Horde.Agent/Commands/Service/ServiceUninstall.cs
@@ -0,0 +1,57 @@
+// Copyright Epic Games, Inc. All Rights Reserved.
+
+using System;
+using System.Threading.Tasks;
+using EpicGames.Core;
+using Horde.Agent.Utility;
+using Microsoft.Extensions.Logging;
+
+namespace Horde.Agent.Commands.Service
+{
+	/// <summary>
+	/// Uninstalls the service
+	/// </summary>
+	[Command("Service", "Uninstall", "Uninstalls the service")]
+	class UninstallCommand : Command
+	{
+		/// <summary>
+		/// Name of the service
+		/// </summary>
+		public const string ServiceName = InstallCommand.ServiceName;
+
+		/// <summary>
+		/// Runs the service indefinitely
+		/// </summary>
+		/// <param name="logger">Logger to use</param>
+		/// <returns>Exit code</returns>
+		public override Task<int> ExecuteAsync(ILogger logger)
+		{
+			using (WindowsServiceManager serviceManager = new WindowsServiceManager())
+			{
+				using (WindowsService service = serviceManager.Open(ServiceName))
+				{
+					if (service.IsValid)
+					{
+						logger.LogInformation("Stopping existing service...");
+						service.Stop();
+
+						WindowsServiceStatus status = service.WaitForStatusChange(WindowsServiceStatus.Stopping, TimeSpan.FromSeconds(30.0));
+						if (status != WindowsServiceStatus.Stopped)
+						{
+							logger.LogError("Unable to stop service (status = {Status})", status);
+							return Task.FromResult(1);
+						}
+
+						logger.LogInformation("Deleting service");
+						service.Delete();
+					}
+					else
+					{
+						logger.LogInformation("Unable to find service {ServiceName}", ServiceName);
+					}
+				}
+			}
+			return Task.FromResult(0);
+		}
+	}
+}
diff --git a/Engine/Source/Programs/Horde/Horde.Agent/Commands/Service/ServiceUpgrade.cs b/Engine/Source/Programs/Horde/Horde.Agent/Commands/Service/ServiceUpgrade.cs
new file mode 100644
--- /dev/null
+++ b/Engine/Source/Programs/Horde/Horde.Agent/Commands/Service/ServiceUpgrade.cs
@@ -0,0 +1,233 @@
+// Copyright Epic Games, Inc. All Rights Reserved.
+
+using System;
+using System.Collections.Generic;
+using System.Diagnostics;
+using System.IO;
+using System.Linq;
+using System.Reflection;
+using System.Runtime.InteropServices;
+using System.Runtime.Versioning;
+using System.ServiceProcess;
+using System.Text;
+using System.Threading.Tasks;
+using EpicGames.Core;
+using Microsoft.Extensions.Logging;
+using Microsoft.Management.Infrastructure;
+
+namespace Horde.Agent.Commands
+{
+	/// <summary>
+	/// Upgrades a running service to the current application
+	/// </summary>
+	[Command("Service", "Upgrade", "Replaces a running service with the application in the current directory")]
+	class UpgradeCommand : Command
+	{
+		/// <summary>
+		/// The process id to replace
+		/// </summary>
+		[CommandLine("-ProcessId=", Required = true)]
+		int ProcessId { get; set; } = -1;
+
+		/// <summary>
+		/// The target directory to install to
+		/// </summary>
+		[CommandLine("-TargetDir=", Required = true)]
+		DirectoryReference TargetDir { get; set; } = null!;
+
+		/// <summary>
+		/// Arguments to forwar to the target executable
+		/// </summary>
+		[CommandLine("-Arguments=", Required = true)]
+		string Arguments { get; set; } = null!;
+
+		/// <summary>
+		/// Upgrades the application to a new version
+		/// </summary>
+		/// <param name="logger">The log output device</param>
+		/// <returns>True if the upgrade succeeded</returns>
+		public override Task<int> ExecuteAsync(ILogger logger)
+		{
+			// Stop the other process
+			logger.LogInformation("Attempting to perform upgrade on process {ProcessId}", ProcessId);
+			using (Process otherProcess = Process.GetProcessById(ProcessId))
+			{
+				// Get the directory containing the target application
+				DirectoryInfo targetDir = new DirectoryInfo(TargetDir.FullName);
+				HashSet<string> targetFiles = new HashSet<string>(targetDir.EnumerateFiles("*", SearchOption.AllDirectories).Select(x => x.FullName), StringComparer.OrdinalIgnoreCase);
+
+				// Find all the source files
+				DirectoryInfo sourceDir = new DirectoryInfo(Path.GetDirectoryName(Assembly.GetExecutingAssembly().Location)!);
+				HashSet<string> sourceFiles = new HashSet<string>(sourceDir.EnumerateFiles("*", SearchOption.AllDirectories).Select(x => x.FullName), StringComparer.OrdinalIgnoreCase);
+
+				// Exclude all the source files from the list of target files, since we may be in a subdirectory
+				targetFiles.ExceptWith(sourceFiles);
+
+				// Ignore any files that are in the saved directory
+				string sourceDataDir = Path.Combine(sourceDir.FullName, "Saved") + Path.DirectorySeparatorChar;
+				sourceFiles.RemoveWhere(x => x.StartsWith(sourceDataDir, StringComparison.OrdinalIgnoreCase));
+
+				string targetDataDir = Path.Combine(targetDir.FullName, "Saved") + Path.DirectorySeparatorChar;
+				targetFiles.RemoveWhere(x => x.StartsWith(targetDataDir, StringComparison.OrdinalIgnoreCase));
+
+				// Copy all the files into the target directory
+				List<Tuple<string, string>> renameFiles = new List<Tuple<string, string>>();
+				foreach (string sourceFile in sourceFiles)
+				{
+					if (!sourceFile.StartsWith(sourceDir.FullName, StringComparison.OrdinalIgnoreCase))
+					{
+						throw new InvalidDataException($"Expected {sourceFile} to be under {sourceDir.FullName}");
+					}
+
+					string targetFile = targetDir.FullName + sourceFile.Substring(sourceDir.FullName.Length);
+					Directory.CreateDirectory(Path.GetDirectoryName(targetFile)!);
+
+					string targetFileBeforeRename = targetFile + ".new";
+					logger.LogDebug("Copying {SourceFile} to {TargetFileBeforeRename}", sourceFile, targetFileBeforeRename);
+					File.Copy(sourceFile, targetFileBeforeRename, true);
+
+					renameFiles.Add(Tuple.Create(targetFileBeforeRename, targetFile));
+					targetFiles.Remove(targetFileBeforeRename);
+				}
+
+				if (RuntimeInformation.IsOSPlatform(OSPlatform.Windows))
+				{
+					UpgradeWindowsService(logger, otherProcess, targetFiles, renameFiles);
+				}
+				else if (RuntimeInformation.IsOSPlatform(OSPlatform.OSX))
+				{
+					UpgradeMacService(logger, otherProcess, targetFiles, renameFiles);
+				}
+				else if (RuntimeInformation.IsOSPlatform(OSPlatform.Linux))
+				{
+					UpgradeLinuxService(logger, otherProcess, targetFiles, renameFiles);
+				}
+				else
+				{
+					logger.LogError("Agent is not running a platform that supports Upgrades. Platform: {Platform}", RuntimeInformation.OSDescription);
+					return Task.FromResult(-1);
+				}
+			}
+			logger.LogInformation("Upgrade complete");
+			return Task.FromResult(0);
+		}
+
+		static void UpgradeFilesInPlace(ILogger logger, HashSet<string> targetFiles, List<Tuple<string, string>> renameFiles)
+		{
+			// Remove all the target files
+			foreach (string targetFile in targetFiles)
+			{
+				logger.LogDebug("Deleting {File}", targetFile);
+				File.SetAttributes(targetFile, FileAttributes.Normal);
+				File.Delete(targetFile);
+			}
+
+			// Rename all the new files into place
+			foreach (Tuple<string, string> pair in renameFiles)
+			{
+				logger.LogDebug("Renaming {SourceFile} to {TargetFile}", pair.Item1, pair.Item2);
+				File.Move(pair.Item1, pair.Item2, true);
+			}
+		}
+
+		static void UpgradeMacService(ILogger logger, Process otherProcess, HashSet<string> targetFiles, List<Tuple<string, string>> renameFiles)
+		{
+			UpgradeFilesInPlace(logger, targetFiles, renameFiles);
+			logger.LogDebug("Upgrade completed, restarting...");
+			otherProcess.Kill();
+		}
+
+		static void UpgradeLinuxService(ILogger logger, Process otherProcess, HashSet<string> targetFiles, List<Tuple<string, string>> renameFiles)
+		{
+			UpgradeFilesInPlace(logger, targetFiles, renameFiles);
+			logger.LogDebug("Upgrade completed, restarting...");
+			otherProcess.Kill();
+		}
+
+		[SupportedOSPlatform("windows")]
+		void UpgradeWindowsService(ILogger logger, Process otherProcess, HashSet<string> targetFiles, List<Tuple<string, string>> renameFiles)
+		{
+			// Try to get the service associated with the passed-in process id
+			using (ServiceController? service = GetServiceForProcess(ProcessId))
+			{
+				// Stop the process
+				if (service == null)
+				{
+					logger.LogInformation("Terminating other process");
+					otherProcess.Kill();
+				}
+				else
+				{
+					logger.LogInformation("Stopping service");
+					service.Stop();
+				}
+				otherProcess.WaitForExit();
+
+				UpgradeFilesInPlace(logger, targetFiles, renameFiles);
+
+				// Run the new application
+				if (service == null)
+				{
+					string driverFileName = "dotnet";
+					string assemblyFileName = Path.Combine(TargetDir.FullName, Path.GetFileName(Assembly.GetExecutingAssembly().Location));
+
+					StringBuilder driverArguments = new StringBuilder();
+					driverArguments.AppendArgument(assemblyFileName);
+					driverArguments.Append(' ');
+					driverArguments.Append(Arguments);
+
+					StringBuilder launch = new StringBuilder();
+					launch.AppendArgument(driverFileName);
+					launch.Append(' ');
+					launch.Append(driverArguments);
+					logger.LogInformation("Launching: {Launch}", launch.ToString());
+
+					using Process newProcess = Process.Start(driverFileName, driverArguments.ToString());
+				}
+				else
+				{
+					// Start the service again
+					logger.LogInformation("Restarting service");
+					service.Start();
+				}
+			}
+		}
+
+		/// <summary>
+		/// Try to find the service controller for the given process id
+		/// </summary>
+		/// <param name="processId">The process id to search for</param>
+		/// <returns>The service controller corresponding to this process</returns>
+		[SupportedOSPlatform("windows")]
+		public static ServiceController? GetServiceForProcess(int processId)
+		{
+			try
+			{
+				using (CimSession session = CimSession.Create(null))
+				{
+					const string QueryNamespace = @"root\cimv2";
+					const string QueryDialect = "WQL";
+
+					foreach (CimInstance instance in session.QueryInstances(QueryNamespace, QueryDialect, $"SELECT Name FROM Win32_Service WHERE ProcessId={processId}"))
+					{
+						foreach (CimProperty property in instance.CimInstanceProperties)
+						{
+							if (property.Name.Equals("Name", StringComparison.OrdinalIgnoreCase))
+							{
+								string? value = property.Value as string;
+								if (value != null)
+								{
+									return new ServiceController(value);
+								}
+							}
+						}
+					}
+				}
+			}
+			catch
+			{
+			}
+			return null;
+		}
+	}
+}
diff --git a/Engine/Source/Programs/Horde/Horde.Agent/Commands/Service/UninstallCommand.cs b/Engine/Source/Programs/Horde/Horde.Agent/Commands/Service/UninstallCommand.cs
deleted file mode 100644
diff --git a/Engine/Source/Programs/Horde/Horde.Agent/Commands/Service/UpgradeCommand.cs b/Engine/Source/Programs/Horde/Horde.Agent/Commands/Service/UpgradeCommand.cs
deleted file mode 100644
diff --git a/Engine/Source/Programs/Horde/Horde.Agent/Commands/Workspace/CleanCommand.cs b/Engine/Source/Programs/Horde/Horde.Agent/Commands/Workspace/CleanCommand.cs
deleted file mode 100644
diff --git a/Engine/Source/Programs/Horde/Horde.Agent/Commands/Workspace/ClearCommand.cs b/Engine/Source/Programs/Horde/Horde.Agent/Commands/Workspace/ClearCommand.cs
deleted file mode 100644
diff --git a/Engine/Source/Programs/Horde/Horde.Agent/Commands/Workspace/DumpCommand.cs b/Engine/Source/Programs/Horde/Horde.Agent/Commands/Workspace/DumpCommand.cs
deleted file mode 100644
diff --git a/Engine/Source/Programs/Horde/Horde.Agent/Commands/Workspace/PopulateCacheCommand.cs b/Engine/Source/Programs/Horde/Horde.Agent/Commands/Workspace/PopulateCacheCommand.cs
deleted file mode 100644
diff --git a/Engine/Source/Programs/Horde/Horde.Agent/Commands/Workspace/PurgeCacheCommand.cs b/Engine/Source/Programs/Horde/Horde.Agent/Commands/Workspace/PurgeCacheCommand.cs
deleted file mode 100644
diff --git a/Engine/Source/Programs/Horde/Horde.Agent/Commands/Workspace/RepairCommand.cs b/Engine/Source/Programs/Horde/Horde.Agent/Commands/Workspace/RepairCommand.cs
deleted file mode 100644
diff --git a/Engine/Source/Programs/Horde/Horde.Agent/Commands/Workspace/RevertCommand.cs b/Engine/Source/Programs/Horde/Horde.Agent/Commands/Workspace/RevertCommand.cs
deleted file mode 100644
diff --git a/Engine/Source/Programs/Horde/Horde.Agent/Commands/Workspace/SetupCommand.cs b/Engine/Source/Programs/Horde/Horde.Agent/Commands/Workspace/SetupCommand.cs
deleted file mode 100644
diff --git a/Engine/Source/Programs/Horde/Horde.Agent/Commands/Workspace/SizesCommand.cs b/Engine/Source/Programs/Horde/Horde.Agent/Commands/Workspace/SizesCommand.cs
deleted file mode 100644
diff --git a/Engine/Source/Programs/Horde/Horde.Agent/Commands/Workspace/StatusCommand.cs b/Engine/Source/Programs/Horde/Horde.Agent/Commands/Workspace/StatusCommand.cs
deleted file mode 100644
diff --git a/Engine/Source/Programs/Horde/Horde.Agent/Commands/Workspace/SyncCommand.cs b/Engine/Source/Programs/Horde/Horde.Agent/Commands/Workspace/SyncCommand.cs
deleted file mode 100644
diff --git a/Engine/Source/Programs/Horde/Horde.Agent/Commands/Workspace/WorkspaceBase.cs b/Engine/Source/Programs/Horde/Horde.Agent/Commands/Workspace/WorkspaceBase.cs
new file mode 100644
--- /dev/null
+++ b/Engine/Source/Programs/Horde/Horde.Agent/Commands/Workspace/WorkspaceBase.cs
@@ -0,0 +1,155 @@
+// Copyright Epic Games, Inc. All Rights Reserved.
+
+using System;
+using System.Collections.Generic;
+using System.ComponentModel;
+using System.Linq;
+using System.Threading;
+using System.Threading.Tasks;
+using EpicGames.Core;
+using EpicGames.Perforce;
+using EpicGames.Perforce.Managed;
+using Microsoft.Extensions.Logging;
+
+namespace Horde.Agent.Commands.Workspace
+{
+	abstract class WorkspaceCommand : Command
+	{
+		[CommandLine("-Server")]
+		[Description("Specifies the Perforce server and port")]
+		protected string ServerAndPort { get; set; } = PerforceSettings.Default.ServerAndPort;
+
+		[CommandLine("-User")]
+		[Description("Specifies the Perforce username")]
+		protected string UserName { get; set; } = PerforceSettings.Default.UserName;
+
+		[CommandLine("-BaseDir", Required = true)]
+		[Description("Base directory to use for syncing workspaces")]
+		protected DirectoryReference BaseDir { get; set; } = null!;
+
+		[CommandLine("-Overwrite")]
+		[Description("")]
+		protected bool Overwrite { get; set; } = false;
+		
+		[CommandLine("-PreferNativeClient")]
+		[Description("Prefer to use native Perforce client (instead of launching separate p4 process)")]
+		protected bool PreferNativeClient { get; set; } = false;
+		
+		[CommandLine("-UseHaveTable")]
+		[Description("Use have-table for syncing")]
+		string UseHaveTable { get; set; } = "true";
+
+		public override void Configure(CommandLineArguments arguments, ILogger logger)
+		{
+			base.Configure(arguments, logger);
+
+			if(BaseDir == null)
+			{
+				for (DirectoryReference? parentDir = DirectoryReference.GetCurrentDirectory(); parentDir != null; parentDir = parentDir.ParentDirectory)
+				{
+					if (ManagedWorkspace.Exists(parentDir))
+					{
+						BaseDir = parentDir;
+						break;
+					}
+				}
+
+				if (BaseDir == null)
+				{
+					throw new FatalErrorException("Unable to find existing repository in current working tree. Specify -BaseDir=... explicitly.");
+				}
+			}
+		}
+
+		public override async Task<int> ExecuteAsync(ILogger logger)
+		{
+			PerforceSettings settings = new(ServerAndPort, UserName) { PreferNativeClient = PreferNativeClient };
+			using IPerforceConnection perforce = await PerforceConnection.CreateAsync(settings, logger);
+			InfoRecord info = await perforce.GetInfoAsync(InfoOptions.ShortOutput, CancellationToken.None);
+			bool useHaveTable = UseHaveTable.Equals("true", StringComparison.Ordinal);
+
+			using Logging.HordeLoggerProvider loggerProvider = new ();
+			ILogger repoLogger = loggerProvider.CreateLogger("Repository");
+			ManagedWorkspace repo = await ManagedWorkspace.LoadOrCreateAsync(info.ClientHost!, BaseDir, Overwrite, useHaveTable, repoLogger, CancellationToken.None);
+			await ExecuteAsync(perforce, repo, logger);
+			return 0;
+		}
+
+		protected abstract Task ExecuteAsync(IPerforceConnection perforce, ManagedWorkspace repo, ILogger logger);
+
+		protected static long ParseSize(string size)
+		{
+			long value;
+			if (size.EndsWith("gb", StringComparison.OrdinalIgnoreCase))
+			{
+				string sizeValue = size.Substring(0, size.Length - 2).TrimEnd();
+				if (Int64.TryParse(sizeValue, out value))
+				{
+					return value * (1024 * 1024 * 1024);
+				}
+			}
+			else if (size.EndsWith("mb", StringComparison.OrdinalIgnoreCase))
+			{
+				string sizeValue = size.Substring(0, size.Length - 2).TrimEnd();
+				if (Int64.TryParse(sizeValue, out value))
+				{
+					return value * (1024 * 1024);
+				}
+			}
+			else if (size.EndsWith("kb", StringComparison.OrdinalIgnoreCase))
+			{
+				string sizeValue = size.Substring(0, size.Length - 2).TrimEnd();
+				if (Int64.TryParse(sizeValue, out value))
+				{
+					return value * 1024;
+				}
+			}
+			else
+			{
+				if (Int64.TryParse(size, out value))
+				{
+					return value;
+				}
+			}
+			throw new FatalErrorException("Invalid size '{0}'", size);
+		}
+
+		protected static int ParseChangeNumber(string change)
+		{
+			int changeNumber;
+			if (Int32.TryParse(change, out changeNumber) && changeNumber > 0)
+			{
+				return changeNumber;
+			}
+			throw new FatalErrorException("Unable to parse change number from '{0}'", change);
+		}
+
+		protected static int ParseChangeNumberOrLatest(string change)
+		{
+			if (change.Equals("Latest", StringComparison.OrdinalIgnoreCase))
+			{
+				return -1;
+			}
+			else
+			{
+				return ParseChangeNumber(change);
+			}
+		}
+
+		protected static List<string> ExpandFilters(List<string> arguments)
+		{
+			List<string> filters = new List<string>();
+			foreach (string argument in arguments)
+			{
+				foreach (string singleArgument in argument.Split(';').Select(x => x.Trim()))
+				{
+					if (singleArgument.Length > 0)
+					{
+						filters.Add(singleArgument);
+					}
+				}
+			}
+			return filters;
+		}
+	}
+}
diff --git a/Engine/Source/Programs/Horde/Horde.Agent/Commands/Workspace/WorkspaceClean.cs b/Engine/Source/Programs/Horde/Horde.Agent/Commands/Workspace/WorkspaceClean.cs
new file mode 100644
--- /dev/null
+++ b/Engine/Source/Programs/Horde/Horde.Agent/Commands/Workspace/WorkspaceClean.cs
@@ -0,0 +1,25 @@
+// Copyright Epic Games, Inc. All Rights Reserved.
+
+using System.ComponentModel;
+using System.Threading;
+using System.Threading.Tasks;
+using EpicGames.Core;
+using EpicGames.Perforce;
+using EpicGames.Perforce.Managed;
+using Microsoft.Extensions.Logging;
+
+namespace Horde.Agent.Commands.Workspace
+{
+	[Command("Workspace", "Clean", "Cleans all modified files from the workspace.")]
+	class CleanCommand : WorkspaceCommand
+	{
+		[CommandLine("-Incremental")]
+		[Description("Performs an incremental sync, without removing intermediates")]
+		public bool Incremental { get; set; } = false;
+
+		protected override Task ExecuteAsync(IPerforceConnection perforce, ManagedWorkspace repo, ILogger logger)
+		{
+			return repo.CleanAsync(!Incremental, CancellationToken.None);
+		}
+	}
+}
diff --git a/Engine/Source/Programs/Horde/Horde.Agent/Commands/Workspace/WorkspaceClear.cs b/Engine/Source/Programs/Horde/Horde.Agent/Commands/Workspace/WorkspaceClear.cs
new file mode 100644
--- /dev/null
+++ b/Engine/Source/Programs/Horde/Horde.Agent/Commands/Workspace/WorkspaceClear.cs
@@ -0,0 +1,20 @@
+// Copyright Epic Games, Inc. All Rights Reserved.
+
+using System.Threading;
+using System.Threading.Tasks;
+using EpicGames.Core;
+using EpicGames.Perforce;
+using EpicGames.Perforce.Managed;
+using Microsoft.Extensions.Logging;
+
+namespace Horde.Agent.Commands.Workspace
+{
+	[Command("Workspace", "Clear", "Empties the staging directory of any files, returning them to the cache")]
+	class ClearCommand : WorkspaceCommand
+	{
+		protected override Task ExecuteAsync(IPerforceConnection perforce, ManagedWorkspace repo, ILogger logger)
+		{
+			return repo.ClearAsync(CancellationToken.None);
+		}
+	}
+}
diff --git a/Engine/Source/Programs/Horde/Horde.Agent/Commands/Workspace/WorkspaceCommand.cs b/Engine/Source/Programs/Horde/Horde.Agent/Commands/Workspace/WorkspaceCommand.cs
deleted file mode 100644
diff --git a/Engine/Source/Programs/Horde/Horde.Agent/Commands/Workspace/WorkspaceDump.cs b/Engine/Source/Programs/Horde/Horde.Agent/Commands/Workspace/WorkspaceDump.cs
new file mode 100644
--- /dev/null
+++ b/Engine/Source/Programs/Horde/Horde.Agent/Commands/Workspace/WorkspaceDump.cs
@@ -0,0 +1,20 @@
+// Copyright Epic Games, Inc. All Rights Reserved.
+
+using System.Threading.Tasks;
+using EpicGames.Core;
+using EpicGames.Perforce;
+using EpicGames.Perforce.Managed;
+using Microsoft.Extensions.Logging;
+
+namespace Horde.Agent.Commands.Workspace
+{
+	[Command("Workspace", "Dump", "Dumps the contents of the repository to the log for analysis")]
+	class Dump : WorkspaceCommand
+	{
+		protected override Task ExecuteAsync(IPerforceConnection perforce, ManagedWorkspace repo, ILogger logger)
+		{
+			repo.Dump();
+			return Task.CompletedTask;
+		}
+	}
+}
diff --git a/Engine/Source/Programs/Horde/Horde.Agent/Commands/Workspace/WorkspacePopulateCache.cs b/Engine/Source/Programs/Horde/Horde.Agent/Commands/Workspace/WorkspacePopulateCache.cs
new file mode 100644
--- /dev/null
+++ b/Engine/Source/Programs/Horde/Horde.Agent/Commands/Workspace/WorkspacePopulateCache.cs
@@ -0,0 +1,50 @@
+// Copyright Epic Games, Inc. All Rights Reserved.
+
+using System;
+using System.Collections.Generic;
+using System.ComponentModel;
+using System.Threading;
+using System.Threading.Tasks;
+using EpicGames.Core;
+using EpicGames.Perforce;
+using EpicGames.Perforce.Managed;
+using Microsoft.Extensions.Logging;
+
+namespace Horde.Agent.Commands.Workspace
+{
+	[Command("Workspace", "PopulateCache", "Populates the cache with the head revision of the given streams")]
+	class PopulateCacheCommand : WorkspaceCommand
+	{
+		[CommandLine("-ClientAndStream=")]
+		[Description("Specifies client and stream pairs, in the format Client:Stream")]
+		List<string> ClientAndStreamParams { get; set; } = new List<string>();
+
+		[CommandLine("-Filter=")]
+		[Description("Filters for the files to sync, in P4 syntax (eg. /Engine/...)")]
+		List<string> Filters { get; set; } = new List<string>();
+
+		[CommandLine("-FakeSync")]
+		[Description("Simulates the sync without actually fetching any files")]
+		bool FakeSync { get; set; } = false;
+
+		protected override async Task ExecuteAsync(IPerforceConnection perforce, ManagedWorkspace repo, ILogger logger)
+		{
+			List<string> expandedFilters = ExpandFilters(Filters);
+
+			List<PopulateRequest> populateRequests = new List<PopulateRequest>();
+			foreach (string clientAndStreamParam in ClientAndStreamParams)
+			{
+				int idx = clientAndStreamParam.IndexOf(':', StringComparison.Ordinal);
+				if (idx == -1)
+				{
+					throw new FatalErrorException("Expected -ClientAndStream=<ClientName>:<StreamName>");
+				}
+
+				using IPerforceConnection perforceClient = await perforce.WithClientAsync(clientAndStreamParam.Substring(0, idx));
+				populateRequests.Add(new PopulateRequest(perforceClient, clientAndStreamParam.Substring(idx + 1), expandedFilters));
+			}
+
+			await repo.PopulateAsync(populateRequests, FakeSync, CancellationToken.None);
+		}
+	}
+}
diff --git a/Engine/Source/Programs/Horde/Horde.Agent/Commands/Workspace/WorkspacePurgeCache.cs b/Engine/Source/Programs/Horde/Horde.Agent/Commands/Workspace/WorkspacePurgeCache.cs
new file mode 100644
--- /dev/null
+++ b/Engine/Source/Programs/Horde/Horde.Agent/Commands/Workspace/WorkspacePurgeCache.cs
@@ -0,0 +1,29 @@
+// Copyright Epic Games, Inc. All Rights Reserved.
+
+using System.Threading;
+using System.Threading.Tasks;
+using EpicGames.Core;
+using EpicGames.Perforce;
+using EpicGames.Perforce.Managed;
+using Microsoft.Extensions.Logging;
+
+namespace Horde.Agent.Commands.Workspace
+{
+	[Command("Workspace", "PurgeCache", "Shrink the size of the cache to the given size")]
+	class PurgeCacheCommand : WorkspaceCommand
+	{
+		[CommandLine("-Size=")]
+		string? SizeParam { get; set; } = null;
+
+		protected override Task ExecuteAsync(IPerforceConnection perforce, ManagedWorkspace repo, ILogger logger)
+		{
+			long size = 0;
+			if (SizeParam != null)
+			{
+				size = ParseSize(SizeParam);
+			}
+
+			return repo.PurgeAsync(size, CancellationToken.None);
+		}
+	}
+}
diff --git a/Engine/Source/Programs/Horde/Horde.Agent/Commands/Workspace/WorkspaceRepair.cs b/Engine/Source/Programs/Horde/Horde.Agent/Commands/Workspace/WorkspaceRepair.cs
new file mode 100644
--- /dev/null
+++ b/Engine/Source/Programs/Horde/Horde.Agent/Commands/Workspace/WorkspaceRepair.cs
@@ -0,0 +1,20 @@
+// Copyright Epic Games, Inc. All Rights Reserved.
+
+using System.Threading;
+using System.Threading.Tasks;
+using EpicGames.Core;
+using EpicGames.Perforce;
+using EpicGames.Perforce.Managed;
+using Microsoft.Extensions.Logging;
+
+namespace Horde.Agent.Commands.Workspace
+{
+	[Command("Workspace", "RepairCache", "Checks the integrity of the cache, and removes any invalid files")]
+	class RepairCommand : WorkspaceCommand
+	{
+		protected override Task ExecuteAsync(IPerforceConnection perforce, ManagedWorkspace repo, ILogger logger)
+		{
+			return repo.RepairAsync(CancellationToken.None);
+		}
+	}
+}
diff --git a/Engine/Source/Programs/Horde/Horde.Agent/Commands/Workspace/WorkspaceRevert.cs b/Engine/Source/Programs/Horde/Horde.Agent/Commands/Workspace/WorkspaceRevert.cs
new file mode 100644
--- /dev/null
+++ b/Engine/Source/Programs/Horde/Horde.Agent/Commands/Workspace/WorkspaceRevert.cs
@@ -0,0 +1,26 @@
+// Copyright Epic Games, Inc. All Rights Reserved.
+
+using System.ComponentModel;
+using System.Threading;
+using System.Threading.Tasks;
+using EpicGames.Core;
+using EpicGames.Perforce;
+using EpicGames.Perforce.Managed;
+using Microsoft.Extensions.Logging;
+
+namespace Horde.Agent.Commands.Workspace
+{
+	[Command("Workspace", "Revert", "Revert all files that are open in the current workspace. Does not replace them with valid revisions.")]
+	class RevertCommand : WorkspaceCommand
+	{
+		[CommandLine("-Client=", Required = true)]
+		[Description("Client to revert all files for")]
+		string ClientName { get; set; } = null!;
+
+		protected override async Task ExecuteAsync(IPerforceConnection perforce, ManagedWorkspace repo, ILogger logger)
+		{
+			using IPerforceConnection perforceClient = await perforce.WithClientAsync(ClientName);
+			await repo.RevertAsync(perforceClient, CancellationToken.None);
+		}
+	}
+}
diff --git a/Engine/Source/Programs/Horde/Horde.Agent/Commands/Workspace/WorkspaceSetup.cs b/Engine/Source/Programs/Horde/Horde.Agent/Commands/Workspace/WorkspaceSetup.cs
new file mode 100644
--- /dev/null
+++ b/Engine/Source/Programs/Horde/Horde.Agent/Commands/Workspace/WorkspaceSetup.cs
@@ -0,0 +1,30 @@
+// Copyright Epic Games, Inc. All Rights Reserved.
+
+using System.ComponentModel;
+using System.Threading;
+using System.Threading.Tasks;
+using EpicGames.Core;
+using EpicGames.Perforce;
+using EpicGames.Perforce.Managed;
+using Microsoft.Extensions.Logging;
+
+namespace Horde.Agent.Commands.Workspace
+{
+	[Command("Workspace", "Setup", "Creates or updates a client to use a given stream")]
+	class SetupCommand : WorkspaceCommand
+	{
+		[CommandLine("-Client=", Required = true)]
+		[Description("Name of the client to create")]
+		string ClientName { get; set; } = null!;
+
+		[CommandLine("-Stream=", Required = true)]
+		[Description("Name of the stream to configure")]
+		string StreamName { get; set; } = null!;
+
+		protected override async Task ExecuteAsync(IPerforceConnection perforce, ManagedWorkspace repo, ILogger logger)
+		{
+			using IPerforceConnection perforceClient = await perforce.WithClientAsync(ClientName);
+			await repo.SetupAsync(perforceClient, StreamName, CancellationToken.None);
+		}
+	}
+}
diff --git a/Engine/Source/Programs/Horde/Horde.Agent/Commands/Workspace/WorkspaceSizes.cs b/Engine/Source/Programs/Horde/Horde.Agent/Commands/Workspace/WorkspaceSizes.cs
new file mode 100644
--- /dev/null
+++ b/Engine/Source/Programs/Horde/Horde.Agent/Commands/Workspace/WorkspaceSizes.cs
@@ -0,0 +1,35 @@
+// Copyright Epic Games, Inc. All Rights Reserved.
+
+using System.Collections.Generic;
+using System.ComponentModel;
+using System.Threading;
+using System.Threading.Tasks;
+using EpicGames.Core;
+using EpicGames.Perforce;
+using EpicGames.Perforce.Managed;
+using Microsoft.Extensions.Logging;
+
+namespace Horde.Agent.Commands.Workspace
+{
+	[Command("Workspace", "Sizes", "Gathers stats for which streams take the most amount of space in the cache for the given configuration")]
+	class SizesCommand : WorkspaceCommand
+	{
+		[CommandLine("-TempClient=", Required = true)]
+		[Description("Name of a temporary client to switch between streams gathering metadata. Will be created if it does not exist.")]
+		string TempClientName { get; set; } = null!;
+
+		[CommandLine("-Stream=")]
+		[Description("Streams that should be included in the output")]
+		List<string> StreamNames { get; set; } = new List<string>();
+
+		[CommandLine("-Filter=")]
+		[Description("Filters for the files to sync, in P4 syntax (eg. /Engine/...)")]
+		List<string> Filters { get; set; } = new List<string>();
+
+		protected override async Task ExecuteAsync(IPerforceConnection perforce, ManagedWorkspace repo, ILogger logger)
+		{
+			using IPerforceConnection perforceClient = await perforce.WithClientAsync(TempClientName);
+			await repo.StatsAsync(perforceClient, StreamNames, Filters, CancellationToken.None);
+		}
+	}
+}
diff --git a/Engine/Source/Programs/Horde/Horde.Agent/Commands/Workspace/WorkspaceStatus.cs b/Engine/Source/Programs/Horde/Horde.Agent/Commands/Workspace/WorkspaceStatus.cs
new file mode 100644
--- /dev/null
+++ b/Engine/Source/Programs/Horde/Horde.Agent/Commands/Workspace/WorkspaceStatus.cs
@@ -0,0 +1,20 @@
+// Copyright Epic Games, Inc. All Rights Reserved.
+
+using System.Threading.Tasks;
+using EpicGames.Core;
+using EpicGames.Perforce;
+using EpicGames.Perforce.Managed;
+using Microsoft.Extensions.Logging;
+
+namespace Horde.Agent.Commands.Workspace
+{
+	[Command("Workspace", "Status", "Prints information about the state of the cache and workspace")]
+	class StatusCommand : WorkspaceCommand
+	{
+		protected override Task ExecuteAsync(IPerforceConnection perforce, ManagedWorkspace repo, ILogger logger)
+		{
+			repo.Status();
+			return Task.CompletedTask;
+		}
+	}
+}
diff --git a/Engine/Source/Programs/Horde/Horde.Agent/Commands/Workspace/WorkspaceSync.cs b/Engine/Source/Programs/Horde/Horde.Agent/Commands/Workspace/WorkspaceSync.cs
new file mode 100644
--- /dev/null
+++ b/Engine/Source/Programs/Horde/Horde.Agent/Commands/Workspace/WorkspaceSync.cs
@@ -0,0 +1,63 @@
+// Copyright Epic Games, Inc. All Rights Reserved.
+
+using System.Collections.Generic;
+using System.ComponentModel;
+using System.Threading;
+using System.Threading.Tasks;
+using EpicGames.Core;
+using EpicGames.Perforce;
+using EpicGames.Perforce.Managed;
+using Microsoft.Extensions.Logging;
+
+namespace Horde.Agent.Commands.Workspace
+{
+	[Command("Workspace", "Sync", "Syncs the files for a particular stream and changelist")]
+	class SyncCommand : WorkspaceCommand
+	{
+		[CommandLine("-Client=", Required = true)]
+		[Description("Name of the client to sync. Will be created if it does not exist.")]
+		public string ClientName { get; set; } = null!;
+
+		[CommandLine("-Stream=", Required = true)]
+		[Description("The stream to sync")]
+		public string StreamName { get; set; } = null!;
+
+		[CommandLine("-Change=")]
+		[Description("The change to sync. May be a changelist number, or 'Latest'")]
+		public string Change { get; set; } = "Latest";
+
+		[CommandLine("-Preflight=")]
+		[Description("The change to unshelve into the workspace")]
+		public int PreflightChange { get; set; } = -1;
+
+		[CommandLine]
+		[Description("Optional path to a cache file used to store workspace metadata. Using a location on a network share allows multiple machines syncing the same CL to only query Perforce state once.")]
+		FileReference? CacheFile { get; set; } = null;
+
+		[CommandLine("-Filter=")]
+		[Description("Filters for the files to sync, in P4 syntax (eg. /Engine/...)")]
+		public List<string> Filters { get; set; } = new List<string>();
+
+		[CommandLine("-Incremental")]
+		[Description("Performs an incremental sync, without removing intermediates")]
+		public bool IncrementalSync { get; set; } = false;
+
+		[CommandLine("-FakeSync")]
+		[Description("Simulates the sync without actually fetching any files")]
+		public bool FakeSync { get; set; } = false;
+
+		protected override async Task ExecuteAsync(IPerforceConnection perforce, ManagedWorkspace repo, ILogger logger)
+		{
+			int changeNumber = ParseChangeNumberOrLatest(Change);
+			List<string> expandedFilters = ExpandFilters(Filters);
+
+			using IPerforceConnection perforceClient = await perforce.WithClientAsync(ClientName);
+			await repo.SyncAsync(perforceClient, StreamName, changeNumber, expandedFilters, !IncrementalSync, FakeSync, CacheFile, CancellationToken.None);
+
+			if (PreflightChange != -1)
+			{
+				await repo.UnshelveAsync(perforceClient, PreflightChange, CancellationToken.None);
+			}
+		}
+	}
+}
