Horde: Use registered server profiles for storage commands by default, and use explicit filesystem client for creating/extracting archive bundles.

diff --git a/Engine/Source/Programs/Horde/Horde.Agent/AgentSettings.cs b/Engine/Source/Programs/Horde/Horde.Agent/AgentSettings.cs
--- a/Engine/Source/Programs/Horde/Horde.Agent/AgentSettings.cs
+++ b/Engine/Source/Programs/Horde/Horde.Agent/AgentSettings.cs
@@ -191,6 +191,11 @@ namespace Horde.Agent
 		/// </summary>
 		public bool EnableNewLogger { get; set; }
 
+		/// <summary>
+		/// Option to use a local storage client rather than connecting through the server. Primarily for convenience when debugging / iterating locally.
+		/// </summary>
+		public bool UseLocalStorageClient { get; set; }
+
 		/// <summary>
 		/// Key/value properties in addition to those set internally by the agent
 		/// </summary>
diff --git a/Engine/Source/Programs/Horde/Horde.Agent/Commands/Bundles/BundleCreate.cs b/Engine/Source/Programs/Horde/Horde.Agent/Commands/Bundles/BundleCreate.cs
--- a/Engine/Source/Programs/Horde/Horde.Agent/Commands/Bundles/BundleCreate.cs
+++ b/Engine/Source/Programs/Horde/Horde.Agent/Commands/Bundles/BundleCreate.cs
@@ -1,130 +1,28 @@
 // Copyright Epic Games, Inc. All Rights Reserved.
 
-using System;
 using System.Diagnostics;
-using System.Net.Http;
-using System.Net.Http.Headers;
 using System.Threading;
 using System.Threading.Tasks;
 using EpicGames.Core;
 using EpicGames.Horde.Storage;
 using EpicGames.Horde.Storage.Backends;
 using EpicGames.Horde.Storage.Nodes;
-using Microsoft.Extensions.Caching.Memory;
 using Microsoft.Extensions.Logging;
-using Microsoft.Extensions.Options;
 
 namespace Horde.Agent.Commands.Bundles
 {
-	abstract class BundleCommandBase : Command
-	{
-		protected interface IStorageClientOwner : IDisposable
-		{
-			public IMemoryCache Cache { get; }
-			public IStorageClient Store { get; }
-		}
-
-		class FileStorageClientOwner : IStorageClientOwner
-		{
-			public IMemoryCache Cache { get; }
-			public IStorageClient Store { get; }
-
-			public FileStorageClientOwner(DirectoryReference storageDir, ILogger logger)
-			{
-				Cache = new MemoryCache(new MemoryCacheOptions());
-				Store = new FileStorageClient(storageDir, logger);
-			}
-
-			public void Dispose()
-			{
-				Cache.Dispose();
-			}
-		}
-
-		class HttpStorageClientOwner : IStorageClientOwner
-		{
-			public IMemoryCache Cache { get; }
-			public IStorageClient Store { get; }
-			readonly IOptions<AgentSettings> _settings;
-			readonly NamespaceId _namespaceId;
-		
-			public HttpStorageClientOwner(IOptions<AgentSettings> settings, NamespaceId namespaceId, ILogger logger)
-			{
-				Cache = new MemoryCache(new MemoryCacheOptions());
-				Store = new HttpStorageClient(CreateDefaultClient, () => new HttpClient(), logger);
-				_settings = settings;
-				_namespaceId = namespaceId;
-			}
-
-			HttpClient CreateDefaultClient()
-			{
-				ServerProfile profile = _settings.Value.GetCurrentServerProfile();
-
-				HttpClient client = new HttpClient();
-				client.BaseAddress = new Uri(profile.Url, $"api/v1/storage/{_namespaceId}/");
-				if (!String.IsNullOrEmpty(profile.Token))
-				{
-					client.DefaultRequestHeaders.Authorization = new AuthenticationHeaderValue("Bearer", profile.Token);
-				}
-
-				return client;
-			}
-
-			public void Dispose()
-			{
-				Cache.Dispose();
-			}
-		}
-
-		public static RefName DefaultRefName { get; } = new RefName("default-ref");
-
-		[CommandLine("-Http")]
-		public bool Http { get; set; }
-
-		[CommandLine("-Namespace=", Description = "Namespace to use for storage")]
-		public NamespaceId NamespaceId { get; set; } = new NamespaceId("default");
-
-		[CommandLine("-StorageDir=", Description = "Overrides the default storage server with a local directory")]
-		public DirectoryReference StorageDir { get; set; } = DirectoryReference.Combine(Program.AppDir, "bundles");
-
-		readonly IOptions<AgentSettings> _settings;
-
-		public BundleCommandBase(IOptions<AgentSettings> settings)
-		{
-			_settings = settings;
-		}
-
-		protected IStorageClientOwner CreateStorageClient(ILogger logger)
-		{
-			if (Http)
-			{
-				return new HttpStorageClientOwner(_settings, NamespaceId, logger);
-			}
-			else
-			{
-				return new FileStorageClientOwner(DirectoryReference.Combine(StorageDir, NamespaceId.ToString()), logger);
-			}
-		}
-	}
-
 	[Command("bundle", "create", "Creates a bundle from a folder on the local hard drive")]
-	class CreateCommand : BundleCommandBase
+	class BundleCreate : Command
 	{
 		[CommandLine("-Ref=")]
-		public RefName RefName { get; set; } = DefaultRefName;
+		public RefName RefName { get; set; } = new RefName("default-bundle");
 
 		[CommandLine("-InputDir=", Required = true)]
 		public DirectoryReference InputDir { get; set; } = null!;
 
-		public CreateCommand(IOptions<AgentSettings> settings)
-			: base(settings)
-		{
-		}
-
 		public override async Task<int> ExecuteAsync(ILogger logger)
 		{
-			using IStorageClientOwner storeOwner = CreateStorageClient(logger);
-			IStorageClient store = storeOwner.Store;
+			FileStorageClient store = new FileStorageClient(DirectoryReference.Combine(Program.DataDir, "Bundles"), logger);
 
 			TreeWriter writer = new TreeWriter(store, prefix: RefName.Text);
 
diff --git a/Engine/Source/Programs/Horde/Horde.Agent/Commands/Bundles/BundleDump.cs b/Engine/Source/Programs/Horde/Horde.Agent/Commands/Bundles/BundleDump.cs
--- a/Engine/Source/Programs/Horde/Horde.Agent/Commands/Bundles/BundleDump.cs
+++ b/Engine/Source/Programs/Horde/Horde.Agent/Commands/Bundles/BundleDump.cs
@@ -2,56 +2,116 @@
 
 using System;
 using System.Collections.Generic;
-using System.Linq;
+using System.IO;
+using System.Reflection;
+using System.Threading;
 using System.Threading.Tasks;
 using EpicGames.Core;
 using EpicGames.Horde.Storage;
 using Microsoft.Extensions.Logging;
-using Microsoft.Extensions.Options;
 
 namespace Horde.Agent.Commands.Bundles
 {
 	[Command("bundle", "dump", "Dumps the contents of a bundle")]
-	internal class DumpCommand : BundleCommandBase
+	internal class BundleDump : Command
 	{
+		[CommandLine("-Namespace=", Description = "Namespace to use for storage")]
+		public NamespaceId NamespaceId { get; set; } = new NamespaceId("default");
+
 		[CommandLine("-Ref=")]
-		public RefName RefName { get; set; } = DefaultRefName;
+		public RefName? RefName { get; set; }
 
 		[CommandLine("-Blob=")]
 		public BlobLocator? BlobId { get; set; }
 
-		public DumpCommand(IOptions<AgentSettings> settings)
-			: base(settings)
+		[CommandLine("-File=")]
+		public string? File { get; set; }
+
+		[CommandLine("-Verbose")]
+		public bool Verbose { get; set; }
+
+		readonly IStorageClientFactory _storageClientFactory;
+
+		public BundleDump(IStorageClientFactory storageClientFactory)
 		{
+			_storageClientFactory = storageClientFactory;
 		}
 
 		public override async Task<int> ExecuteAsync(ILogger logger)
 		{
-			using IStorageClientOwner storeOwner = CreateStorageClient(logger);
-			IStorageClient store = storeOwner.Store;
-
-			if (BlobId == null)
+			Bundle bundle;
+			if (File != null)
+			{
+				FileReference Location = new FileReference(File);
+				byte[] bytes = await FileReference.ReadAllBytesAsync(Location);
+				bundle = new Bundle(new MemoryReader(bytes));
+				logger.LogInformation("Summary for blob {Location}", Location);
+			}
+			else if (BlobId != null)
 			{
-				NodeLocator locator = await store.ReadRefTargetAsync(RefName);
-				BlobId = locator.Blob;
+				IStorageClient store = _storageClientFactory.Create(NamespaceId);
+				bundle = await store.ReadBundleAsync(BlobId.Value);
+				logger.LogInformation("Summary for blob {BlobId}", BlobId.Value);
 			}
+			else if (RefName != null)
+			{
+				IStorageClient store = _storageClientFactory.Create(NamespaceId);
 
-			logger.LogInformation("Summary for blob {BlobId}", BlobId.Value);
-			Bundle bundle = await store.ReadBundleAsync(BlobId.Value);
+				NodeLocator locator = await store.ReadRefTargetAsync(RefName.Value);
+				logger.LogInformation("Ref {RefName} -> {Locator}", RefName.Value, locator);
+
+				bundle = await store.ReadBundleAsync(locator.Blob);
+				logger.LogInformation("Summary for blob {BlobId}", locator.Blob);
+			}
+			else
+			{
+				logger.LogError("No blob specified - use -File/Blob/Ref=...");
+				return 1;
+			}
 
 			BundleHeader header = bundle.Header;
+
+			Dictionary<Guid, string> typeIdToName = new Dictionary<Guid, string>();
+
+			TreeReaderOptions options = new TreeReaderOptions();
+			foreach (Type type in options.Types)
+			{
+				TreeNodeAttribute? attribute = type.GetCustomAttribute<TreeNodeAttribute>();
+				if (attribute != null)
+				{
+					typeIdToName.Add(Guid.Parse(attribute.Guid), type.Name);
+				}
+			}
+
+			string[] types = new string[header.Types.Count];
+			for (int idx = 0; idx < header.Types.Count; idx++)
+			{
+				string? name;
+				if (!typeIdToName.TryGetValue(header.Types[idx].Guid, out name))
+				{
+					name = header.Types[idx].Guid.ToString();
+				}
+				types[idx] = $"{name} v{header.Types[idx].Version}";
+			}
+
 			int packetStart = 0;
 
 			logger.LogInformation("");
 			logger.LogInformation("Imports: {NumImports}", header.Imports.Count);
 
+			List<string> nodeNames = new List<string>();
+
 			int refIdx = 0;
 			foreach (BundleImport import in header.Imports)
 			{
 				logger.LogInformation("  From blob {BlobId} ({NumExports} nodes)", import.Locator, import.Exports.Count);
 				foreach (int exportIdx in import.Exports)
 				{
-					logger.LogInformation("    [{Index}] IMP {BlobId}:{ExportIdx}", refIdx, import.Locator, exportIdx);
+					nodeNames.Add($"IMP {import.Locator}#{exportIdx}");
+					if (Verbose)
+					{
+						logger.LogInformation("    [{Index}] IMP {ExportIdx}", refIdx, exportIdx);
+					}
 					refIdx++;
 				}
 			}
@@ -61,12 +121,23 @@ namespace Horde.Agent.Commands.Bundles
 
 			int packetIdx = 0;
 			int packetOffset = 0;
-			foreach (BundleExport export in header.Exports)
+			for(int exportIdx = 0; exportIdx < header.Exports.Count; exportIdx++)
 			{
-				string refs = (export.References.Count == 0) ? "[ ]" : $"[ {String.Join(", ", export.References.Select(x => x.ToString()))} ]";
-				logger.LogInformation("  [{Index}] EXP {ExportHash} (length: {NumBytes:n0}, packet: {PacketIdx}, refs: {Refs})", refIdx, export.Hash, export.Length, packetIdx, refs);
+				BundleExport export = header.Exports[exportIdx];
+				logger.LogInformation("  [{Index}] EXP {ExportHash} (type: {Type}, length: {NumBytes:n0}, packet: {PacketIdx})", refIdx, export.Hash, types[export.TypeIdx], export.Length, packetIdx);
+
+				nodeNames.Add($"EXP {exportIdx}");
 				refIdx++;
 
+				if (Verbose)
+				{
+					for (int idx = 0; idx < export.References.Count; idx++)
+					{
+						int nodeIdx = export.References[idx];
+						logger.LogInformation("            REF {RefIdx,-3} -> {NodeIdx,-3} ({RefName})", idx, nodeIdx, nodeNames[nodeIdx]);
+					}
+				}
+
 				packetOffset += export.Length;
 				if(packetOffset >= header.Packets[packetIdx].DecodedLength)
 				{
diff --git a/Engine/Source/Programs/Horde/Horde.Agent/Commands/Bundles/BundleExtract.cs b/Engine/Source/Programs/Horde/Horde.Agent/Commands/Bundles/BundleExtract.cs
--- a/Engine/Source/Programs/Horde/Horde.Agent/Commands/Bundles/BundleExtract.cs
+++ b/Engine/Source/Programs/Horde/Horde.Agent/Commands/Bundles/BundleExtract.cs
@@ -5,31 +5,28 @@ using System.Threading;
 using System.Threading.Tasks;
 using EpicGames.Core;
 using EpicGames.Horde.Storage;
+using EpicGames.Horde.Storage.Backends;
 using EpicGames.Horde.Storage.Nodes;
+using Microsoft.Extensions.Caching.Memory;
 using Microsoft.Extensions.Logging;
-using Microsoft.Extensions.Options;
 
 namespace Horde.Agent.Commands.Bundles
 {
 	[Command("bundle", "extract", "Extracts data from a bundle to the local hard drive")]
-	internal class ExtractCommand : BundleCommandBase
+	internal class BundleExtract : Command
 	{
 		[CommandLine("-Ref=")]
-		public RefName RefName { get; set; } = DefaultRefName;
+		public RefName RefName { get; set; } = new RefName("default-bundle");
 
 		[CommandLine("-OutputDir=", Required = true)]
 		public DirectoryReference OutputDir { get; set; } = null!;
 
-		public ExtractCommand(IOptions<AgentSettings> settings)
-			: base(settings)
-		{
-		}
-
 		public override async Task<int> ExecuteAsync(ILogger logger)
 		{
-			using IStorageClientOwner owner = CreateStorageClient(logger);
-			IStorageClient store = owner.Store;
-			TreeReader reader = new TreeReader(owner.Store, owner.Cache, logger);
+			FileStorageClient store = new FileStorageClient(DirectoryReference.Combine(Program.DataDir, "Bundles"), logger);
+
+			using MemoryCache cache = new MemoryCache(new MemoryCacheOptions());
+			TreeReader reader = new TreeReader(store, cache, logger);
 
 			Stopwatch timer = Stopwatch.StartNew();
 
diff --git a/Engine/Source/Programs/Horde/Horde.Agent/Commands/Vcs/VcsBase.cs b/Engine/Source/Programs/Horde/Horde.Agent/Commands/Vcs/VcsBase.cs
--- a/Engine/Source/Programs/Horde/Horde.Agent/Commands/Vcs/VcsBase.cs
+++ b/Engine/Source/Programs/Horde/Horde.Agent/Commands/Vcs/VcsBase.cs
@@ -9,18 +9,19 @@ using System.Threading.Tasks;
 using EpicGames.Core;
 using EpicGames.Horde.Storage;
 using EpicGames.Horde.Storage.Nodes;
-using Horde.Agent.Commands.Bundles;
 using Microsoft.CodeAnalysis;
 using Microsoft.Extensions.Logging;
-using Microsoft.Extensions.Options;
 
 namespace Horde.Agent.Commands.Vcs
 {
-	abstract class VcsBase : BundleCommandBase
+	abstract class VcsBase : Command
 	{
 		[CommandLine("-BaseDir=")]
 		public DirectoryReference? BaseDir { get; set; }
 
+		[CommandLine("-Namespace=")]
+		public NamespaceId NamespaceId { get; set; } = new NamespaceId("horde-perforce");
+
 		protected class FileState
 		{
 			public IoHash Hash { get; set; }
@@ -151,53 +152,18 @@ namespace Horde.Agent.Commands.Vcs
 			}
 		}
 
-		public static IEnumerable<(TKey, TOldValue?, TNewValue?)> Zip<TKey, TOldValue, TNewValue>(IReadOnlyDictionary<TKey, TOldValue>? oldDictionary, IReadOnlyDictionary<TKey, TNewValue>? newDictionary)
-			where TKey : notnull 
-			where TOldValue : class 
-			where TNewValue : class
-		{
-			if (newDictionary != null)
-			{
-				foreach ((TKey key, TNewValue newValue) in newDictionary)
-				{
-					TOldValue? oldValue;
-					if (oldDictionary == null || !oldDictionary.TryGetValue(key, out oldValue))
-					{
-						yield return (key, null, newValue);
-					}
-					else
-					{
-						yield return (key, oldValue, newValue);
-					}
-				}
-			}
-			if (oldDictionary != null)
-			{
-				foreach ((TKey key, TOldValue oldValue) in oldDictionary)
-				{
-					if (newDictionary == null || !newDictionary.ContainsKey(key))
-					{
-						yield return (key, oldValue, null);
-					}
-				}
-			}
-		}
-
 		const string DataDir = ".horde";
 		const string StateFileName = "state.dat";
 
-		protected VcsBase(IOptions<AgentSettings> settings)
-			: base(settings)
-		{
-			NamespaceId = new NamespaceId("horde-perforce");
-		}
+		readonly IStorageClientFactory _storageClientFactory;
 
-		protected IStorageClientOwner CreateStorageClient(DirectoryReference rootDir, ILogger logger)
+		protected VcsBase(IStorageClientFactory storageClientFactory)
 		{
-			StorageDir = DirectoryReference.Combine(rootDir, DataDir, "bundles");
-			return base.CreateStorageClient(logger);
+			_storageClientFactory = storageClientFactory;
 		}
 
+		protected IStorageClient CreateStorageClient() => _storageClientFactory.Create(NamespaceId);
+
 		protected static async Task<WorkspaceState> ReadStateAsync(DirectoryReference rootDir)
 		{
 			FileReference stateFile = FileReference.Combine(rootDir, DataDir, StateFileName);
@@ -298,7 +264,7 @@ namespace Horde.Agent.Commands.Vcs
 
 		protected static void RemoveAddedFiles(DirectoryState oldState, DirectoryState newState)
 		{
-			List<(Utf8String, DirectoryState?, DirectoryState?)> directoryDeltas = Zip(oldState.Directories, newState.Directories).ToList();
+			List<(Utf8String, DirectoryState?, DirectoryState?)> directoryDeltas = EnumerableExtensions.Zip(oldState.Directories, newState.Directories).ToList();
 			foreach ((Utf8String name, DirectoryState? oldSubDirState, _) in directoryDeltas)
 			{
 				if (oldSubDirState == null)
@@ -307,7 +273,7 @@ namespace Horde.Agent.Commands.Vcs
 				}
 			}
 
-			List<(Utf8String, FileState?, FileState?)> fileDeltas = Zip(oldState.Files, newState.Files).ToList();
+			List<(Utf8String, FileState?, FileState?)> fileDeltas = EnumerableExtensions.Zip(oldState.Files, newState.Files).ToList();
 			foreach ((Utf8String name, FileState? oldFileState, _) in fileDeltas)
 			{
 				if (oldFileState == null)
@@ -324,7 +290,7 @@ namespace Horde.Agent.Commands.Vcs
 				return newState;
 			}
 
-			foreach ((_, DirectoryState? oldSubDirState, DirectoryState? newSubDirState) in Zip(oldState.Directories, newState.Directories))
+			foreach ((_, DirectoryState? oldSubDirState, DirectoryState? newSubDirState) in EnumerableExtensions.Zip(oldState.Directories, newState.Directories))
 			{
 				if (oldSubDirState == null || newSubDirState == null || oldSubDirState != DedupTrees(oldSubDirState, newSubDirState))
 				{
@@ -332,7 +298,7 @@ namespace Horde.Agent.Commands.Vcs
 				}
 			}
 
-			foreach ((_, FileState? oldFileState, FileState? newFileState) in Zip(oldState.Files, newState.Files))
+			foreach ((_, FileState? oldFileState, FileState? newFileState) in EnumerableExtensions.Zip(oldState.Files, newState.Files))
 			{
 				if (oldFileState != newFileState)
 				{
@@ -347,7 +313,7 @@ namespace Horde.Agent.Commands.Vcs
 
 		static void PrintDelta(string prefix, DirectoryState oldState, DirectoryState newState, ILogger logger)
 		{
-			foreach ((Utf8String name, DirectoryState? oldSubDirState, DirectoryState? newSubDirState) in Zip(oldState.Directories, newState.Directories))
+			foreach ((Utf8String name, DirectoryState? oldSubDirState, DirectoryState? newSubDirState) in EnumerableExtensions.Zip(oldState.Directories, newState.Directories))
 			{
 				if (oldSubDirState == null)
 				{
@@ -363,7 +329,7 @@ namespace Horde.Agent.Commands.Vcs
 				}
 			}
 
-			foreach ((Utf8String name, FileState? oldFileState, FileState? newFileState) in Zip(oldState.Files, newState.Files))
+			foreach ((Utf8String name, FileState? oldFileState, FileState? newFileState) in EnumerableExtensions.Zip(oldState.Files, newState.Files))
 			{
 				if (oldFileState == null)
 				{
@@ -397,9 +363,4 @@ namespace Horde.Agent.Commands.Vcs
 			return tip;
 		}
 	}
-
-
-
-
-
 }
diff --git a/Engine/Source/Programs/Horde/Horde.Agent/Commands/Vcs/VcsBranch.cs b/Engine/Source/Programs/Horde/Horde.Agent/Commands/Vcs/VcsBranch.cs
--- a/Engine/Source/Programs/Horde/Horde.Agent/Commands/Vcs/VcsBranch.cs
+++ b/Engine/Source/Programs/Horde/Horde.Agent/Commands/Vcs/VcsBranch.cs
@@ -14,8 +14,8 @@ namespace Horde.Agent.Commands.Vcs
 		[CommandLine(Prefix = "-Name=", Required = true)]
 		public string Name { get; set; } = "";
 
-		public VcsBranchCommand(IOptions<AgentSettings> settings)
-			: base(settings)
+		public VcsBranchCommand(IStorageClientFactory storageClientFactory)
+			: base(storageClientFactory)
 		{
 		}
 
@@ -25,8 +25,7 @@ namespace Horde.Agent.Commands.Vcs
 
 			WorkspaceState workspaceState = await ReadStateAsync(rootDir);
 
-			using IStorageClientOwner owner = CreateStorageClient(rootDir, logger);
-			IStorageClient store = owner.Store;
+			IStorageClient store = CreateStorageClient();
 
 			RefName branchName = new RefName(Name);
 			if (await store.HasRefAsync(branchName))
diff --git a/Engine/Source/Programs/Horde/Horde.Agent/Commands/Vcs/VcsCheckout.cs b/Engine/Source/Programs/Horde/Horde.Agent/Commands/Vcs/VcsCheckout.cs
--- a/Engine/Source/Programs/Horde/Horde.Agent/Commands/Vcs/VcsCheckout.cs
+++ b/Engine/Source/Programs/Horde/Horde.Agent/Commands/Vcs/VcsCheckout.cs
@@ -6,6 +6,7 @@ using System.Threading.Tasks;
 using EpicGames.Core;
 using EpicGames.Horde.Storage;
 using EpicGames.Horde.Storage.Nodes;
+using Microsoft.Extensions.Caching.Memory;
 using Microsoft.Extensions.Logging;
 using Microsoft.Extensions.Options;
 
@@ -26,8 +27,8 @@ namespace Horde.Agent.Commands.Vcs
 		[CommandLine("-Force")]
 		public bool Force { get; set; }
 
-		public VcsCheckoutCommand(IOptions<AgentSettings> settings)
-			: base(settings)
+		public VcsCheckoutCommand(IStorageClientFactory storageClientFactory)
+			: base(storageClientFactory)
 		{
 		}
 
@@ -55,9 +56,10 @@ namespace Horde.Agent.Commands.Vcs
 
 			RefName branchName = (Branch != null) ? new RefName(Branch) : workspaceState.Branch;
 
-			using IStorageClientOwner owner = CreateStorageClient(rootDir, logger);
-			IStorageClient store = owner.Store;
-			TreeReader reader = new TreeReader(owner.Store, owner.Cache, logger);
+			IStorageClient store = CreateStorageClient();
+
+			using MemoryCache cache = new MemoryCache(new MemoryCacheOptions());
+			TreeReader reader = new TreeReader(store, cache, logger);
 
 			CommitNode? tip = await GetCommitAsync(reader, branchName, Change);
 			if (tip == null)
@@ -80,7 +82,7 @@ namespace Horde.Agent.Commands.Vcs
 			DirectoryState newState = new DirectoryState();
 
 			DirectoryNode directoryNode = await directoryRef.ExpandAsync(reader);
-			foreach ((Utf8String name, DirectoryEntry? subDirEntry, DirectoryState? subDirState) in Zip(directoryNode.NameToDirectory, directoryState?.Directories))
+			foreach ((Utf8String name, DirectoryEntry? subDirEntry, DirectoryState? subDirState) in EnumerableExtensions.Zip(directoryNode.NameToDirectory, directoryState?.Directories))
 			{
 				DirectoryReference subDirPath = DirectoryReference.Combine(dirPath, name.ToString());
 				if (subDirEntry == null)
@@ -96,7 +98,7 @@ namespace Horde.Agent.Commands.Vcs
 				}
 			}
 
-			foreach ((Utf8String name, FileEntry? fileEntry, FileState? fileState) in Zip(directoryNode.NameToFile, directoryState?.Files))
+			foreach ((Utf8String name, FileEntry? fileEntry, FileState? fileState) in EnumerableExtensions.Zip(directoryNode.NameToFile, directoryState?.Files))
 			{
 				FileReference filePath = FileReference.Combine(dirPath, name.ToString());
 				if (fileEntry == null)
diff --git a/Engine/Source/Programs/Horde/Horde.Agent/Commands/Vcs/VcsClone.cs b/Engine/Source/Programs/Horde/Horde.Agent/Commands/Vcs/VcsClone.cs
--- a/Engine/Source/Programs/Horde/Horde.Agent/Commands/Vcs/VcsClone.cs
+++ b/Engine/Source/Programs/Horde/Horde.Agent/Commands/Vcs/VcsClone.cs
@@ -4,6 +4,7 @@ using System.Threading.Tasks;
 using EpicGames.Core;
 using EpicGames.Horde.Storage;
 using EpicGames.Horde.Storage.Nodes;
+using Microsoft.Extensions.Caching.Memory;
 using Microsoft.Extensions.Logging;
 using Microsoft.Extensions.Options;
 
@@ -12,8 +13,8 @@ namespace Horde.Agent.Commands.Vcs
 	[Command("vcs", "clone", "Initialize a directory for VCS-like operations")]
 	class VcsCloneCommand : VcsCheckoutCommand
 	{
-		public VcsCloneCommand(IOptions<AgentSettings> settings)
-			: base(settings)
+		public VcsCloneCommand(IStorageClientFactory storageClientFactory)
+			: base(storageClientFactory)
 		{
 		}
 
@@ -27,9 +28,10 @@ namespace Horde.Agent.Commands.Vcs
 
 			RefName branchName = new RefName(Branch ?? "ue5-main");
 
-			using IStorageClientOwner owner = CreateStorageClient(BaseDir, logger);
-			IStorageClient store = owner.Store;
-			TreeReader reader = new TreeReader(owner.Store, owner.Cache, logger);
+			IStorageClient store = CreateStorageClient();
+
+			using MemoryCache cache = new MemoryCache(new MemoryCacheOptions());
+			TreeReader reader = new TreeReader(store, cache, logger);
 
 			CommitNode? tip = await GetCommitAsync(reader, branchName, Change);
 			if (tip == null)
diff --git a/Engine/Source/Programs/Horde/Horde.Agent/Commands/Vcs/VcsCommit.cs b/Engine/Source/Programs/Horde/Horde.Agent/Commands/Vcs/VcsCommit.cs
--- a/Engine/Source/Programs/Horde/Horde.Agent/Commands/Vcs/VcsCommit.cs
+++ b/Engine/Source/Programs/Horde/Horde.Agent/Commands/Vcs/VcsCommit.cs
@@ -8,8 +8,8 @@ using System.Threading.Tasks;
 using EpicGames.Core;
 using EpicGames.Horde.Storage;
 using EpicGames.Horde.Storage.Nodes;
+using Microsoft.Extensions.Caching.Memory;
 using Microsoft.Extensions.Logging;
-using Microsoft.Extensions.Options;
 
 namespace Horde.Agent.Commands.Vcs
 {
@@ -22,8 +22,8 @@ namespace Horde.Agent.Commands.Vcs
 		[CommandLine("-Message=")]
 		public string? Message { get; set; }
 
-		public VcsCommitCommand(IOptions<AgentSettings> settings)
-			: base(settings)
+		public VcsCommitCommand(IStorageClientFactory storageClientFactory)
+			: base(storageClientFactory)
 		{
 		}
 
@@ -50,9 +50,10 @@ namespace Horde.Agent.Commands.Vcs
 
 			PrintDelta(oldState, newState, logger);
 
-			using IStorageClientOwner owner = CreateStorageClient(rootDir, logger);
-			IStorageClient store = owner.Store;
-			TreeReader reader = new TreeReader(owner.Store, owner.Cache, logger);
+			IStorageClient store = CreateStorageClient();
+
+			using MemoryCache cache = new MemoryCache(new MemoryCacheOptions());
+			TreeReader reader = new TreeReader(store, cache, logger);
 
 			CommitNode? tip = await reader.TryReadNodeAsync<CommitNode>(workspaceState.Branch);
 			TreeNodeRef<CommitNode>? tipRef = (tip == null) ? null : new TreeNodeRef<CommitNode>(tip);
@@ -111,7 +112,7 @@ namespace Horde.Agent.Commands.Vcs
 			directories.Add((rootRef, newState));
 
 			DirectoryNode root = await rootRef.ExpandAsync(reader);
-			foreach ((Utf8String name, DirectoryState? oldSubDirState, DirectoryState? newSubDirState) in Zip(oldState?.Directories, newState.Directories))
+			foreach ((Utf8String name, DirectoryState? oldSubDirState, DirectoryState? newSubDirState) in EnumerableExtensions.Zip(oldState?.Directories, newState.Directories))
 			{
 				if (newSubDirState == null)
 				{
@@ -123,7 +124,7 @@ namespace Horde.Agent.Commands.Vcs
 				}
 			}
 
-			foreach ((Utf8String name, FileState? oldFileState, FileState? newFileState) in Zip(oldState?.Files, newState.Files))
+			foreach ((Utf8String name, FileState? oldFileState, FileState? newFileState) in EnumerableExtensions.Zip(oldState?.Files, newState.Files))
 			{
 				if (newFileState == null)
 				{
diff --git a/Engine/Source/Programs/Horde/Horde.Agent/Commands/Vcs/VcsInit.cs b/Engine/Source/Programs/Horde/Horde.Agent/Commands/Vcs/VcsInit.cs
--- a/Engine/Source/Programs/Horde/Horde.Agent/Commands/Vcs/VcsInit.cs
+++ b/Engine/Source/Programs/Horde/Horde.Agent/Commands/Vcs/VcsInit.cs
@@ -2,16 +2,16 @@
 
 using System.Threading.Tasks;
 using EpicGames.Core;
+using EpicGames.Horde.Storage;
 using Microsoft.Extensions.Logging;
-using Microsoft.Extensions.Options;
 
 namespace Horde.Agent.Commands.Vcs
 {
-	[Command("Vcs", "Init", "Initialize a directory for VCS-like operations")]
+	[Command("vcs", "init", "Initialize a directory for VCS-like operations")]
 	class VcsInitCommand : VcsBase
 	{
-		public VcsInitCommand(IOptions<AgentSettings> settings)
-			: base(settings)
+		public VcsInitCommand(IStorageClientFactory storageClientFactory)
+			: base(storageClientFactory)
 		{
 		}
 
diff --git a/Engine/Source/Programs/Horde/Horde.Agent/Commands/Vcs/VcsLog.cs b/Engine/Source/Programs/Horde/Horde.Agent/Commands/Vcs/VcsLog.cs
--- a/Engine/Source/Programs/Horde/Horde.Agent/Commands/Vcs/VcsLog.cs
+++ b/Engine/Source/Programs/Horde/Horde.Agent/Commands/Vcs/VcsLog.cs
@@ -5,6 +5,7 @@ using System.Threading.Tasks;
 using EpicGames.Core;
 using EpicGames.Horde.Storage;
 using EpicGames.Horde.Storage.Nodes;
+using Microsoft.Extensions.Caching.Memory;
 using Microsoft.Extensions.Logging;
 using Microsoft.Extensions.Options;
 
@@ -16,8 +17,8 @@ namespace Horde.Agent.Commands.Vcs
 		[CommandLine("-Count")]
 		public int Count { get; set; } = 20;
 
-		public VcsLogCommand(IOptions<AgentSettings> settings)
-			: base(settings)
+		public VcsLogCommand(IStorageClientFactory storageClientFactory)
+			: base(storageClientFactory)
 		{
 		}
 
@@ -27,9 +28,10 @@ namespace Horde.Agent.Commands.Vcs
 
 			WorkspaceState workspaceState = await ReadStateAsync(rootDir);
 
-			using IStorageClientOwner owner = CreateStorageClient(rootDir, logger);
-			IStorageClient store = owner.Store;
-			TreeReader reader = new TreeReader(owner.Store, owner.Cache, logger);
+			IStorageClient store = CreateStorageClient();
+
+			using MemoryCache cache = new MemoryCache(new MemoryCacheOptions());
+			TreeReader reader = new TreeReader(store, cache, logger);
 
 			List<CommitNode> commits = new List<CommitNode>();
 
diff --git a/Engine/Source/Programs/Horde/Horde.Agent/Commands/Vcs/VcsStatus.cs b/Engine/Source/Programs/Horde/Horde.Agent/Commands/Vcs/VcsStatus.cs
--- a/Engine/Source/Programs/Horde/Horde.Agent/Commands/Vcs/VcsStatus.cs
+++ b/Engine/Source/Programs/Horde/Horde.Agent/Commands/Vcs/VcsStatus.cs
@@ -2,16 +2,16 @@
 
 using System.Threading.Tasks;
 using EpicGames.Core;
+using EpicGames.Horde.Storage;
 using Microsoft.Extensions.Logging;
-using Microsoft.Extensions.Options;
 
 namespace Horde.Agent.Commands.Vcs
 {
 	[Command("Vcs", "Status", "Find status of local files")]
 	class VcsStatusCommand : VcsBase
 	{
-		public VcsStatusCommand(IOptions<AgentSettings> settings)
-			: base(settings)
+		public VcsStatusCommand(IStorageClientFactory storageClientFactory)
+			: base(storageClientFactory)
 		{
 		}
 
diff --git a/Engine/Source/Programs/Horde/Horde.Agent/Program.cs b/Engine/Source/Programs/Horde/Horde.Agent/Program.cs
--- a/Engine/Source/Programs/Horde/Horde.Agent/Program.cs
+++ b/Engine/Source/Programs/Horde/Horde.Agent/Program.cs
@@ -197,6 +197,8 @@ namespace Horde.Agent
 			services.AddSingleton<IServerLoggerFactory, ServerLoggerFactory>();
 			services.AddHostedService<WorkerService>();
 
+			services.AddSingleton<IStorageClientFactory, StorageClientFactory>();
+
 			// Allow commands to augment the service collection for their own DI service providers
 			services.AddSingleton<DefaultServices>(x => new DefaultServices(config, services));
 
diff --git a/Engine/Source/Programs/Horde/Horde.Agent/Utility/StorageClientFactory.cs b/Engine/Source/Programs/Horde/Horde.Agent/Utility/StorageClientFactory.cs
new file mode 100644
--- /dev/null
+++ b/Engine/Source/Programs/Horde/Horde.Agent/Utility/StorageClientFactory.cs
@@ -0,0 +1,63 @@
+// Copyright Epic Games, Inc. All Rights Reserved.
+
+using System;
+using System.Net.Http.Headers;
+using System.Net.Http;
+using EpicGames.Horde.Storage;
+using Microsoft.Extensions.Logging;
+using EpicGames.Horde.Storage.Backends;
+using Microsoft.Extensions.Options;
+using EpicGames.Core;
+
+namespace Horde.Agent.Utility
+{
+	/// <summary>
+	/// Class which creates <see cref="IStorageClient"/> instances
+	/// </summary>
+	class StorageClientFactory : IStorageClientFactory
+	{
+		readonly IOptions<AgentSettings> _settings;
+		readonly ILogger _logger;
+
+		/// <summary>
+		/// Constructor
+		/// </summary>
+		public StorageClientFactory(IOptions<AgentSettings> settings, ILogger<IStorageClient> logger)
+		{
+			_settings = settings;
+			_logger = logger;
+		}
+
+		/// <summary>
+		/// Creates a client for the given namespace
+		/// </summary>
+		/// <param name="namespaceId">Namespace to get a client for</param>
+		/// <returns>New storage client instance</returns>
+		public IStorageClient Create(NamespaceId namespaceId)
+		{
+			if (_settings.Value.UseLocalStorageClient)
+			{
+				return new FileStorageClient(DirectoryReference.Combine(Program.DataDir, "Storage", namespaceId.ToString()), _logger);
+			}
+			else
+			{
+				return new HttpStorageClient(() => CreateDefaultHttpClient(namespaceId), () => new HttpClient(), _logger);
+			}
+		}
+
+		HttpClient CreateDefaultHttpClient(NamespaceId namespaceId)
+		{
+			ServerProfile profile = _settings.Value.GetCurrentServerProfile();
+
+			HttpClient client = new HttpClient();
+			client.BaseAddress = new Uri(profile.Url, $"api/v1/storage/{namespaceId}/");
+			if (!String.IsNullOrEmpty(profile.Token))
+			{
+				client.DefaultRequestHeaders.Authorization = new AuthenticationHeaderValue("Bearer", profile.Token);
+			}
+
+			return client;
+		}
+
+	}
+}
diff --git a/Engine/Source/Programs/Shared/EpicGames.Core/EnumerableExtensions.cs b/Engine/Source/Programs/Shared/EpicGames.Core/EnumerableExtensions.cs
--- a/Engine/Source/Programs/Shared/EpicGames.Core/EnumerableExtensions.cs
+++ b/Engine/Source/Programs/Shared/EpicGames.Core/EnumerableExtensions.cs
@@ -94,5 +94,46 @@ namespace EpicGames.Core
 
 			return maxElement;
 		}
+
+		/// <summary>
+		/// Zips two dictionaries, returning a sequence of keys with the old and new values (or null if they have been removed).
+		/// </summary>
+		/// <typeparam name="TKey">Common key between the two dictionaries</typeparam>
+		/// <typeparam name="TOldValue">Type of values in the first dictionary</typeparam>
+		/// <typeparam name="TNewValue">Type of values in the second dictionary</typeparam>
+		/// <param name="oldDictionary">The first dictionary</param>
+		/// <param name="newDictionary">The second dictionary</param>
+		/// <returns>Sequence of key/value pairs from each dictionary</returns>
+		public static IEnumerable<(TKey, TOldValue?, TNewValue?)> Zip<TKey, TOldValue, TNewValue>(this IReadOnlyDictionary<TKey, TOldValue>? oldDictionary, IReadOnlyDictionary<TKey, TNewValue>? newDictionary)
+			where TKey : notnull
+			where TOldValue : class
+			where TNewValue : class
+		{
+			if (newDictionary != null)
+			{
+				foreach ((TKey key, TNewValue newValue) in newDictionary)
+				{
+					TOldValue? oldValue;
+					if (oldDictionary == null || !oldDictionary.TryGetValue(key, out oldValue))
+					{
+						yield return (key, null, newValue);
+					}
+					else
+					{
+						yield return (key, oldValue, newValue);
+					}
+				}
+			}
+			if (oldDictionary != null)
+			{
+				foreach ((TKey key, TOldValue oldValue) in oldDictionary)
+				{
+					if (newDictionary == null || !newDictionary.ContainsKey(key))
+					{
+						yield return (key, oldValue, null);
+					}
+				}
+			}
+		}
 	}
 }
diff --git a/Engine/Source/Programs/Shared/EpicGames.Horde/Storage/IStorageClient.cs b/Engine/Source/Programs/Shared/EpicGames.Horde/Storage/IStorageClient.cs
--- a/Engine/Source/Programs/Shared/EpicGames.Horde/Storage/IStorageClient.cs
+++ b/Engine/Source/Programs/Shared/EpicGames.Horde/Storage/IStorageClient.cs
@@ -198,6 +198,19 @@ namespace EpicGames.Horde.Storage
 	{
 	}
 
+	/// <summary>
+	/// Allows creating storage clients for different namespaces
+	/// </summary>
+	public interface IStorageClientFactory
+	{
+		/// <summary>
+		/// Creates a storage client for the given namespace
+		/// </summary>
+		/// <param name="namespaceId">Namespace to manipulate</param>
+		/// <returns>Storage client instance</returns>
+		IStorageClient Create(NamespaceId namespaceId);
+	}
+
 	/// <summary>
 	/// Extension methods for <see cref="IStorageClient"/>
 	/// </summary>
