Horde: Rename command classes to match filenames.

diff --git a/Engine/Source/Programs/Horde/Horde.Agent/Commands/Bundles/BundleDump.cs b/Engine/Source/Programs/Horde/Horde.Agent/Commands/Bundles/BundleDump.cs
--- a/Engine/Source/Programs/Horde/Horde.Agent/Commands/Bundles/BundleDump.cs
+++ b/Engine/Source/Programs/Horde/Horde.Agent/Commands/Bundles/BundleDump.cs
@@ -2,9 +2,7 @@
 
 using System;
 using System.Collections.Generic;
-using System.IO;
 using System.Reflection;
-using System.Threading;
 using System.Threading.Tasks;
 using EpicGames.Core;
 using EpicGames.Horde.Storage;
diff --git a/Engine/Source/Programs/Horde/Horde.Agent/Commands/Vcs/VcsBranch.cs b/Engine/Source/Programs/Horde/Horde.Agent/Commands/Vcs/VcsBranch.cs
--- a/Engine/Source/Programs/Horde/Horde.Agent/Commands/Vcs/VcsBranch.cs
+++ b/Engine/Source/Programs/Horde/Horde.Agent/Commands/Vcs/VcsBranch.cs
@@ -4,17 +4,16 @@ using System.Threading.Tasks;
 using EpicGames.Core;
 using EpicGames.Horde.Storage;
 using Microsoft.Extensions.Logging;
-using Microsoft.Extensions.Options;
 
 namespace Horde.Agent.Commands.Vcs
 {
 	[Command("vcs", "branch", "Switch to a new branch")]
-	class VcsBranchCommand : VcsBase
+	class VcsBranch : VcsBase
 	{
 		[CommandLine(Prefix = "-Name=", Required = true)]
 		public string Name { get; set; } = "";
 
-		public VcsBranchCommand(IStorageClientFactory storageClientFactory)
+		public VcsBranch(IStorageClientFactory storageClientFactory)
 			: base(storageClientFactory)
 		{
 		}
diff --git a/Engine/Source/Programs/Horde/Horde.Agent/Commands/Vcs/VcsCheckout.cs b/Engine/Source/Programs/Horde/Horde.Agent/Commands/Vcs/VcsCheckout.cs
--- a/Engine/Source/Programs/Horde/Horde.Agent/Commands/Vcs/VcsCheckout.cs
+++ b/Engine/Source/Programs/Horde/Horde.Agent/Commands/Vcs/VcsCheckout.cs
@@ -8,12 +8,11 @@ using EpicGames.Horde.Storage;
 using EpicGames.Horde.Storage.Nodes;
 using Microsoft.Extensions.Caching.Memory;
 using Microsoft.Extensions.Logging;
-using Microsoft.Extensions.Options;
 
 namespace Horde.Agent.Commands.Vcs
 {
 	[Command("vcs", "checkout", "Checkout a particular branch/change")]
-	class VcsCheckoutCommand : VcsBase
+	class VcsCheckout : VcsBase
 	{
 		[CommandLine("-Branch")]
 		public string? Branch { get; set; }
@@ -27,7 +26,7 @@ namespace Horde.Agent.Commands.Vcs
 		[CommandLine("-Force")]
 		public bool Force { get; set; }
 
-		public VcsCheckoutCommand(IStorageClientFactory storageClientFactory)
+		public VcsCheckout(IStorageClientFactory storageClientFactory)
 			: base(storageClientFactory)
 		{
 		}
diff --git a/Engine/Source/Programs/Horde/Horde.Agent/Commands/Vcs/VcsClone.cs b/Engine/Source/Programs/Horde/Horde.Agent/Commands/Vcs/VcsClone.cs
--- a/Engine/Source/Programs/Horde/Horde.Agent/Commands/Vcs/VcsClone.cs
+++ b/Engine/Source/Programs/Horde/Horde.Agent/Commands/Vcs/VcsClone.cs
@@ -6,14 +6,13 @@ using EpicGames.Horde.Storage;
 using EpicGames.Horde.Storage.Nodes;
 using Microsoft.Extensions.Caching.Memory;
 using Microsoft.Extensions.Logging;
-using Microsoft.Extensions.Options;
 
 namespace Horde.Agent.Commands.Vcs
 {
 	[Command("vcs", "clone", "Initialize a directory for VCS-like operations")]
-	class VcsCloneCommand : VcsCheckoutCommand
+	class VcsClone : VcsCheckout
 	{
-		public VcsCloneCommand(IStorageClientFactory storageClientFactory)
+		public VcsClone(IStorageClientFactory storageClientFactory)
 			: base(storageClientFactory)
 		{
 		}
diff --git a/Engine/Source/Programs/Horde/Horde.Agent/Commands/Vcs/VcsCommit.cs b/Engine/Source/Programs/Horde/Horde.Agent/Commands/Vcs/VcsCommit.cs
--- a/Engine/Source/Programs/Horde/Horde.Agent/Commands/Vcs/VcsCommit.cs
+++ b/Engine/Source/Programs/Horde/Horde.Agent/Commands/Vcs/VcsCommit.cs
@@ -14,7 +14,7 @@ using Microsoft.Extensions.Logging;
 namespace Horde.Agent.Commands.Vcs
 {
 	[Command("vcs", "commit", "Commits data to the VCS store")]
-	class VcsCommitCommand : VcsBase
+	class VcsCommit : VcsBase
 	{
 		[CommandLine("-User=")]
 		public string UserName { get; set; } = Environment.UserName;
@@ -22,7 +22,7 @@ namespace Horde.Agent.Commands.Vcs
 		[CommandLine("-Message=")]
 		public string? Message { get; set; }
 
-		public VcsCommitCommand(IStorageClientFactory storageClientFactory)
+		public VcsCommit(IStorageClientFactory storageClientFactory)
 			: base(storageClientFactory)
 		{
 		}
diff --git a/Engine/Source/Programs/Horde/Horde.Agent/Commands/Vcs/VcsInit.cs b/Engine/Source/Programs/Horde/Horde.Agent/Commands/Vcs/VcsInit.cs
--- a/Engine/Source/Programs/Horde/Horde.Agent/Commands/Vcs/VcsInit.cs
+++ b/Engine/Source/Programs/Horde/Horde.Agent/Commands/Vcs/VcsInit.cs
@@ -8,9 +8,9 @@ using Microsoft.Extensions.Logging;
 namespace Horde.Agent.Commands.Vcs
 {
 	[Command("vcs", "init", "Initialize a directory for VCS-like operations")]
-	class VcsInitCommand : VcsBase
+	class VcsInit : VcsBase
 	{
-		public VcsInitCommand(IStorageClientFactory storageClientFactory)
+		public VcsInit(IStorageClientFactory storageClientFactory)
 			: base(storageClientFactory)
 		{
 		}
diff --git a/Engine/Source/Programs/Horde/Horde.Agent/Commands/Vcs/VcsLog.cs b/Engine/Source/Programs/Horde/Horde.Agent/Commands/Vcs/VcsLog.cs
--- a/Engine/Source/Programs/Horde/Horde.Agent/Commands/Vcs/VcsLog.cs
+++ b/Engine/Source/Programs/Horde/Horde.Agent/Commands/Vcs/VcsLog.cs
@@ -12,12 +12,12 @@ using Microsoft.Extensions.Options;
 namespace Horde.Agent.Commands.Vcs
 {
 	[Command("Vcs", "Log", "Print a history of commits")]
-	class VcsLogCommand : VcsBase
+	class VcsLog : VcsBase
 	{
 		[CommandLine("-Count")]
 		public int Count { get; set; } = 20;
 
-		public VcsLogCommand(IStorageClientFactory storageClientFactory)
+		public VcsLog(IStorageClientFactory storageClientFactory)
 			: base(storageClientFactory)
 		{
 		}
diff --git a/Engine/Source/Programs/Horde/Horde.Agent/Commands/Vcs/VcsStatus.cs b/Engine/Source/Programs/Horde/Horde.Agent/Commands/Vcs/VcsStatus.cs
--- a/Engine/Source/Programs/Horde/Horde.Agent/Commands/Vcs/VcsStatus.cs
+++ b/Engine/Source/Programs/Horde/Horde.Agent/Commands/Vcs/VcsStatus.cs
@@ -8,9 +8,9 @@ using Microsoft.Extensions.Logging;
 namespace Horde.Agent.Commands.Vcs
 {
 	[Command("Vcs", "Status", "Find status of local files")]
-	class VcsStatusCommand : VcsBase
+	class VcsStatus : VcsBase
 	{
-		public VcsStatusCommand(IStorageClientFactory storageClientFactory)
+		public VcsStatus(IStorageClientFactory storageClientFactory)
 			: base(storageClientFactory)
 		{
 		}
diff --git a/Engine/Source/Programs/Horde/Horde.Agent/Commands/Workspace/WorkspaceBase.cs b/Engine/Source/Programs/Horde/Horde.Agent/Commands/Workspace/WorkspaceBase.cs
--- a/Engine/Source/Programs/Horde/Horde.Agent/Commands/Workspace/WorkspaceBase.cs
+++ b/Engine/Source/Programs/Horde/Horde.Agent/Commands/Workspace/WorkspaceBase.cs
@@ -13,7 +13,7 @@ using Microsoft.Extensions.Logging;
 
 namespace Horde.Agent.Commands.Workspace
 {
-	abstract class WorkspaceCommand : Command
+	abstract class WorkspaceBase : Command
 	{
 		[CommandLine("-Server")]
 		[Description("Specifies the Perforce server and port")]
diff --git a/Engine/Source/Programs/Horde/Horde.Agent/Commands/Workspace/WorkspaceClean.cs b/Engine/Source/Programs/Horde/Horde.Agent/Commands/Workspace/WorkspaceClean.cs
--- a/Engine/Source/Programs/Horde/Horde.Agent/Commands/Workspace/WorkspaceClean.cs
+++ b/Engine/Source/Programs/Horde/Horde.Agent/Commands/Workspace/WorkspaceClean.cs
@@ -11,7 +11,7 @@ using Microsoft.Extensions.Logging;
 namespace Horde.Agent.Commands.Workspace
 {
 	[Command("Workspace", "Clean", "Cleans all modified files from the workspace.")]
-	class CleanCommand : WorkspaceCommand
+	class WorkspaceClean : WorkspaceBase
 	{
 		[CommandLine("-Incremental")]
 		[Description("Performs an incremental sync, without removing intermediates")]
diff --git a/Engine/Source/Programs/Horde/Horde.Agent/Commands/Workspace/WorkspaceClear.cs b/Engine/Source/Programs/Horde/Horde.Agent/Commands/Workspace/WorkspaceClear.cs
--- a/Engine/Source/Programs/Horde/Horde.Agent/Commands/Workspace/WorkspaceClear.cs
+++ b/Engine/Source/Programs/Horde/Horde.Agent/Commands/Workspace/WorkspaceClear.cs
@@ -10,7 +10,7 @@ using Microsoft.Extensions.Logging;
 namespace Horde.Agent.Commands.Workspace
 {
 	[Command("Workspace", "Clear", "Empties the staging directory of any files, returning them to the cache")]
-	class ClearCommand : WorkspaceCommand
+	class WorkspaceClear : WorkspaceBase
 	{
 		protected override Task ExecuteAsync(IPerforceConnection perforce, ManagedWorkspace repo, ILogger logger)
 		{
diff --git a/Engine/Source/Programs/Horde/Horde.Agent/Commands/Workspace/WorkspaceDump.cs b/Engine/Source/Programs/Horde/Horde.Agent/Commands/Workspace/WorkspaceDump.cs
--- a/Engine/Source/Programs/Horde/Horde.Agent/Commands/Workspace/WorkspaceDump.cs
+++ b/Engine/Source/Programs/Horde/Horde.Agent/Commands/Workspace/WorkspaceDump.cs
@@ -9,7 +9,7 @@ using Microsoft.Extensions.Logging;
 namespace Horde.Agent.Commands.Workspace
 {
 	[Command("Workspace", "Dump", "Dumps the contents of the repository to the log for analysis")]
-	class Dump : WorkspaceCommand
+	class WorkspaceDump : WorkspaceBase
 	{
 		protected override Task ExecuteAsync(IPerforceConnection perforce, ManagedWorkspace repo, ILogger logger)
 		{
diff --git a/Engine/Source/Programs/Horde/Horde.Agent/Commands/Workspace/WorkspacePopulateCache.cs b/Engine/Source/Programs/Horde/Horde.Agent/Commands/Workspace/WorkspacePopulateCache.cs
--- a/Engine/Source/Programs/Horde/Horde.Agent/Commands/Workspace/WorkspacePopulateCache.cs
+++ b/Engine/Source/Programs/Horde/Horde.Agent/Commands/Workspace/WorkspacePopulateCache.cs
@@ -13,7 +13,7 @@ using Microsoft.Extensions.Logging;
 namespace Horde.Agent.Commands.Workspace
 {
 	[Command("Workspace", "PopulateCache", "Populates the cache with the head revision of the given streams")]
-	class PopulateCacheCommand : WorkspaceCommand
+	class WorkspacePopulateCache : WorkspaceBase
 	{
 		[CommandLine("-ClientAndStream=")]
 		[Description("Specifies client and stream pairs, in the format Client:Stream")]
diff --git a/Engine/Source/Programs/Horde/Horde.Agent/Commands/Workspace/WorkspacePurgeCache.cs b/Engine/Source/Programs/Horde/Horde.Agent/Commands/Workspace/WorkspacePurgeCache.cs
--- a/Engine/Source/Programs/Horde/Horde.Agent/Commands/Workspace/WorkspacePurgeCache.cs
+++ b/Engine/Source/Programs/Horde/Horde.Agent/Commands/Workspace/WorkspacePurgeCache.cs
@@ -10,7 +10,7 @@ using Microsoft.Extensions.Logging;
 namespace Horde.Agent.Commands.Workspace
 {
 	[Command("Workspace", "PurgeCache", "Shrink the size of the cache to the given size")]
-	class PurgeCacheCommand : WorkspaceCommand
+	class WorkspacePurgeCache : WorkspaceBase
 	{
 		[CommandLine("-Size=")]
 		string? SizeParam { get; set; } = null;
diff --git a/Engine/Source/Programs/Horde/Horde.Agent/Commands/Workspace/WorkspaceRepair.cs b/Engine/Source/Programs/Horde/Horde.Agent/Commands/Workspace/WorkspaceRepair.cs
--- a/Engine/Source/Programs/Horde/Horde.Agent/Commands/Workspace/WorkspaceRepair.cs
+++ b/Engine/Source/Programs/Horde/Horde.Agent/Commands/Workspace/WorkspaceRepair.cs
@@ -10,7 +10,7 @@ using Microsoft.Extensions.Logging;
 namespace Horde.Agent.Commands.Workspace
 {
 	[Command("Workspace", "RepairCache", "Checks the integrity of the cache, and removes any invalid files")]
-	class RepairCommand : WorkspaceCommand
+	class WorkspaceRepair : WorkspaceBase
 	{
 		protected override Task ExecuteAsync(IPerforceConnection perforce, ManagedWorkspace repo, ILogger logger)
 		{
diff --git a/Engine/Source/Programs/Horde/Horde.Agent/Commands/Workspace/WorkspaceRevert.cs b/Engine/Source/Programs/Horde/Horde.Agent/Commands/Workspace/WorkspaceRevert.cs
--- a/Engine/Source/Programs/Horde/Horde.Agent/Commands/Workspace/WorkspaceRevert.cs
+++ b/Engine/Source/Programs/Horde/Horde.Agent/Commands/Workspace/WorkspaceRevert.cs
@@ -11,7 +11,7 @@ using Microsoft.Extensions.Logging;
 namespace Horde.Agent.Commands.Workspace
 {
 	[Command("Workspace", "Revert", "Revert all files that are open in the current workspace. Does not replace them with valid revisions.")]
-	class RevertCommand : WorkspaceCommand
+	class WorkspaceRevert : WorkspaceBase
 	{
 		[CommandLine("-Client=", Required = true)]
 		[Description("Client to revert all files for")]
diff --git a/Engine/Source/Programs/Horde/Horde.Agent/Commands/Workspace/WorkspaceSetup.cs b/Engine/Source/Programs/Horde/Horde.Agent/Commands/Workspace/WorkspaceSetup.cs
--- a/Engine/Source/Programs/Horde/Horde.Agent/Commands/Workspace/WorkspaceSetup.cs
+++ b/Engine/Source/Programs/Horde/Horde.Agent/Commands/Workspace/WorkspaceSetup.cs
@@ -11,7 +11,7 @@ using Microsoft.Extensions.Logging;
 namespace Horde.Agent.Commands.Workspace
 {
 	[Command("Workspace", "Setup", "Creates or updates a client to use a given stream")]
-	class SetupCommand : WorkspaceCommand
+	class WorkspaceSetup : WorkspaceBase
 	{
 		[CommandLine("-Client=", Required = true)]
 		[Description("Name of the client to create")]
diff --git a/Engine/Source/Programs/Horde/Horde.Agent/Commands/Workspace/WorkspaceSizes.cs b/Engine/Source/Programs/Horde/Horde.Agent/Commands/Workspace/WorkspaceSizes.cs
--- a/Engine/Source/Programs/Horde/Horde.Agent/Commands/Workspace/WorkspaceSizes.cs
+++ b/Engine/Source/Programs/Horde/Horde.Agent/Commands/Workspace/WorkspaceSizes.cs
@@ -12,7 +12,7 @@ using Microsoft.Extensions.Logging;
 namespace Horde.Agent.Commands.Workspace
 {
 	[Command("Workspace", "Sizes", "Gathers stats for which streams take the most amount of space in the cache for the given configuration")]
-	class SizesCommand : WorkspaceCommand
+	class WorkspaceSizes : WorkspaceBase
 	{
 		[CommandLine("-TempClient=", Required = true)]
 		[Description("Name of a temporary client to switch between streams gathering metadata. Will be created if it does not exist.")]
diff --git a/Engine/Source/Programs/Horde/Horde.Agent/Commands/Workspace/WorkspaceStatus.cs b/Engine/Source/Programs/Horde/Horde.Agent/Commands/Workspace/WorkspaceStatus.cs
--- a/Engine/Source/Programs/Horde/Horde.Agent/Commands/Workspace/WorkspaceStatus.cs
+++ b/Engine/Source/Programs/Horde/Horde.Agent/Commands/Workspace/WorkspaceStatus.cs
@@ -9,7 +9,7 @@ using Microsoft.Extensions.Logging;
 namespace Horde.Agent.Commands.Workspace
 {
 	[Command("Workspace", "Status", "Prints information about the state of the cache and workspace")]
-	class StatusCommand : WorkspaceCommand
+	class WorkspaceStatus : WorkspaceBase
 	{
 		protected override Task ExecuteAsync(IPerforceConnection perforce, ManagedWorkspace repo, ILogger logger)
 		{
diff --git a/Engine/Source/Programs/Horde/Horde.Agent/Commands/Workspace/WorkspaceSync.cs b/Engine/Source/Programs/Horde/Horde.Agent/Commands/Workspace/WorkspaceSync.cs
--- a/Engine/Source/Programs/Horde/Horde.Agent/Commands/Workspace/WorkspaceSync.cs
+++ b/Engine/Source/Programs/Horde/Horde.Agent/Commands/Workspace/WorkspaceSync.cs
@@ -12,7 +12,7 @@ using Microsoft.Extensions.Logging;
 namespace Horde.Agent.Commands.Workspace
 {
 	[Command("Workspace", "Sync", "Syncs the files for a particular stream and changelist")]
-	class SyncCommand : WorkspaceCommand
+	class WorkspaceSync : WorkspaceBase
 	{
 		[CommandLine("-Client=", Required = true)]
 		[Description("Name of the client to sync. Will be created if it does not exist.")]
