Horde: Include MD5 digests as custom data in replicated file entries.

diff --git a/Engine/Source/Programs/Horde/Horde.Build/Perforce/PerforceReplicator.cs b/Engine/Source/Programs/Horde/Horde.Build/Perforce/PerforceReplicator.cs
--- a/Engine/Source/Programs/Horde/Horde.Build/Perforce/PerforceReplicator.cs
+++ b/Engine/Source/Programs/Horde/Horde.Build/Perforce/PerforceReplicator.cs
@@ -4,6 +4,7 @@ using System;
 using System.Collections.Generic;
 using System.Diagnostics;
 using System.Linq;
+using System.Security.Cryptography;
 using System.Threading;
 using System.Threading.Tasks;
 using EpicGames.Core;
@@ -90,26 +91,63 @@ namespace Horde.Build.Perforce
 			}
 		}
 
-		class FileWriter
+		class FileWriter : IDisposable
 		{
-			readonly Utf8String _path;
-			FileNode _node;
-			readonly long _size;
+			class Handle
+			{
+				public Utf8String _path;
+				public FileNode _node = null!;
+				public long _size;
+				public readonly IncrementalHash _hash = IncrementalHash.CreateHash(HashAlgorithmName.MD5);
+			}
+
+			readonly Stack<Handle> _freeHandles = new Stack<Handle>();
+			readonly Dictionary<int, Handle> _openHandles = new Dictionary<int, Handle>();
+
+			public void Dispose()
+			{
+				foreach (Handle handle in _freeHandles)
+				{
+					handle._hash.Dispose();
+				}
+				foreach (Handle handle in _openHandles.Values)
+				{
+					handle._hash.Dispose();
+				}
+			}
+
+			public void Open(int fd, Utf8String path, long size)
+			{
+				Handle handle = (_freeHandles.Count > 0)? _freeHandles.Pop() : new Handle();
+
+				handle._path = path;
+				handle._node = new LeafFileNode();
+				handle._size = size;
 
-			public Utf8String Path => _path;
-			public FileNode Node => _node;
-			public long Size => _size;
+				_openHandles.Add(fd, handle);
+			}
 
-			public FileWriter(Utf8String path, long size)
+			public async Task AppendAsync(int fd, ReadOnlyMemory<byte> data, ChunkingOptions options, TreeWriter writer, CancellationToken cancellationToken)
 			{
-				_path = path;
-				_node = new LeafFileNode();
-				_size = size;
+				Handle handle = _openHandles[fd];
+				handle._hash.AppendData(data.Span);
+				handle._node = await handle._node.AppendAsync(data, options, writer, cancellationToken);
 			}
 
-			public async Task AppendAsync(ReadOnlyMemory<byte> data, ChunkingOptions options, TreeWriter writer, CancellationToken cancellationToken)
+			public void Close(int fd, out Utf8String path, out FileNode node, out byte[] hash)
 			{
-				_node = await _node.AppendAsync(data, options, writer, cancellationToken);
+				Handle handle = _openHandles[fd];
+				if (handle._node.Length != handle._size)
+				{
+					throw new ReplicationException($"Invalid size for replicated file '{handle._path}'. Expected {handle._size}, got {handle._node.Length}.");
+				}
+
+				path = handle._path;
+				node = handle._node;
+				hash = handle._hash.GetHashAndReset();
+
+				_openHandles.Remove(fd);
+				_freeHandles.Push(handle);
 			}
 		}
 
@@ -311,7 +349,7 @@ namespace Horde.Build.Perforce
 				Stopwatch processTimer = new Stopwatch();
 				Stopwatch gcTimer = new Stopwatch();
 
-				Dictionary<int, FileWriter> handles = new Dictionary<int, FileWriter>();
+				using FileWriter fileWriter = new FileWriter();
 				await foreach (PerforceResponse response in perforce.StreamCommandAsync("sync", Array.Empty<string>(), syncPaths, null, typeof(SyncRecord), true, default))
 				{
 					PerforceError? error = response.Error;
@@ -347,27 +385,19 @@ namespace Horde.Build.Perforce
 								throw new ReplicationException($"Unable to find file entry for {file}");
 							}
 
-							handles[io.File] = new FileWriter(file, fileSize);
+							fileWriter.Open(io.File, file, fileSize);
 						}
 						else if (io.Command == PerforceIoCommand.Write)
 						{
-							FileWriter file = handles[io.File];
-							await file.AppendAsync(io.Payload, options.ChunkingOptions, writer, cancellationToken);
+							await fileWriter.AppendAsync(io.File, io.Payload, options.ChunkingOptions, writer, cancellationToken);
 						}
 						else if (io.Command == PerforceIoCommand.Close)
 						{
-							FileWriter? file = null;
-							if (handles.Remove(io.File, out file))
-							{
-								FileNode node = file.Node;
-								if (node.Length != file.Size)
-								{
-									throw new ReplicationException($"Invalid size for replicated file '{file.Path}'. Expected {file.Size}, got {node.Length}.");
-								}
+							fileWriter.Close(io.File, out Utf8String path, out FileNode node, out byte[] hash);
 
-								FileEntry entry = await root.AddFileByPathAsync(reader, file.Path, FileEntryFlags.None, node, cancellationToken);
-								await writer.WriteAsync(entry, cancellationToken);
-							}
+							FileEntry entry = await root.AddFileByPathAsync(reader, path, FileEntryFlags.None, node, cancellationToken);
+							entry.CustomData = hash;
+							await writer.WriteAsync(entry, cancellationToken);
 						}
 						else if (io.Command == PerforceIoCommand.Unlink)
 						{
diff --git a/Engine/Source/Programs/Shared/EpicGames.Horde/Storage/Nodes/DirectoryNode.cs b/Engine/Source/Programs/Shared/EpicGames.Horde/Storage/Nodes/DirectoryNode.cs
--- a/Engine/Source/Programs/Shared/EpicGames.Horde/Storage/Nodes/DirectoryNode.cs
+++ b/Engine/Source/Programs/Shared/EpicGames.Horde/Storage/Nodes/DirectoryNode.cs
@@ -72,7 +72,7 @@ namespace EpicGames.Horde.Storage.Nodes
 		/// <summary>
 		/// Custom user data for this file entry
 		/// </summary>
-		public ReadOnlyMemory<byte> CustomData { get; } = ReadOnlyMemory<byte>.Empty;
+		public ReadOnlyMemory<byte> CustomData { get; set; } = ReadOnlyMemory<byte>.Empty;
 
 		/// <summary>
 		/// Cached hash of the target node
