Horde: Support IStorageClientFactory through StorageService.

diff --git a/Engine/Source/Programs/Horde/Horde.Agent/Commands/Bundles/BundleCreate.cs b/Engine/Source/Programs/Horde/Horde.Agent/Commands/Bundles/BundleCreate.cs
--- a/Engine/Source/Programs/Horde/Horde.Agent/Commands/Bundles/BundleCreate.cs
+++ b/Engine/Source/Programs/Horde/Horde.Agent/Commands/Bundles/BundleCreate.cs
@@ -5,7 +5,6 @@ using System.Threading;
 using System.Threading.Tasks;
 using EpicGames.Core;
 using EpicGames.Horde.Storage;
-using EpicGames.Horde.Storage.Backends;
 using EpicGames.Horde.Storage.Nodes;
 using Microsoft.Extensions.Logging;
 
@@ -14,15 +13,25 @@ namespace Horde.Agent.Commands.Bundles
 	[Command("bundle", "create", "Creates a bundle from a folder on the local hard drive")]
 	class BundleCreate : Command
 	{
+		[CommandLine("-Namespace=", Description = "Namespace to use for storage")]
+		public NamespaceId NamespaceId { get; set; } = new NamespaceId("default");
+
 		[CommandLine("-Ref=")]
 		public RefName RefName { get; set; } = new RefName("default-bundle");
 
 		[CommandLine("-InputDir=", Required = true)]
 		public DirectoryReference InputDir { get; set; } = null!;
 
+		readonly IStorageClientFactory _storageClientFactory;
+
+		public BundleCreate(IStorageClientFactory storageClientFactory)
+		{
+			_storageClientFactory = storageClientFactory;
+		}
+
 		public override async Task<int> ExecuteAsync(ILogger logger)
 		{
-			FileStorageClient store = new FileStorageClient(DirectoryReference.Combine(Program.DataDir, "Bundles"), logger);
+			IStorageClient store = await _storageClientFactory.GetClientAsync(NamespaceId);
 
 			TreeWriter writer = new TreeWriter(store, prefix: RefName.Text);
 
diff --git a/Engine/Source/Programs/Horde/Horde.Agent/Commands/Bundles/BundleDump.cs b/Engine/Source/Programs/Horde/Horde.Agent/Commands/Bundles/BundleDump.cs
--- a/Engine/Source/Programs/Horde/Horde.Agent/Commands/Bundles/BundleDump.cs
+++ b/Engine/Source/Programs/Horde/Horde.Agent/Commands/Bundles/BundleDump.cs
@@ -47,13 +47,13 @@ namespace Horde.Agent.Commands.Bundles
 			}
 			else if (BlobId != null)
 			{
-				IStorageClient store = _storageClientFactory.Create(NamespaceId);
+				IStorageClient store = await _storageClientFactory.GetClientAsync(NamespaceId);
 				bundle = await store.ReadBundleAsync(BlobId.Value);
 				logger.LogInformation("Summary for blob {BlobId}", BlobId.Value);
 			}
 			else if (RefName != null)
 			{
-				IStorageClient store = _storageClientFactory.Create(NamespaceId);
+				IStorageClient store = await _storageClientFactory.GetClientAsync(NamespaceId);
 
 				NodeLocator locator = await store.ReadRefTargetAsync(RefName.Value);
 				logger.LogInformation("Ref {RefName} -> {Locator}", RefName.Value, locator);
diff --git a/Engine/Source/Programs/Horde/Horde.Agent/Commands/Bundles/BundleExtract.cs b/Engine/Source/Programs/Horde/Horde.Agent/Commands/Bundles/BundleExtract.cs
--- a/Engine/Source/Programs/Horde/Horde.Agent/Commands/Bundles/BundleExtract.cs
+++ b/Engine/Source/Programs/Horde/Horde.Agent/Commands/Bundles/BundleExtract.cs
@@ -5,7 +5,6 @@ using System.Threading;
 using System.Threading.Tasks;
 using EpicGames.Core;
 using EpicGames.Horde.Storage;
-using EpicGames.Horde.Storage.Backends;
 using EpicGames.Horde.Storage.Nodes;
 using Microsoft.Extensions.Caching.Memory;
 using Microsoft.Extensions.Logging;
@@ -15,15 +14,25 @@ namespace Horde.Agent.Commands.Bundles
 	[Command("bundle", "extract", "Extracts data from a bundle to the local hard drive")]
 	internal class BundleExtract : Command
 	{
+		[CommandLine("-Namespace=", Description = "Namespace to use for storage")]
+		public NamespaceId NamespaceId { get; set; } = new NamespaceId("default");
+
 		[CommandLine("-Ref=")]
 		public RefName RefName { get; set; } = new RefName("default-bundle");
 
 		[CommandLine("-OutputDir=", Required = true)]
 		public DirectoryReference OutputDir { get; set; } = null!;
 
+		readonly IStorageClientFactory _storageClientFactory;
+
+		public BundleExtract(IStorageClientFactory storageClientFactory)
+		{
+			_storageClientFactory = storageClientFactory;
+		}
+
 		public override async Task<int> ExecuteAsync(ILogger logger)
 		{
-			FileStorageClient store = new FileStorageClient(DirectoryReference.Combine(Program.DataDir, "Bundles"), logger);
+			IStorageClient store = await _storageClientFactory.GetClientAsync(NamespaceId);
 
 			using MemoryCache cache = new MemoryCache(new MemoryCacheOptions());
 			TreeReader reader = new TreeReader(store, cache, logger);
diff --git a/Engine/Source/Programs/Horde/Horde.Agent/Commands/Vcs/VcsBase.cs b/Engine/Source/Programs/Horde/Horde.Agent/Commands/Vcs/VcsBase.cs
--- a/Engine/Source/Programs/Horde/Horde.Agent/Commands/Vcs/VcsBase.cs
+++ b/Engine/Source/Programs/Horde/Horde.Agent/Commands/Vcs/VcsBase.cs
@@ -162,7 +162,7 @@ namespace Horde.Agent.Commands.Vcs
 			_storageClientFactory = storageClientFactory;
 		}
 
-		protected IStorageClient CreateStorageClient() => _storageClientFactory.Create(NamespaceId);
+		protected ValueTask<IStorageClient> GetStorageClientAsync() => _storageClientFactory.GetClientAsync(NamespaceId);
 
 		protected static async Task<WorkspaceState> ReadStateAsync(DirectoryReference rootDir)
 		{
diff --git a/Engine/Source/Programs/Horde/Horde.Agent/Commands/Vcs/VcsBranch.cs b/Engine/Source/Programs/Horde/Horde.Agent/Commands/Vcs/VcsBranch.cs
--- a/Engine/Source/Programs/Horde/Horde.Agent/Commands/Vcs/VcsBranch.cs
+++ b/Engine/Source/Programs/Horde/Horde.Agent/Commands/Vcs/VcsBranch.cs
@@ -24,7 +24,7 @@ namespace Horde.Agent.Commands.Vcs
 
 			WorkspaceState workspaceState = await ReadStateAsync(rootDir);
 
-			IStorageClient store = CreateStorageClient();
+			IStorageClient store = await GetStorageClientAsync();
 
 			RefName branchName = new RefName(Name);
 			if (await store.HasRefAsync(branchName))
diff --git a/Engine/Source/Programs/Horde/Horde.Agent/Commands/Vcs/VcsCheckout.cs b/Engine/Source/Programs/Horde/Horde.Agent/Commands/Vcs/VcsCheckout.cs
--- a/Engine/Source/Programs/Horde/Horde.Agent/Commands/Vcs/VcsCheckout.cs
+++ b/Engine/Source/Programs/Horde/Horde.Agent/Commands/Vcs/VcsCheckout.cs
@@ -55,7 +55,7 @@ namespace Horde.Agent.Commands.Vcs
 
 			RefName branchName = (Branch != null) ? new RefName(Branch) : workspaceState.Branch;
 
-			IStorageClient store = CreateStorageClient();
+			IStorageClient store = await GetStorageClientAsync();
 
 			using MemoryCache cache = new MemoryCache(new MemoryCacheOptions());
 			TreeReader reader = new TreeReader(store, cache, logger);
diff --git a/Engine/Source/Programs/Horde/Horde.Agent/Commands/Vcs/VcsClone.cs b/Engine/Source/Programs/Horde/Horde.Agent/Commands/Vcs/VcsClone.cs
--- a/Engine/Source/Programs/Horde/Horde.Agent/Commands/Vcs/VcsClone.cs
+++ b/Engine/Source/Programs/Horde/Horde.Agent/Commands/Vcs/VcsClone.cs
@@ -27,7 +27,7 @@ namespace Horde.Agent.Commands.Vcs
 
 			RefName branchName = new RefName(Branch ?? "ue5-main");
 
-			IStorageClient store = CreateStorageClient();
+			IStorageClient store = await GetStorageClientAsync();
 
 			using MemoryCache cache = new MemoryCache(new MemoryCacheOptions());
 			TreeReader reader = new TreeReader(store, cache, logger);
diff --git a/Engine/Source/Programs/Horde/Horde.Agent/Commands/Vcs/VcsCommit.cs b/Engine/Source/Programs/Horde/Horde.Agent/Commands/Vcs/VcsCommit.cs
--- a/Engine/Source/Programs/Horde/Horde.Agent/Commands/Vcs/VcsCommit.cs
+++ b/Engine/Source/Programs/Horde/Horde.Agent/Commands/Vcs/VcsCommit.cs
@@ -50,7 +50,7 @@ namespace Horde.Agent.Commands.Vcs
 
 			PrintDelta(oldState, newState, logger);
 
-			IStorageClient store = CreateStorageClient();
+			IStorageClient store = await GetStorageClientAsync();
 
 			using MemoryCache cache = new MemoryCache(new MemoryCacheOptions());
 			TreeReader reader = new TreeReader(store, cache, logger);
diff --git a/Engine/Source/Programs/Horde/Horde.Agent/Commands/Vcs/VcsLog.cs b/Engine/Source/Programs/Horde/Horde.Agent/Commands/Vcs/VcsLog.cs
--- a/Engine/Source/Programs/Horde/Horde.Agent/Commands/Vcs/VcsLog.cs
+++ b/Engine/Source/Programs/Horde/Horde.Agent/Commands/Vcs/VcsLog.cs
@@ -7,7 +7,6 @@ using EpicGames.Horde.Storage;
 using EpicGames.Horde.Storage.Nodes;
 using Microsoft.Extensions.Caching.Memory;
 using Microsoft.Extensions.Logging;
-using Microsoft.Extensions.Options;
 
 namespace Horde.Agent.Commands.Vcs
 {
@@ -28,7 +27,7 @@ namespace Horde.Agent.Commands.Vcs
 
 			WorkspaceState workspaceState = await ReadStateAsync(rootDir);
 
-			IStorageClient store = CreateStorageClient();
+			IStorageClient store = await GetStorageClientAsync();
 
 			using MemoryCache cache = new MemoryCache(new MemoryCacheOptions());
 			TreeReader reader = new TreeReader(store, cache, logger);
diff --git a/Engine/Source/Programs/Horde/Horde.Agent/Utility/StorageClientFactory.cs b/Engine/Source/Programs/Horde/Horde.Agent/Utility/StorageClientFactory.cs
--- a/Engine/Source/Programs/Horde/Horde.Agent/Utility/StorageClientFactory.cs
+++ b/Engine/Source/Programs/Horde/Horde.Agent/Utility/StorageClientFactory.cs
@@ -8,6 +8,8 @@ using Microsoft.Extensions.Logging;
 using EpicGames.Horde.Storage.Backends;
 using Microsoft.Extensions.Options;
 using EpicGames.Core;
+using System.Threading;
+using System.Threading.Tasks;
 
 namespace Horde.Agent.Utility
 {
@@ -28,21 +30,19 @@ namespace Horde.Agent.Utility
 			_logger = logger;
 		}
 
-		/// <summary>
-		/// Creates a client for the given namespace
-		/// </summary>
-		/// <param name="namespaceId">Namespace to get a client for</param>
-		/// <returns>New storage client instance</returns>
-		public IStorageClient Create(NamespaceId namespaceId)
+		/// <inheritdoc/>
+		public ValueTask<IStorageClient> GetClientAsync(NamespaceId namespaceId, CancellationToken cancellationToken)
 		{
+			IStorageClient client;
 			if (_settings.Value.UseLocalStorageClient)
 			{
-				return new FileStorageClient(DirectoryReference.Combine(Program.DataDir, "Storage", namespaceId.ToString()), _logger);
+				client = new FileStorageClient(DirectoryReference.Combine(Program.DataDir, "Storage", namespaceId.ToString()), _logger);
 			}
 			else
 			{
-				return new HttpStorageClient(() => CreateDefaultHttpClient(namespaceId), () => new HttpClient(), _logger);
+				client = new HttpStorageClient(() => CreateDefaultHttpClient(namespaceId), () => new HttpClient(), _logger);
 			}
+			return new ValueTask<IStorageClient>(client);
 		}
 
 		HttpClient CreateDefaultHttpClient(NamespaceId namespaceId)
@@ -58,6 +58,5 @@ namespace Horde.Agent.Utility
 
 			return client;
 		}
-
 	}
 }
diff --git a/Engine/Source/Programs/Horde/Horde.Build/Startup.cs b/Engine/Source/Programs/Horde/Horde.Build/Startup.cs
--- a/Engine/Source/Programs/Horde/Horde.Build/Startup.cs
+++ b/Engine/Source/Programs/Horde/Horde.Build/Startup.cs
@@ -441,6 +441,7 @@ namespace Horde.Build
 			services.AddSingleton<DeviceService>();						
 			services.AddSingleton<NoticeService>();
 			services.AddSingleton<StorageService>();
+			services.AddSingleton<IStorageClientFactory>(sp => sp.GetRequiredService<StorageService>());
 			services.AddSingleton<TestDataService>();
 
 			if (settings.JiraUrl != null)
diff --git a/Engine/Source/Programs/Horde/Horde.Build/Storage/StorageService.cs b/Engine/Source/Programs/Horde/Horde.Build/Storage/StorageService.cs
--- a/Engine/Source/Programs/Horde/Horde.Build/Storage/StorageService.cs
+++ b/Engine/Source/Programs/Horde/Horde.Build/Storage/StorageService.cs
@@ -89,7 +89,7 @@ namespace Horde.Build.Storage
 	/// <summary>
 	/// Functionality related to the storage service
 	/// </summary>
-	public sealed class StorageService : IHostedService, IDisposable
+	public sealed class StorageService : IHostedService, IDisposable, IStorageClientFactory
 	{
 		sealed class StorageClient : StorageClientBase, IStorageClientImpl, IDisposable
 		{
@@ -488,6 +488,9 @@ namespace Horde.Build.Storage
 			return namespaceInfo.Client;
 		}
 
+		/// <inheritdoc/>
+		async ValueTask<IStorageClient> IStorageClientFactory.GetClientAsync(NamespaceId namespaceId, CancellationToken cancellationToken) => await GetClientAsync(namespaceId, cancellationToken);
+
 		/// <summary>
 		/// Authorize operations for the given store
 		/// </summary>
diff --git a/Engine/Source/Programs/Shared/EpicGames.Horde/Storage/IStorageClient.cs b/Engine/Source/Programs/Shared/EpicGames.Horde/Storage/IStorageClient.cs
--- a/Engine/Source/Programs/Shared/EpicGames.Horde/Storage/IStorageClient.cs
+++ b/Engine/Source/Programs/Shared/EpicGames.Horde/Storage/IStorageClient.cs
@@ -207,8 +207,9 @@ namespace EpicGames.Horde.Storage
 		/// Creates a storage client for the given namespace
 		/// </summary>
 		/// <param name="namespaceId">Namespace to manipulate</param>
+		/// <param name="cancellationToken">Cancellation token for the operation</param>
 		/// <returns>Storage client instance</returns>
-		IStorageClient Create(NamespaceId namespaceId);
+		ValueTask<IStorageClient> GetClientAsync(NamespaceId namespaceId, CancellationToken cancellationToken = default);
 	}
 
 	/// <summary>
