Horde: Add a synthetic benchmark to write chunked bundle data through to /dev/null.

diff --git a/Engine/Source/Programs/Horde/Horde.Build/Commands/Bundles/BundleWriteTest.cs b/Engine/Source/Programs/Horde/Horde.Build/Commands/Bundles/BundleWriteTest.cs
new file mode 100644
--- /dev/null
+++ b/Engine/Source/Programs/Horde/Horde.Build/Commands/Bundles/BundleWriteTest.cs
@@ -0,0 +1,64 @@
+// Copyright Epic Games, Inc. All Rights Reserved.
+
+using System;
+using System.Diagnostics;
+using System.IO;
+using System.Security.Cryptography;
+using System.Threading;
+using System.Threading.Tasks;
+using EpicGames.Core;
+using EpicGames.Horde.Storage;
+using EpicGames.Horde.Storage.Nodes;
+using Microsoft.Extensions.Logging;
+
+namespace Horde.Build.Commands.Bundles
+{
+	[Command("bundle", "writetest", "Synthetic benchmark for bundle write performance")]
+	class BundleWriteTestCommand : Command
+	{
+		class FakeStorageClient : IStorageClient
+		{
+			public Task DeleteRefAsync(RefName name, CancellationToken cancellationToken = default) => Task.CompletedTask;
+			public Task<Stream> ReadBlobAsync(BlobLocator locator, CancellationToken cancellationToken = default) => throw new NotImplementedException();
+			public Task<Stream> ReadBlobRangeAsync(BlobLocator locator, int offset, int length, CancellationToken cancellationToken = default) => throw new NotImplementedException();
+			public Task<NodeLocator> TryReadRefTargetAsync(RefName name, DateTime cacheTime = default, CancellationToken cancellationToken = default) => throw new NotImplementedException();
+			public Task<BlobLocator> WriteBlobAsync(Stream stream, Utf8String prefix = default, CancellationToken cancellationToken = default) => Task.FromResult(BlobLocator.Create(HostId.Empty));
+			public Task<NodeLocator> WriteRefAsync(RefName name, Bundle bundle, int exportIdx, Utf8String prefix = default, RefOptions? options = null, CancellationToken cancellationToken = default) => throw new NotImplementedException();
+			public Task WriteRefTargetAsync(RefName name, NodeLocator target, RefOptions? options = null, CancellationToken cancellationToken = default) => Task.CompletedTask;
+		}
+
+		public override async Task<int> ExecuteAsync(ILogger logger)
+		{
+			FakeStorageClient store = new FakeStorageClient();
+
+			TreeOptions options = new TreeOptions();
+			options.CompressionFormat = BundleCompressionFormat.None;
+			TreeWriter writer = new TreeWriter(store, options);
+
+			ChunkingOptions chunkingOptions = new ChunkingOptions();
+
+			FileNode node = new LeafFileNode();
+
+			byte[] buffer = new byte[64 * 1024];
+			RandomNumberGenerator.Fill(buffer);
+
+			Stopwatch timer = Stopwatch.StartNew();
+			long length = 0;
+			double nextTime = 2.0;
+
+			for (; ; )
+			{
+				node = await node.AppendAsync(buffer, chunkingOptions, writer, default);
+				length += buffer.Length;
+
+				double time = timer.Elapsed.TotalSeconds;
+				if (time > nextTime)
+				{
+					long size = GC.GetTotalMemory(true);
+					logger.LogInformation("Written {Length:n0}mb in {Time:0.000}s ({Rate:n0}mb/s) (heap size: {Size:n0}mb)", length / (1024 * 1024), time, length / (1024 * 1024 * time), size / (1024.0 * 1024.0));
+					nextTime = time + 2.0;
+				}
+			}
+		}
+	}
+}
