Horde: Propagate workspace method from stream config to RPC message

diff --git a/Engine/Source/Programs/Horde/Horde.Build/Agents/IAgent.cs b/Engine/Source/Programs/Horde/Horde.Build/Agents/IAgent.cs
--- a/Engine/Source/Programs/Horde/Horde.Build/Agents/IAgent.cs
+++ b/Engine/Source/Programs/Horde/Horde.Build/Agents/IAgent.cs
@@ -169,19 +169,23 @@ namespace Horde.Build.Agents
 		public HordeCommon.Rpc.Messages.AgentWorkspace ToRpcMessage(IPerforceServer server, PerforceCredentials? credentials)
 		{
 			// Construct the message
-			HordeCommon.Rpc.Messages.AgentWorkspace result = new HordeCommon.Rpc.Messages.AgentWorkspace();
-			result.ConfiguredCluster = Cluster;
-			result.ConfiguredUserName = UserName;
-			result.ServerAndPort = server.ServerAndPort;
-			result.UserName = credentials?.UserName ?? UserName;
-			result.Password = credentials?.Password;
-			result.Identifier = Identifier;
-			result.Stream = Stream;
+			HordeCommon.Rpc.Messages.AgentWorkspace result = new ()
+			{
+				ConfiguredCluster = Cluster, ConfiguredUserName = UserName,
+				ServerAndPort = server.ServerAndPort,
+				UserName = credentials?.UserName ?? UserName,
+				Password = credentials?.Password,
+				Identifier = Identifier,
+				Stream = Stream,
+				Incremental = Incremental,
+				Method = Method
+			};
+			
 			if (View != null)
 			{
 				result.View.AddRange(View);
 			}
-			result.Incremental = Incremental;
+			
 			return result;
 		}
 	}
@@ -753,23 +757,26 @@ namespace Horde.Build.Agents
 			}
 
 			// Construct the message
-			HordeCommon.Rpc.Messages.AgentWorkspace result = new HordeCommon.Rpc.Messages.AgentWorkspace();
-			result.ConfiguredCluster = workspace.Cluster;
-			result.ConfiguredUserName = workspace.UserName;
-			result.Cluster = cluster?.Name;
-			result.BaseServerAndPort = baseServerAndPort;
-			result.ServerAndPort = serverAndPort;
-			result.UserName = credentials?.UserName ?? workspace.UserName;
-			result.Password = credentials?.Password;
-			result.Identifier = workspace.Identifier;
-			result.Stream = workspace.Stream;
+			HordeCommon.Rpc.Messages.AgentWorkspace result = new ()
+			{
+				ConfiguredCluster = workspace.Cluster, ConfiguredUserName = workspace.UserName,
+				Cluster = cluster?.Name,
+				BaseServerAndPort = baseServerAndPort,
+				ServerAndPort = serverAndPort,
+				UserName = credentials?.UserName ?? workspace.UserName,
+				Password = credentials?.Password,
+				Identifier = workspace.Identifier,
+				Stream = workspace.Stream,
+				Incremental = workspace.Incremental,
+				Method = workspace.Method
+			};
+
 			if (workspace.View != null)
 			{
 				result.View.AddRange(workspace.View);
 			}
-			result.Incremental = workspace.Incremental;
+			
 			workspaceMessages.Add(result);
-
 			return true;
 		}
 	}
diff --git a/Engine/Source/Programs/Horde/Horde.Build/Jobs/JobTaskSource.cs b/Engine/Source/Programs/Horde/Horde.Build/Jobs/JobTaskSource.cs
--- a/Engine/Source/Programs/Horde/Horde.Build/Jobs/JobTaskSource.cs
+++ b/Engine/Source/Programs/Horde/Horde.Build/Jobs/JobTaskSource.cs
@@ -857,7 +857,7 @@ namespace Horde.Build.Jobs
 			task.JobName = leaseName.ToString();
 			task.Executor = job.Executor ?? String.Empty;
 
-			List<HordeCommon.Rpc.Messages.AgentWorkspace> workspaces = new List<HordeCommon.Rpc.Messages.AgentWorkspace>();
+			List<HordeCommon.Rpc.Messages.AgentWorkspace> workspaces = new ();
 
 			PerforceCluster? cluster = globals.Config.FindPerforceCluster(workspace.Cluster);
 			if (cluster == null)
@@ -868,6 +868,8 @@ namespace Horde.Build.Jobs
 			AgentWorkspace? autoSdkWorkspace = useAutoSdk ? agent.GetAutoSdkWorkspace(cluster) : null;
 			if (autoSdkWorkspace != null)
 			{
+				autoSdkWorkspace.Method = workspace.Method;
+				
 				if (!await agent.TryAddWorkspaceMessage(autoSdkWorkspace, cluster, _perforceLoadBalancer, workspaces))
 				{
 					return null;
