Horde: Expose method from workspace config in REST API

diff --git a/Engine/Source/Programs/Horde/Horde.Build/Agents/AgentMessages.cs b/Engine/Source/Programs/Horde/Horde.Build/Agents/AgentMessages.cs
--- a/Engine/Source/Programs/Horde/Horde.Build/Agents/AgentMessages.cs
+++ b/Engine/Source/Programs/Horde/Horde.Build/Agents/AgentMessages.cs
@@ -288,6 +288,11 @@ namespace Horde.Build.Agents
 		/// Whether to use an incremental workspace
 		/// </summary>
 		public bool BIncremental { get; set; }
+		
+		/// <summary>
+		/// Method to use when syncing/materializing data from Perforce
+		/// </summary>
+		public string? Method { get; set; }
 
 		/// <summary>
 		/// Constructor
@@ -301,6 +306,7 @@ namespace Horde.Build.Agents
 			Stream = workspace.Stream;
 			View = workspace.View;
 			BIncremental = workspace.Incremental;
+			Method = workspace.Method;
 		}
 	}
 
