Horde: Fix bug sending nulls in AgentWorkspace protobuf message

diff --git a/Engine/Source/Programs/Horde/Horde.Agent/Execution/PerforceExecutor.cs b/Engine/Source/Programs/Horde/Horde.Agent/Execution/PerforceExecutor.cs
--- a/Engine/Source/Programs/Horde/Horde.Agent/Execution/PerforceExecutor.cs
+++ b/Engine/Source/Programs/Horde/Horde.Agent/Execution/PerforceExecutor.cs
@@ -412,7 +412,7 @@ namespace Horde.Agent.Execution
 			const string ManagedWorkspaceValue = "managedWorkspace";
 			const string UseHaveTableKey = "useHaveTable";
 			
-			if (method == null)
+			if (String.IsNullOrEmpty(method))
 			{
 				return true;
 			}
diff --git a/Engine/Source/Programs/Horde/Horde.Build/Agents/IAgent.cs b/Engine/Source/Programs/Horde/Horde.Build/Agents/IAgent.cs
--- a/Engine/Source/Programs/Horde/Horde.Build/Agents/IAgent.cs
+++ b/Engine/Source/Programs/Horde/Horde.Build/Agents/IAgent.cs
@@ -178,7 +178,7 @@ namespace Horde.Build.Agents
 				Identifier = Identifier,
 				Stream = Stream,
 				Incremental = Incremental,
-				Method = Method
+				Method = Method ?? String.Empty
 			};
 			
 			if (View != null)
@@ -768,7 +768,7 @@ namespace Horde.Build.Agents
 				Identifier = workspace.Identifier,
 				Stream = workspace.Stream,
 				Incremental = workspace.Incremental,
-				Method = workspace.Method
+				Method = workspace.Method ?? String.Empty
 			};
 
 			if (workspace.View != null)
