Horde: Fix null reference exception constructing file writer during p4 replication.

diff --git a/Engine/Source/Programs/Horde/Horde.Build/Perforce/PerforceReplicator.cs b/Engine/Source/Programs/Horde/Horde.Build/Perforce/PerforceReplicator.cs
--- a/Engine/Source/Programs/Horde/Horde.Build/Perforce/PerforceReplicator.cs
+++ b/Engine/Source/Programs/Horde/Horde.Build/Perforce/PerforceReplicator.cs
@@ -96,10 +96,15 @@ namespace Horde.Build.Perforce
 			class Handle
 			{
 				public Utf8String _path;
-				public FileNodeWriter _fileWriter = null!;
+				public readonly FileNodeWriter _fileWriter;
 				public long _size;
 				public long _sizeWritten;
 				public readonly IncrementalHash _hash = IncrementalHash.CreateHash(HashAlgorithmName.MD5);
+
+				public Handle(TreeWriter writer, ChunkingOptions options)
+				{
+					_fileWriter = new FileNodeWriter(writer, options);
+				}
 			}
 
 			readonly TreeWriter _writer;
@@ -127,7 +132,11 @@ namespace Horde.Build.Perforce
 
 			public void Open(int fd, Utf8String path, long size)
 			{
-				Handle handle = (_freeHandles.Count > 0)? _freeHandles.Pop() : new Handle();
+				Handle? handle;
+				if (!_freeHandles.TryPop(out handle))
+				{
+					handle = new Handle(_writer, _options);
+				}
 
 				handle._fileWriter.Reset();
 				handle._path = path;
