Horde: Include file entry flags for executable, read-only, and text matching replicated Perforce data.

diff --git a/Engine/Source/Programs/Horde/Horde.Build/Perforce/PerforceReplicator.cs b/Engine/Source/Programs/Horde/Horde.Build/Perforce/PerforceReplicator.cs
--- a/Engine/Source/Programs/Horde/Horde.Build/Perforce/PerforceReplicator.cs
+++ b/Engine/Source/Programs/Horde/Horde.Build/Perforce/PerforceReplicator.cs
@@ -1,6 +1,7 @@
 // Copyright Epic Games, Inc. All Rights Reserved.
 
 using System;
+using System.Buffers.Binary;
 using System.Collections.Generic;
 using System.Diagnostics;
 using System.Linq;
@@ -11,6 +12,7 @@ using EpicGames.Core;
 using EpicGames.Horde.Storage;
 using EpicGames.Horde.Storage.Nodes;
 using EpicGames.Perforce;
+using EpicGames.Perforce.Managed;
 using Horde.Build.Storage;
 using Horde.Build.Streams;
 using Horde.Build.Utilities;
@@ -77,6 +79,22 @@ namespace Horde.Build.Perforce
 			}
 		}
 
+		// Partial mirror of FileSysType from P4 API (filesys.h)
+		const uint FST_TEXT = 0x0001;
+		const uint FST_BINARY = 0x0002;
+		const uint FST_UNICODE = 0x000c;
+		const uint FST_UTF16 = 0x000e;
+		const uint FST_UTF8 = 0x000f;
+		const uint FST_MASK = 0x000f;
+
+		// Mirrors FilePerm from P4 API (filesys.h)
+		const uint FPM_RO = 0;     // leave file read-only
+		const uint FPM_RW = 1;     // leave file read-write
+		const uint FPM_ROO = 2;    // leave file read-only (owner)
+		const uint FPM_RXO = 3;    // set file read-execute (owner) NO W
+		const uint FPM_RWO = 4;    // set file read-write (owner) NO X
+		const uint FPM_RWXO = 5;   // set file read-write-execute (owner)
+
 		[DebuggerDisplay("{_path}")]
 		class DirectoryToSync
 		{
@@ -91,11 +109,14 @@ namespace Horde.Build.Perforce
 			}
 		}
 
+		record class FileInfo(Utf8String Path, FileEntryFlags Flags, long Length, byte[] Md5, NodeHandle Node);
+
 		class FileWriter : IDisposable
 		{
 			class Handle
 			{
 				public Utf8String _path;
+				public FileEntryFlags _flags;
 				public readonly FileNodeWriter _fileWriter;
 				public long _size;
 				public long _sizeWritten;
@@ -130,7 +151,7 @@ namespace Horde.Build.Perforce
 				}
 			}
 
-			public void Open(int fd, Utf8String path, long size)
+			public void Open(int fd, Utf8String path, long size, FileEntryFlags flags)
 			{
 				Handle? handle;
 				if (!_freeHandles.TryPop(out handle))
@@ -141,6 +162,7 @@ namespace Horde.Build.Perforce
 				handle._fileWriter.Reset();
 				handle._path = path;
 				handle._size = size;
+				handle._flags = flags;
 
 				_openHandles.Add(fd, handle);
 			}
@@ -154,7 +176,7 @@ namespace Horde.Build.Perforce
 				handle._sizeWritten += data.Length;
 			}
 
-			public async Task<(Utf8String, NodeHandle, byte[], long)> CloseAsync(int fd, CancellationToken cancellationToken)
+			public async Task<FileInfo> CloseAsync(int fd, CancellationToken cancellationToken)
 			{
 				Handle handle = _openHandles[fd];
 				if (handle._sizeWritten != handle._size)
@@ -162,15 +184,14 @@ namespace Horde.Build.Perforce
 					throw new ReplicationException($"Invalid size for replicated file '{handle._path}'. Expected {handle._size}, got {handle._sizeWritten}.");
 				}
 
-				Utf8String path = handle._path;
 				NodeHandle node = await handle._fileWriter.FlushAsync(cancellationToken);
 				byte[] hash = handle._hash.GetHashAndReset();
-				long length = handle._fileWriter.Length;
+				FileInfo info = new FileInfo(handle._path, handle._flags, handle._size, hash, node);
 
 				_openHandles.Remove(fd);
 				_freeHandles.Push(handle);
 
-				return (path, node, hash, length);
+				return info;
 			}
 		}
 
@@ -395,7 +416,27 @@ namespace Horde.Build.Perforce
 					{
 						if (io.Command == PerforceIoCommand.Open)
 						{
-							Utf8String file = GetClientRelativePath(io.Payload, clientInfo.Client.Root);
+							UnpackOpenPayload(io.Payload, out Utf8String path, out int type, out int mode, out int perms);
+
+							FileEntryFlags flags = 0;
+							if ((type & FST_MASK) != FST_BINARY)
+							{
+								flags |= FileEntryFlags.Text;
+							}
+							if ((type & FST_UTF16) != 0)
+							{
+								flags |= FileEntryFlags.Utf16;
+							}
+							if (perms == FPM_RO || perms == FPM_ROO || perms == FPM_RXO)
+							{
+								flags |= FileEntryFlags.ReadOnly;
+							}
+							if (perms == FPM_RWXO || perms == FPM_RXO)
+							{
+								flags |= FileEntryFlags.Executable;
+							}
+
+							Utf8String file = GetClientRelativePath(path, clientInfo.Client.Root);
 							int offset = GetFileOffset(file);
 
 							long fileSize = 0;
@@ -408,7 +449,7 @@ namespace Horde.Build.Perforce
 								throw new ReplicationException($"Unable to find file entry for {file}");
 							}
 
-							fileWriter.Open(io.File, file, fileSize);
+							fileWriter.Open(io.File, file, fileSize, flags);
 						}
 						else if (io.Command == PerforceIoCommand.Write)
 						{
@@ -416,15 +457,15 @@ namespace Horde.Build.Perforce
 						}
 						else if (io.Command == PerforceIoCommand.Close)
 						{
-							(Utf8String path, NodeHandle handle, byte[] hash, long length) = await fileWriter.CloseAsync(io.File, cancellationToken);
-
-							FileEntry entry = await root.AddFileByPathAsync(reader, path, FileEntryFlags.None, length, handle, cancellationToken);
-							entry.CustomData = hash;
+							FileInfo info = await fileWriter.CloseAsync(io.File, cancellationToken);
+							FileEntry entry = await root.AddFileByPathAsync(reader, info.Path, info.Flags, info.Length, info.Node, cancellationToken);
+							entry.CustomData = info.Md5;
 							await writer.WriteAsync(entry, cancellationToken);
 						}
 						else if (io.Command == PerforceIoCommand.Unlink)
 						{
-							Utf8String file = GetClientRelativePath(io.Payload, clientInfo.Client.Root);
+							UnpackUnlinkPayload(io.Payload, out Utf8String path);
+							Utf8String file = GetClientRelativePath(path, clientInfo.Client.Root);
 							await root.DeleteFileByPathAsync(reader, file, cancellationToken);
 						}
 						else
@@ -561,21 +602,37 @@ namespace Horde.Build.Perforce
 			}
 		}
 
-		static Utf8String GetClientRelativePath(ReadOnlyMemory<byte> data, Utf8String clientRoot)
+		static void UnpackOpenPayload(ReadOnlyMemory<byte> data, out Utf8String path, out int flags, out int type, out int perms)
+		{
+			int length = data.Span.IndexOf((byte)0);
+			if (length == -1)
+			{
+				throw new InvalidStreamException("Invalid data returned by Perforce server; open path does not contain null terminator.");
+			}
+
+			path = new Utf8String(data.Slice(0, length)).Clone();
+			flags = BinaryPrimitives.ReadInt32LittleEndian(data.Span.Slice(length + 1, sizeof(int)));
+			type = BinaryPrimitives.ReadInt32LittleEndian(data.Span.Slice(length + 5, sizeof(int)));
+			perms = BinaryPrimitives.ReadInt32LittleEndian(data.Span.Slice(length + 9, sizeof(int)));
+		}
+
+		static void UnpackUnlinkPayload(ReadOnlyMemory<byte> data, out Utf8String path)
 		{
 			int length = data.Span.IndexOf((byte)0);
 			if (length != -1)
 			{
 				data = data.Slice(0, length);
 			}
+			path = new Utf8String(data).Clone();
+		}
 
-			Utf8String path = new Utf8String(data);
+		static Utf8String GetClientRelativePath(Utf8String path, Utf8String clientRoot)
+		{
 			if (!path.StartsWith(clientRoot, Utf8StringComparer.Ordinal))
 			{
 				throw new ArgumentException($"Unable to make path {path} relative to client root {clientRoot}");
 			}
-
-			return path.Substring(clientRoot.Length).Clone();
+			return path.Substring(clientRoot.Length);
 		}
 
 		async Task<ReplicationClient?> FindReplicationClientAsync(IStream stream)
