Horde: Add automation home view (WIP) and upgrade TS target to ES2015
#rnx
#skipci

diff --git a/Engine/Source/Programs/Horde/HordeDashboard/src/App.tsx b/Engine/Source/Programs/Horde/HordeDashboard/src/App.tsx
--- a/Engine/Source/Programs/Horde/HordeDashboard/src/App.tsx
+++ b/Engine/Source/Programs/Horde/HordeDashboard/src/App.tsx
@@ -11,6 +11,7 @@ import dashboard from './backend/Dashboard';
 import { AdminToken } from './components/AdminToken';
 import { AgentView } from './components/AgentView';
 import { AuditLogView } from './components/AuditLog';
+import { AutomationView } from './components/AutomationView';
 import { DashboardView } from './components/DashboardView';
 import { DebugView } from './components/DebugView';
 import { DeviceView } from './components/DeviceView';
@@ -134,6 +135,7 @@ const Main: React.FC = () => {
          <Route path="/server/settings" component={ServerSettingsView} />
          <Route path="/audit/agent/:agentId" component={AuditLogView} />
          <Route path="/audit/issue/:issueId" component={AuditLogView} />
+         <Route path="/automation" component={AutomationView} />
          <Route path={["/debug/lease/:leaseId"]} component={DebugView} />
          {hordePlugins.routes.map((route, index) => {
             return <Route key={`key_plugin_route_${index}`} path={route.path} component={route.component} />;
diff --git a/Engine/Source/Programs/Horde/HordeDashboard/src/backend/Api.ts b/Engine/Source/Programs/Horde/HordeDashboard/src/backend/Api.ts
--- a/Engine/Source/Programs/Horde/HordeDashboard/src/backend/Api.ts
+++ b/Engine/Source/Programs/Horde/HordeDashboard/src/backend/Api.ts
@@ -500,10 +500,10 @@ export type UpdatePoolRequest = {
 
 	/// Frequency to run conforms, in hours, set 0 to disable
 	conformInterval?: number;
-		
+
 	/** Pool sizing strategy */
 	sizeStrategy?: PoolSizeStrategy;
-		
+
 	/** The minimum nunmber of agents to retain in this pool */
 	minAgents?: number;
 
@@ -741,14 +741,13 @@ export type Condition = {
 }
 
 /// Available pool sizing strategies
-export enum PoolSizeStrategy
-{
+export enum PoolSizeStrategy {
 	/// Strategy based on lease utilization
 	LeaseUtilization = "LeaseUtilization",
-	
+
 	/// Strategy based on size of job build queue
 	JobQueue = "JobQueue",
-	
+
 	/// No-op strategy used as fallback/default behavior
 	NoOp = "NoOp"
 }
@@ -788,7 +787,7 @@ export type GetPoolResponse = {
 
 	/// Frequency to run conforms, in hours, set 0 to disable
 	conformInterval?: number;
-		
+
 	/// Cooldown time between scale-out events in seconds
 	scaleOutCooldown?: number;
 
@@ -1896,50 +1895,50 @@ export type UpdateStepStateRequest = {
 
 /// Response describing a template
 export type GetTemplateResponseBase = {
-		
+
 	/// Name of the template
 	name: string;
 
 	/// Default priority for this job
 	priority?: Priority;
-	
+
 	/// Whether to allow preflights of this template
 	allowPreflights: boolean;
-	
+
 	/// Whether to always update issues on jobs using this template
 	updateIssues: boolean;
-	
+
 	/// The initial agent type to parse the BuildGraph script on
 	initialAgentType?: string;
-	
+
 	/// Path to a file within the stream to submit to generate a new changelist for jobs
 	submitNewChange?: string;
-	
+
 	/// Parameters for the job.
 	arguments: string[];
-	
+
 	/// List of parameters for this template
 	parameters: ParameterData[];
 }
 
 /// Query selecting the base changelist to use
 export type ChangeQueryConfig = {
-	
+
 	/// Name of this query, for display on the dashboard.
 	name?: string;
-	
+
 	/// Condition to evaluate before deciding to use this query. May query tags in a preflight.	
 	condition?: any;
-		
+
 	/// The template id to query	
 	templateId?: string;
-	
+
 	/// The target to query	
 	target?: string;
-	
+
 	/// Whether to match a job that produced warnings	
 	outcomes?: JobStepOutcome[];
-	
+
 	/// Finds the last commit with this tag	
 	commitTag?: any;
 }
@@ -2512,7 +2511,7 @@ export type GetGraphResponse = {
 
 /**The timing info for a job*/
 export type GetJobTimingResponse = {
-	
+
 	/** The job response */
 	jobResponse: JobData;
 
@@ -4191,18 +4190,18 @@ export type DeviceTelemetryQuery = {
 };
 
 export type GetDevicePoolReservationTelemetryResponse = {
-	
+
 	/** Device id for reservation */
 	deviceId: string;
 
 	/** Job id associated with reservation */
-	jobId?: string; 
+	jobId?: string;
 
 	/** The step id of reservation */
-	stepId?: string; 
+	stepId?: string;
 
 	/** The name of the job holding reservation */
-	jobName?: string; 
+	jobName?: string;
 
 	/** The name of the step holding reservation */
 	stepName?: string;
@@ -4248,19 +4247,19 @@ export type GetTelemetryInfoResponse = {
 
 	/// The UTC time the telemetry data was created	
 	createTimeUtc: Date | string;
-	
+
 	/// The stream id which utilized device	
 	streamId?: string;
-	
+
 	/// The job id which utilized device	
 	jobId?: string;
-	
+
 	/// The job's step id	
 	stepId?: string;
 
 	/// The job name which utilized device	
 	jobName?: string;
-	
+
 	/// The job's step name
 	stepName?: string;
 
@@ -4269,7 +4268,7 @@ export type GetTelemetryInfoResponse = {
 
 	/// If this telemetry has a reservation, the finish time of the reservation
 	reservationFinishUtc?: Date | string;
-	
+
 	/// If this telemetry marks a detected device issue, the time of the issue	
 	problemTimeUtc?: Date | string;
 }
@@ -4284,4 +4283,163 @@ export type GetDeviceTelemetryResponse = {
 	/// Individual telemetry data points	
 	telemetry: GetTelemetryInfoResponse[];
 
-}
\ No newline at end of file
+}
+
+// Test Data ----------------------------
+
+/// A test emvironment running in a stream	
+export type GetTestMetaResponse = {
+
+	/// Meta unique id for environment
+	id: string;
+
+	/// The platforms in the environment		
+	platforms: string[];
+
+	/// The build configurations being tested		
+	configurations: string[];
+
+	/// The build targets being tested		
+	buildTargets: string[];
+
+	/// The test project name		
+	projectName: string;
+
+	/// The rendering hardware interface being used with the test		
+	rhi: string;
+}
+
+
+/// A test that runs in a stream
+export type GetTestResponse = {
+
+	/// The id of the test
+	id: string;
+
+	/// The name of the test 		
+	name: string;
+
+	/// The display name of the test, if any 		
+	displayName?: string;
+
+	/// The suite the test belongs to, if any
+	suiteName?: string;
+
+	/// The environments the test runs in
+	metadata: string[];
+}
+
+/// A test that runs in a stream
+export type GetTestsRequest = {
+	/// The id of the test
+	testIds: string[];
+}
+
+
+/// A test suite that runs in a stream, contain subtests	
+export type GetTestSuiteResponse = {
+
+	/// The id of the test suite
+	id: string;
+
+	/// The name of the test suite		
+	name: string;
+
+	/// The environments the suite runs in
+	metadata: string[];
+}
+
+/// Describes tests running in a stream
+export type GetTestStreamResponse = {
+
+	/// The stream id	
+	streamId: string;
+
+	/// Individual tests which run in the stream	
+	tests: GetTestResponse[];
+
+	/// Test suites that run in the stream	
+	testSuites: GetTestSuiteResponse[];
+
+	/// Test suites that run in the stream	
+	testMetadata: GetTestMetaResponse[];
+}
+
+/// Test outcome
+export enum TestOutcome {
+	/// The test was successful
+	Success = "Success",
+	/// The test failed
+	Failure = "Failure",
+	/// The test was skipped
+	Skipped = "Skipped",
+	/// The test had an unspecified result
+	Unspecified = "Unspecified"
+}
+
+
+/// Suite test data
+export type GetSuiteTestDataResponse = {
+
+	/// The test id	
+	testId: string;
+
+	/// The ourcome of the suite test	
+	outcome: TestOutcome;
+
+	/// How long the suite test ran	 (TimeSpan)
+	duration: string;
+
+	/// Test UID for looking up in test details
+	uid: string;
+}
+
+/// Test details
+export type GetTestDataDetailsResponse = {
+	
+	/// The corresponding test ref	
+	id: string;
+	
+	/// The test documents for this ref	
+	testDataIds: string[];
+	
+	/// Suite test data		
+	suiteTests?: GetSuiteTestDataResponse[];
+	
+}
+
+/// Testt data ref
+export type GetTestDataRefResponse = {
+	/// ref id
+	id: string;
+
+	/// The associated stream
+	streamId: string;
+
+	/// How long the test ran (TimeSpan)
+	duration: string;
+
+	/// The build changelist upon which the test ran, may not correspond to the job changelist
+	buildChangeList: number;
+
+	/// The environment the test ran on 
+	metaId: string;
+
+	/// The test id in stream
+	testId?: string;
+
+	/// The outcome of the test
+	outcome?: TestOutcome;
+
+	/// The if of the stream test suite
+	suiteId?: string;
+
+	/// The number of suite tests skipped
+	suiteSkipCount?: number;
+
+	/// The number of suite tests with warnings
+	suiteWaringCount?: number;
+
+	/// The number of suite tests swith errors
+	suiteErrorCount?: number;
+}
diff --git a/Engine/Source/Programs/Horde/HordeDashboard/src/backend/Backend.ts b/Engine/Source/Programs/Horde/HordeDashboard/src/backend/Backend.ts
--- a/Engine/Source/Programs/Horde/HordeDashboard/src/backend/Backend.ts
+++ b/Engine/Source/Programs/Horde/HordeDashboard/src/backend/Backend.ts
@@ -3,7 +3,7 @@
 import { configure } from "mobx";
 import { isNumber } from 'util';
 import templateCache from '../backend/TemplateCache';
-import { AgentData, AgentQuery, ArtifactData, AuditLogEntry, AuditLogQuery, BatchUpdatePoolRequest, ChangeSummaryData, CreateDeviceRequest, CreateDeviceResponse, CreateExternalIssueRequest, CreateExternalIssueResponse, CreateJobRequest, CreateJobResponse, CreateNoticeRequest, CreatePoolRequest, CreateSoftwareResponse, CreateSubscriptionRequest, CreateSubscriptionResponse, DashboardPreference, DevicePoolTelemetryQuery, DeviceTelemetryQuery, EventData, FindIssueResponse, FindJobTimingsResponse, GetAgentSoftwareChannelResponse, GetArtifactZipRequest, GetDashboardConfigResponse, GetDevicePlatformResponse, GetDevicePoolResponse, GetDevicePoolTelemetryResponse, GetDeviceReservationResponse, GetDeviceResponse, GetDeviceTelemetryResponse, GetExternalIssueProjectResponse, GetExternalIssueResponse, GetGraphResponse, GetIssueStreamResponse, GetJobsTabResponse, GetJobStepRefResponse, GetJobStepTraceResponse, GetJobTimingResponse, GetLogEventResponse, GetNoticeResponse, GetNotificationResponse, GetPerforceServerStatusResponse, GetPoolResponse, GetServerInfoResponse, GetServerSettingsResponse, GetSoftwareResponse, GetSubscriptionResponse, GetUserResponse, GetUtilizationTelemetryResponse, GlobalConfig, IssueData, IssueQuery, IssueQueryV2, JobData, JobQuery, JobsTabColumnType, JobStepOutcome, JobStreamQuery, JobTimingsQuery, LeaseData, LogData, LogLineData, PoolData, ProjectData, ScheduleData, ScheduleQuery, SearchLogFileResponse, ServerUpdateResponse, SessionData, StreamData, TabType, TestData, UpdateAgentRequest, UpdateDeviceRequest, UpdateGlobalConfigRequest, UpdateIssueRequest, UpdateJobRequest, UpdateLeaseRequest, UpdateNoticeRequest, UpdateNotificationsRequest, UpdatePoolRequest, UpdateServerSettingsRequest, UpdateStepRequest, UpdateStepResponse, UpdateTemplateRefRequest, UpdateUserRequest, UsersQuery } from './Api';
+import { AgentData, AgentQuery, ArtifactData, AuditLogEntry, AuditLogQuery, BatchUpdatePoolRequest, ChangeSummaryData, CreateDeviceRequest, CreateDeviceResponse, CreateExternalIssueRequest, CreateExternalIssueResponse, CreateJobRequest, CreateJobResponse, CreateNoticeRequest, CreatePoolRequest, CreateSoftwareResponse, CreateSubscriptionRequest, CreateSubscriptionResponse, DashboardPreference, DevicePoolTelemetryQuery, DeviceTelemetryQuery, EventData, FindIssueResponse, FindJobTimingsResponse, GetAgentSoftwareChannelResponse, GetArtifactZipRequest, GetDashboardConfigResponse, GetDevicePlatformResponse, GetDevicePoolResponse, GetDevicePoolTelemetryResponse, GetDeviceReservationResponse, GetDeviceResponse, GetDeviceTelemetryResponse, GetExternalIssueProjectResponse, GetExternalIssueResponse, GetGraphResponse, GetIssueStreamResponse, GetJobsTabResponse, GetJobStepRefResponse, GetJobStepTraceResponse, GetJobTimingResponse, GetLogEventResponse, GetNoticeResponse, GetNotificationResponse, GetPerforceServerStatusResponse, GetPoolResponse, GetServerInfoResponse, GetServerSettingsResponse, GetSoftwareResponse, GetSubscriptionResponse, GetTestDataDetailsResponse, GetTestDataRefResponse, GetTestMetaResponse, GetTestResponse, GetTestsRequest, GetTestStreamResponse, GetUserResponse, GetUtilizationTelemetryResponse, GlobalConfig, IssueData, IssueQuery, IssueQueryV2, JobData, JobQuery, JobsTabColumnType, JobStepOutcome, JobStreamQuery, JobTimingsQuery, LeaseData, LogData, LogLineData, PoolData, ProjectData, ScheduleData, ScheduleQuery, SearchLogFileResponse, ServerUpdateResponse, SessionData, StreamData, TabType, TestData, UpdateAgentRequest, UpdateDeviceRequest, UpdateGlobalConfigRequest, UpdateIssueRequest, UpdateJobRequest, UpdateLeaseRequest, UpdateNoticeRequest, UpdateNotificationsRequest, UpdatePoolRequest, UpdateServerSettingsRequest, UpdateStepRequest, UpdateStepResponse, UpdateTemplateRefRequest, UpdateUserRequest, UsersQuery } from './Api';
 import dashboard from './Dashboard';
 import { ChallengeStatus, Fetch } from './Fetch';
 import graphCache, { GraphQuery } from './GraphCache';
@@ -626,24 +626,6 @@ export class Backend {
 
     }
 
-    getLog(logId: string, offset?: number, count?: number): Promise<string> {
-
-        return new Promise<string>((resolve, reject) => {
-
-            this.backend.get(`/api/v1/logs/${logId}/data`, {
-                params: {
-                    offset: offset,
-                    count: count
-                }
-            }).then((value) => {
-                resolve(value.data as string);
-            }).catch(reason => {
-                reject(reason);
-            });
-        });
-
-    }
-
     createJob(request: CreateJobRequest): Promise<CreateJobResponse> {
 
         return new Promise<CreateJobResponse>((resolve, reject) => {
@@ -1107,6 +1089,65 @@ export class Backend {
         });
     }
 
+    getTestMetadata(auttomationProjects?: string[], platforms?: string[], targets?: string[], configurations?: string[]): Promise<GetTestMetaResponse[]> {
+        return new Promise<GetTestMetaResponse[]>((resolve, reject) => {
+            this.backend.get(`/api/v2/testdata/metadata`, {
+                params: { project: auttomationProjects, platform: platforms, target: targets, configuration: configurations }
+            }).then((value) => {
+                resolve(value.data as GetTestMetaResponse[]);
+            }).catch((reason) => {
+                reject(reason);
+            });
+        });
+    }
+
+    getTestRefs(streamIds: string[], metaIds: string[], minCreateTime?: string, maxCreateTime?: string, minChange?: number, maxChange?: number, testIds?: string[], suiteIds?: string[]): Promise<GetTestDataRefResponse[]> {
+        return new Promise<GetTestDataRefResponse[]>((resolve, reject) => {
+            this.backend.get(`/api/v2/testdata/refs`, {
+                params: { id: streamIds, mid: metaIds, tid: testIds, sid: suiteIds, minCreateTime: minCreateTime, maxCreateTime: maxCreateTime, minChange: minChange, maxChange: maxChange }
+            }).then((value) => {
+                resolve(value.data as GetTestDataRefResponse[]);
+            }).catch((reason) => {
+                reject(reason);
+            });
+        });
+    }
+
+    getTestDetails(refIds: string[]): Promise<GetTestDataDetailsResponse[]> {
+        return new Promise<GetTestDataDetailsResponse[]>((resolve, reject) => {
+            this.backend.get(`/api/v2/testdata/details`, {
+                params: { id: refIds }
+            }).then((value) => {
+                resolve(value.data as GetTestDataDetailsResponse[]);
+            }).catch((reason) => {
+                reject(reason);
+            });
+        });
+    }
+
+    getTests(testIds: string[]): Promise<GetTestResponse[]> {
+        const request: GetTestsRequest = { testIds: testIds };        
+        return new Promise<GetTestResponse[]>((resolve, reject) => {
+            this.backend.post(`/api/v2/testdata/tests`, request).then((value) => {
+                resolve(value.data as GetTestResponse[]);
+            }).catch((reason) => {
+                reject(reason);
+            });
+        });
+    }
+
+    getTestStreams(streamIds: string[]): Promise<GetTestStreamResponse[]> {
+        return new Promise<GetTestStreamResponse[]>((resolve, reject) => {
+            this.backend.get(`/api/v2/testdata/streams`, {
+                params: { id: streamIds }
+            }).then((value) => {
+                resolve(value.data as GetTestStreamResponse[]);
+            }).catch((reason) => {
+                reject(reason);
+            });
+        });
+    }
+
     getJobTestData(jobId: string, stepId?: string): Promise<TestData[]> {
         return new Promise<TestData[]>((resolve, reject) => {
             const params: any = {
diff --git a/Engine/Source/Programs/Horde/HordeDashboard/src/backend/Dashboard.ts b/Engine/Source/Programs/Horde/HordeDashboard/src/backend/Dashboard.ts
--- a/Engine/Source/Programs/Horde/HordeDashboard/src/backend/Dashboard.ts
+++ b/Engine/Source/Programs/Horde/HordeDashboard/src/backend/Dashboard.ts
@@ -333,8 +333,8 @@ export class Dashboard {
             [StatusColor.Running, dark ? "#146579" : theme.palette.blueLight],
             [StatusColor.Waiting, dark ? "#474542" : "#A19F9D"],
             [StatusColor.Ready, dark ? "#474542" : "#A19F9D"],
-            [StatusColor.Skipped, dark ? "#63625c" : "#F3F2F1"],
-            [StatusColor.Unspecified, "#637087"]
+            [StatusColor.Skipped, dark ? "#63625c" : "#C3C2C1"],
+            [StatusColor.Unspecified, "#646464"]
         ]);
 
         colors.set(StatusColor.Aborted, colors.get(StatusColor.Failure)!);
diff --git a/Engine/Source/Programs/Horde/HordeDashboard/src/base/utilities/timeUtils.ts b/Engine/Source/Programs/Horde/HordeDashboard/src/base/utilities/timeUtils.ts
--- a/Engine/Source/Programs/Horde/HordeDashboard/src/base/utilities/timeUtils.ts
+++ b/Engine/Source/Programs/Horde/HordeDashboard/src/base/utilities/timeUtils.ts
@@ -379,7 +379,7 @@ export const getHumanTime = (timeIn: Date | string | undefined): string => {
 };
 
 
-export const getShortNiceTime = (timeIn: Date | string | undefined, relative: boolean = false): string => {
+export const getShortNiceTime = (timeIn: Date | string | undefined, relative: boolean = false, includeHour:boolean = false): string => {
 
     if (!timeIn) {
         return "";
@@ -399,6 +399,11 @@ export const getShortNiceTime = (timeIn: Date | string | undefined, relative: bo
         timeStr = time.format('MMM Do');
     }
 
+    if (includeHour) {
+        const format = dashboard.display24HourClock ? "HH:mm:ss z" : "LT z";
+        timeStr += ` at ${time.format(format)}`;
+    }
+    
     return timeStr;
 
 };
diff --git a/Engine/Source/Programs/Horde/HordeDashboard/src/components/AutomationSuiteDetails.tsx b/Engine/Source/Programs/Horde/HordeDashboard/src/components/AutomationSuiteDetails.tsx
new file mode 100644
--- /dev/null
+++ b/Engine/Source/Programs/Horde/HordeDashboard/src/components/AutomationSuiteDetails.tsx
@@ -0,0 +1,1021 @@
+import { DetailsList, DetailsListLayoutMode, Dropdown, IColumn, IconButton, IDropdownOption, mergeStyleSets, Modal, SelectionMode, Spinner, SpinnerSize, Stack, Text, TextField } from "@fluentui/react";
+import * as d3 from "d3";
+import { action, observable } from 'mobx';
+import { observer } from 'mobx-react-lite';
+import React, { useState } from "react";
+import backend from "../backend";
+import { GetSuiteTestDataResponse, GetTestDataDetailsResponse, GetTestDataRefResponse, GetTestMetaResponse, GetTestResponse, GetTestSuiteResponse, TestOutcome } from "../backend/Api";
+import dashboard, { StatusColor } from "../backend/Dashboard";
+import { projectStore } from "../backend/ProjectStore";
+import { useWindowSize } from "../base/utilities/hooks";
+import { getHumanTime, getShortNiceTime } from "../base/utilities/timeUtils";
+import { hordeClasses, modeColors } from "../styles/Styles";
+import { AutomationSuiteTest } from "./AutomationSuiteTest";
+
+type TestId = string;
+type SelectionType = d3.Selection<SVGGElement, unknown, null, undefined>;
+type DivSelectionType = d3.Selection<HTMLDivElement, unknown, null, undefined>;
+type TestState = "Success" | "Failed" | "Consecutive Failures" | "Skipped" | "Unspecified";
+type FilterState = {
+   testState?: TestState;
+   text?: string;
+}
+
+const suiteGraphWidth = 1180;
+
+type TestStatus = {
+   testId: string,
+   success?: number;
+   error?: number;
+   skip?: number;
+   unspecified?: number;
+}
+
+type TestSelection = {
+   ref: GetTestDataRefResponse;
+   testId: string;
+   details: GetTestDataDetailsResponse;
+}
+
+type TestMetrics = {
+   total: number;
+   error: number;
+   success: number;
+   skipped: number;
+   unspecified: number;
+}
+
+class SuiteHandler {
+
+   constructor(suite: GetTestSuiteResponse, suiteRefs: GetTestDataRefResponse[], metaData: GetTestMetaResponse) {
+
+
+      this.suite = suite;
+
+      const dedupe = new Set<number>();
+
+      this.suiteRefs = suiteRefs.filter(r => {
+         if (!dedupe.has(r.buildChangeList)) {
+            dedupe.add(r.buildChangeList);
+            return true;
+         }
+         return false;
+      });
+
+      if (this.suiteRefs.length !== suiteRefs.length) {
+         console.warn("Warning: Test suite details had duplicate ref changelists");
+      }
+
+      this.metaData = metaData;
+
+      this.metaName = `${metaData.platforms.join(" - ")} / ${metaData.configurations.join(" - ")} / ${metaData.buildTargets.join(" - ")} / ${metaData.rhi === "default" ? "Default" : metaData.rhi?.toUpperCase()}`;
+
+      this.load();
+   }
+
+   setSelection(selection?: TestSelection) {
+
+      if ((this.selection?.ref.id === selection?.ref.id) && (this.selection?.testId === selection?.testId)) {
+         return;
+      }
+
+      this.selection = selection;
+      this.setSelectionUpdated();
+   }
+
+   @action
+   setUpdated() {
+      this.updated++;
+   }
+
+   @action
+   setSelectionUpdated() {
+      this.selectionUpdated++;
+   }
+
+   async load() {
+
+      console.log(`Loading ${this.suiteRefs.length} suite ref details for metaData: ${JSON.stringify(this.metaData)}`);
+
+      this.changeDates = new Map();
+      this.suiteRefs.forEach(r => {
+
+         this.refs.set(r.id, r);
+
+         if (!this.refDates.get(r.id)) {
+            const timestamp = r.id.substring(0, 8)
+            this.refDates.set(r.id, new Date(parseInt(timestamp, 16) * 1000));
+         }
+
+         const date = this.changeDates.get(r.buildChangeList);
+         const rdate = this.refDates.get(r.id)!;
+         if (!date || date.getTime() > rdate.getTime()) {
+            this.changeDates.set(r.buildChangeList, rdate);
+         }
+      });
+
+
+      let details = await backend.getTestDetails(this.suiteRefs.map(r => r.id));
+      details = details.filter(d => !!d.suiteTests);
+      details.forEach(d => this.testDetails.set(d.id, d));
+
+      const uniqueTests = new Set<string>(details.map(d => d.suiteTests!.map(t => t.testId)).flat())
+      const request = Array.from(uniqueTests).filter(tid => !SuiteHandler.tests.has(tid));
+
+      if (request.length) {
+         const tests = await backend.getTests(Array.from(uniqueTests).filter(tid => !SuiteHandler.tests.has(tid)));
+
+         const replace = ["Project.Functional Tests.Tests.", "Project.Maps.AllInPIE."];
+
+         function fixName(name: string) {
+
+            for (let r of replace) {
+               if (name.startsWith(r)) {
+                  name = name.substring(r.length)
+               }
+            }
+            return name;
+         }
+
+         tests.forEach(t => {
+            SuiteHandler.tests.set(t.id, t);
+            const fixed = fixName(t.displayName ?? t.name);
+            SuiteHandler.fixedTestNames.set(t.id, fixed);
+            SuiteHandler.lowerCaseTestNames.set(t.id, fixed.toLowerCase());
+         });
+      }
+
+      const refs = this.suiteRefs.sort((a, b) => a.buildChangeList - b.buildChangeList).reverse();
+
+      for (const ref of refs) {
+         const details = this.testDetails.get(ref.id);
+         const suiteTests = details?.suiteTests;
+         if (!suiteTests) {
+            return;
+         }
+
+         for (const test of suiteTests) {
+
+            let testData = this.testData.get(test.testId);
+
+            if (!testData) {
+               testData = [];
+               this.testData.set(test.testId, testData);
+            }
+
+            testData.push({ ...test, refId: ref.id });
+
+            let status = this.testStatus.get(test.testId);
+            if (!status) {
+               status = { testId: test.testId };
+               this.testStatus.set(test.testId, status);
+            }
+         }
+      }
+
+      for (const [, testData] of this.testData) {
+
+
+         for (const test of testData) {
+
+            const error = test.outcome === "Failure";
+            if (!error) {
+               break;
+            }
+            const status = this.testStatus.get(test.testId)!;
+            status.error = status.error ? status.error + 1 : 1;
+         }
+
+         for (const test of testData) {
+
+            const skipped = test.outcome === "Skipped";
+            if (!skipped) {
+               break;
+            }
+            const status = this.testStatus.get(test.testId)!;
+            status.skip = status.skip ? status.skip + 1 : 1;
+         }
+
+         for (const test of testData) {
+
+            const unspecified = test.outcome === "Unspecified";
+            if (!unspecified) {
+               break;
+            }
+            const status = this.testStatus.get(test.testId)!;
+            status.unspecified = status.unspecified ? status.unspecified + 1 : 1;
+         }
+
+
+         for (const test of testData) {
+
+            const success = (test.outcome !== "Failure" && test.outcome !== "Skipped" && test.outcome !== "Unspecified");
+            if (!success) {
+               break;
+            }
+            const status = this.testStatus.get(test.testId)!;
+            status.success = status.success ? status.success + 1 : 1;
+         }
+
+      }
+
+      const metrics = this.metrics = {
+         total: 0,
+         error: 0,
+         success: 0,
+         skipped: 0,
+         unspecified: 0
+      };
+
+      this.testStatus.forEach((v, k) => {
+         metrics.total++;
+         if (v.error) {
+            metrics.error++;
+         }
+         if (v.success) {
+            metrics.success++;
+         }
+         if (v.skip) {
+            metrics.skipped++;
+         }
+         if (v.unspecified) {
+            metrics.unspecified++;
+         }
+      });
+
+      this.loaded = true;
+
+      this.setUpdated();
+
+      // const sdetails = this.testDetails.get(debugRef)!;
+      //this.setSelection({ ref: this.suiteRefs.find(s => s.id === debugRef)!, details: sdetails, testId: debugTestId })
+   }
+
+   setFilterState(nstate?: FilterState) {
+
+      if ((this.filterState?.testState === nstate?.testState) && (this.filterState?.text === nstate?.text)) {
+         return;
+      }
+
+      this.filterState = (nstate?.testState || nstate?.text) ? nstate : undefined;
+
+      this.setUpdated();
+   }
+
+   getFilterState(): FilterState | undefined {
+      return this.filterState;
+   }
+
+
+   getFilteredTestData(testId?: string): TestData[] {
+
+      if (!this.filterState) {
+         if (testId) {
+            return this.testData.get(testId)?.flat() ?? [];
+         }
+         return Array.from(this.testData.values()).flat();
+      }
+
+      const testData: TestData[] = [];
+
+      const textLowercase = this.filterState.text?.toLowerCase();
+
+      let testStatus = this.testStatus;
+
+      if (testId) {
+         const ts = this.testStatus.get(testId);
+         testStatus = new Map();
+         if (ts) {
+            testStatus.set(testId, ts);
+         }
+      }
+
+      testStatus.forEach(ts => {
+
+         const sdata = this.testData.get(ts.testId);
+         if (!sdata) {
+            return;
+         }
+
+         const test = SuiteHandler.tests.get(ts.testId);
+
+         if (!test) {
+            return;
+         }
+
+         if (textLowercase) {
+            const name = SuiteHandler.lowerCaseTestNames.get(test.id)!;
+            if (name.indexOf(textLowercase) === -1) {
+               return;
+            }
+         }
+
+         if (!this.filterState!.testState) {
+            testData.push(...sdata);
+            return;
+         }
+
+         const testState = this.filterState!.testState;
+
+         if (testState === "Failed") {
+            if (ts.error) {
+               testData.push(...sdata);
+            }
+         }
+
+         if (testState === "Consecutive Failures") {
+            if (ts.error && ts.error > 1) {
+               testData.push(...sdata);
+            }
+         }
+
+         if (testState === "Success") {
+            if (ts.success) {
+               testData.push(...sdata);
+            }
+         }
+
+         if (testState === "Skipped") {
+            if (ts.skip) {
+               testData.push(...sdata);
+            }
+         }
+
+         if (testState === "Unspecified") {
+            if (ts.unspecified) {
+               testData.push(...sdata);
+            }
+         }
+
+      });
+
+      return testData;
+
+   }
+
+   @observable
+   updated: number = 0;
+
+   @observable
+   selectionUpdated: number = 0;
+
+   loaded = false
+
+   testDetails = new Map<string, GetTestDataDetailsResponse>();
+   suite: GetTestSuiteResponse;
+   suiteRefs: GetTestDataRefResponse[];
+
+   refs = new Map<string, GetTestDataRefResponse>();
+
+   metaData: GetTestMetaResponse;
+   metaName: string;
+
+   // quantized so common CL -> lands at same time in views
+   changeDates = new Map<number, Date>();
+   refDates = new Map<string, Date>();
+
+
+   static tests = new Map<TestId, GetTestResponse>();
+
+   private filterState?: FilterState;
+
+   testStatus: Map<TestId, TestStatus> = new Map();
+   private testData: Map<TestId, TestData[]> = new Map();
+
+   selection?: TestSelection;
+
+   metrics?: TestMetrics;
+
+   static fixedTestNames = new Map<TestId, string>();
+   static lowerCaseTestNames = new Map<TestId, string>();
+
+}
+
+type TestData = GetSuiteTestDataResponse & {
+   refId: string;
+}
+
+class SuiteGraphRenderer {
+
+   constructor(handler: SuiteHandler, testId?: string) {
+      this.testId = testId;
+      this.handler = handler;
+      this.margin = { top: 0, right: 8, bottom: 0, left: 16 };
+
+      this.headerOnly = !testId;
+
+   }
+
+   initData() {
+
+      const handler = this.handler;
+
+      this.testData = this.headerOnly ? [] : handler.getFilteredTestData(this.testId);
+   }
+
+   render(container: HTMLDivElement) {
+
+      const handler = this.handler;
+
+      if (!handler.testDetails.size || (this.hasRendered && !this.forceRender)) {
+         return;
+      }
+
+      this.clear();
+
+      this.initData();
+
+      this.hasRendered = true;
+      this.forceRender = false;
+
+      const X = d3.map(this.testData, (t) => t);
+      const xDomain = d3.extent(handler.changeDates.values(), d => d.getTime() / 1000);
+
+      let Y: TestOutcome[] = [];
+      const scolors = dashboard.getStatusColors();
+      const colors: Record<string, string> = {
+         "Success": scolors.get(StatusColor.Success)!,
+         "Failure": scolors.get(StatusColor.Failure)!,
+         "Warning": scolors.get(StatusColor.Warnings)!,
+         "Unspecified": scolors.get(StatusColor.Unspecified)!,
+         "Skipped": scolors.get(StatusColor.Skipped)!
+      };
+
+      Y = d3.map(this.testData, (d) => d.outcome);
+
+      const width = suiteGraphWidth - 32; // important: the 32 offset is to provide more mouse room to roll down latest CL's
+      const height = this.headerOnly ? 20 : 16;
+
+      const xRange = [this.margin.left, width - this.margin.right];
+
+      const xScale = d3.scaleTime(xDomain as any, xRange);
+
+      const I = d3.range(X.length);
+
+      const svgId = `suite_svg_${id_counter++}`;
+
+      const svg = d3.select(container)
+         .append("svg")
+         .attr("id", svgId)
+         .attr("width", width + 48)
+         .attr("height", height)
+         .attr("viewBox", [0, 0, width, height] as any)
+
+
+      const g = svg.append("g")
+
+      if (!this.headerOnly) {
+         g.append("line")
+            .attr("stroke", () => {
+               return dashboard.darktheme ? "#6D6C6B" : "#4D4C4B";
+            })
+            .attr("stroke-width", 1)
+            .attr("stroke-linecap", 4)
+            .attr("stroke-opacity", dashboard.darktheme ? 0.35 : 0.25)
+            .attr("x1", () => this.margin.left)
+            .attr("x2", () => width - this.margin.right)
+            .attr("y1", () => 8)
+            .attr("y2", () => 8)
+      }
+
+      const radius = 3.5
+      g.selectAll("circle")
+         .data(I)
+         .join("circle")
+         .attr("id", i => `suite_rect${this.testData[i as any].refId}`)
+         .attr("cx", i => {
+            const ref = handler.refs.get(X[i as any].refId)!;
+            return xScale(handler.changeDates.get(ref.buildChangeList)!.getTime() / 1000);
+         })
+         .attr("cy", i => 8)
+         .attr("fill", i => colors[Y[i as any]] ?? "#0000FF")
+         .attr("width", 4)
+         .attr("r", radius);
+
+      const xAxis = (g: SelectionType) => {
+
+         const dateMin = new Date(xDomain[0]! * 1000);
+         const dateMax = new Date(xDomain[1]! * 1000);
+
+         let ticks: number[] = [];
+         for (const date of d3.timeDays(dateMin, dateMax, 1).reverse()) {
+            ticks.push(date.getTime() / 1000);
+         }
+
+         if (ticks.length > 14) {
+            let nticks = [...ticks];
+            // remove first and last, will be readded 
+            const first = nticks.shift()!;
+            const last = nticks.pop()!;
+
+            const n = Math.floor(nticks.length / 12);
+
+            const rticks: number[] = [];
+            for (let i = 0; i < nticks.length; i = i + n) {
+               rticks.push(nticks[i]);
+            }
+
+            rticks.unshift(first);
+            rticks.push(last);
+            ticks = rticks;
+
+         }
+
+         g.attr("transform", `translate(0,${this.headerOnly ? 16 : 0})`)
+            .style("font-family", "Horde Open Sans Regular")
+            .style("font-size", "9px")
+            .call(d3.axisTop(xScale)
+               .tickValues(ticks)
+               .tickFormat(d => {
+                  return this.headerOnly ? getHumanTime(new Date((d as number) * 1000)) : "";
+               })
+               .tickSizeOuter(0))
+            .call(g => g.select(".domain").remove())
+            .call(g => g.selectAll(".tick line").attr("stroke-opacity", dashboard.darktheme ? 0.35 : 0.25).clone()
+               .attr("y2", height - this.margin.bottom)
+            )
+      }
+
+      svg.append("g").call(xAxis, xScale)
+
+      const root = d3.select("#root_suite_test_list") as any as DivSelectionType
+
+      const tooltip = this.tooltip = root
+         .append("div")
+         .attr("id", "tooltip")
+         .style("display", "none")
+         .style("background-color", modeColors.background)
+         .style("border", "solid")
+         .style("border-width", "1px")
+         .style("border-radius", "3px")
+         .style("border-color", dashboard.darktheme ? "#413F3D" : "#2D3F5F")
+         .style("padding", "8px")
+         .style("position", "absolute")
+         .style("z-index", "1")
+         .style("pointer-events", "none");
+
+      svg.on("mousemove", (event) => handleMouseMove(event))
+      svg.on("mouseleave", (event) => handleMouseLeave(event))
+      svg.on("click", (event) => handleMouseClick(event))
+
+      const closestRef = (x: number, y: number): GetTestDataRefResponse | undefined => {
+
+         let closest = this.testData.reduce((best, data, i) => {
+            const ref = handler.refs.get(data.refId)!;
+            const timeStamp = handler.changeDates.get(ref.buildChangeList)!.getTime() / 1000;
+            const sx = xScale(timeStamp);
+            let absx = Math.abs(sx - x)
+
+            if (absx < best.value) {
+               return { index: i, value: absx };
+            }
+            else {
+               return best;
+            }
+         }, { index: 0, value: Number.MAX_SAFE_INTEGER });
+
+         if (closest) {
+            return handler.refs.get(this.testData[closest.index]?.refId);
+         }
+
+         return undefined;
+
+      }
+
+      const handleMouseMove = (event: any) => {
+
+         const mouseX = d3.pointer(event)[0];
+         const mouseY = d3.pointer(event)[1];
+
+         const closest = closestRef(mouseX, mouseY);
+         if (closest) {
+            svg.selectAll(`circle`).attr("r", radius);
+            svg.select(`#suite_rect${closest.id}`).attr("r", radius * 2);
+
+            const timestamp = closest.id.substring(0, 8)
+
+            const date = getShortNiceTime(new Date(parseInt(timestamp, 16) * 1000), true, true);
+            let desc = `CL ${closest.buildChangeList} <br/>`;
+            //desc += `Duration ${msecToElapsed(moment.duration(closest.duration).asMilliseconds(), true, true)} <br/>`;
+            desc += `${date} <br/>`;
+
+            const timeStamp = handler.changeDates.get(closest.buildChangeList)!.getTime() / 1000;
+            const tx = xScale(timeStamp);
+
+            updateTooltip(true, tx, 0, desc);
+         }
+      }
+
+      const handleMouseLeave = (event: any) => {
+         svg.selectAll(`circle`).attr("r", radius);
+         tooltip.style("display", "none");
+      }
+
+      const handleMouseClick = (event: any) => {
+
+         const mouseX = d3.pointer(event)[0];
+         const mouseY = d3.pointer(event)[1];
+
+         const closest = closestRef(mouseX, mouseY);
+         if (closest && this.testId) {
+            const details = handler.testDetails.get(closest.id);
+            if (details) {
+               handler.setSelection({ ref: closest, details: details, testId: this.testId });
+            }
+         }
+      }
+
+      const updateTooltip = (show: boolean, x?: number, y?: number, html?: string) => {
+
+         if (!tooltip) {
+            return;
+         }
+
+         x = x ?? 0;
+         y = y ?? 0;
+
+         var elem = document.querySelector('#' + svgId) as any;
+         var parent = document.querySelector('#root_suite_test_list') as any;
+
+         if (!elem || !parent) {
+            return;
+         }
+
+         const b = elem.getBoundingClientRect()
+         const pb = parent.getBoundingClientRect()
+
+         x += (b.left - pb.left);
+         y += (b.top - pb.top);
+
+         tooltip
+            .style("display", "block")
+            .html(html ?? "")
+            .style("position", `absolute`)
+            .style("width", `max-content`)
+            .style("top", (y - 52) + "px")
+            .style("left", `${x + 24}px`)
+            .style("transform", "translateX(-108%)")
+            .style("font-family", "Horde Open Sans Semibold")
+            .style("font-size", "10px")
+            .style("line-height", "16px")
+            .style("shapeRendering", "crispEdges")
+            .style("stroke", "none")
+      }
+
+   }
+
+   clear() {
+
+   }
+
+   tooltip?: DivSelectionType;
+
+   testId?: string;
+   testData: TestData[] = [];
+   headerOnly = false;
+
+   handler: SuiteHandler;
+
+   private readonly clipId = "suite_details_clip_path";
+
+   margin: { top: number, right: number, bottom: number, left: number }
+   hasRendered = false;
+   forceRender = false;
+}
+
+const styles = mergeStyleSets({
+   list: {
+      selectors: {
+         '.ms-List-cell,.ms-DetailsRow-cell,.ms-DetailsRow': {
+            paddingTop: 0,
+            paddingBottom: 0,
+            paddingRight: 0,
+            paddingLeft: 4,
+            minHeight: 0,
+            height: 20
+         },
+         '.ms-List-cell:nth-child(odd) .ms-DetailsRow': {
+            background: "#DDDFDF",
+         },
+         '.ms-List-cell:nth-child(even) .ms-DetailsRow': {
+            backgroundColor: "#FFFFFF",
+         }
+
+      }
+   }
+});
+
+
+
+const SuiteGraph: React.FC<{ handler: SuiteHandler, testId?: string }> = observer(({ handler, testId }) => {
+
+   const graph_container_id = `automation_suite_graph_container_${id_counter++}`;
+
+   const [container, setContainer] = useState<HTMLDivElement | null>(null);
+   const [state, setState] = useState<{ graph?: SuiteGraphRenderer }>({});
+
+   // subscribe
+   if (handler.updated) { }
+
+   if (!state.graph) {
+      setState({ ...state, graph: new SuiteGraphRenderer(handler, testId) })
+      return null;
+   }
+
+   if (container) {
+      try {
+         state.graph?.render(container);
+      } catch (err) {
+         console.error(err);
+      }
+
+   }
+
+   return <Stack className={hordeClasses.horde}>
+      <Stack style={{ paddingLeft: 8, paddingTop: 8 }}>
+         <div id={graph_container_id} className="horde-no-darktheme" style={{ shapeRendering: "crispEdges", userSelect: "none", position: "relative" }} ref={(ref: HTMLDivElement) => setContainer(ref)} onMouseEnter={() => { }} onMouseLeave={() => { }} />
+      </Stack>
+   </Stack>;
+
+})
+
+const SuiteOperationsBar: React.FC<{ handler: SuiteHandler }> = observer(({ handler }) => {
+
+   if (handler.updated) { }
+
+   const stateItems: IDropdownOption[] = ["All", "Success", "Failed", "Consecutive Failures", "Skipped", "Unspecified"].map(state => {
+      return {
+         key: state,
+         text: state
+      };
+   });
+
+
+   return <Stack className={hordeClasses.modal} horizontal style={{ paddingBottom: 18 }}>
+      <Stack className="horde-no-darktheme" grow />
+      <Stack horizontal tokens={{ childrenGap: 24 }} verticalFill={true} verticalAlign={"center"}>
+         <Stack>
+            <Dropdown style={{ width: 180 }} options={stateItems} selectedKey={handler.getFilterState()?.testState ?? "All"}
+               onChange={(event, option, index) => {
+
+                  let cstate = handler.getFilterState() ?? {};
+                  if (option) {
+                     if (option.key === "All") {
+                        if (cstate.testState) {
+                           handler.setFilterState({ ...cstate, testState: undefined })
+                        }
+                     } else {
+                        if (!cstate) {
+                           cstate = {};
+                        }
+                        handler.setFilterState({ ...cstate, testState: option.key as TestState });
+                     }
+                  }
+               }} />
+         </Stack>
+         <Stack>
+            <TextField placeholder="Filter" style={{ width: 190 }} spellCheck={false} autoComplete="off" deferredValidationTime={1000} validateOnLoad={false} onGetErrorMessage={(newValue: any) => {
+               let cstate = handler.getFilterState() ?? {};
+               if (!newValue) {
+                  newValue = undefined;
+               }
+               if (cstate.text !== newValue) {
+                  handler.setFilterState({ ...cstate, text: newValue });
+               }
+               return undefined
+            }} />
+         </Stack>
+
+      </Stack>
+   </Stack>;
+})
+
+let id_counter = 0;
+
+const SuiteTestView: React.FC<{ handler: SuiteHandler }> = observer(({ handler }) => {
+
+   if (handler.selectionUpdated) { }
+
+   const selection = handler.selection;
+
+   if (selection == null) {
+      return null;
+   }
+
+   const test = SuiteHandler.tests.get(selection.testId);
+   if (!test || !handler.metaData || !handler.suite) {
+      return null;
+   }
+
+   return <Stack key={`suite_test_view_${id_counter++}`} style={{ flexShrink: 0, height: "100%" }}>
+      <AutomationSuiteTest suite={handler.suite} test={test} testRef={selection.ref} details={selection.details} metaData={handler.metaData} onClose={() => { }} />
+   </Stack>
+
+});
+
+const SuiteTestList: React.FC<{ handler: SuiteHandler }> = observer(({ handler }) => {
+
+   const windowSize = useWindowSize();
+
+   if (handler.updated) { }
+
+   interface TestItem {
+      name: string;
+      testId: string;
+   };
+
+   const scolors = dashboard.getStatusColors();
+
+
+   const columns = [
+      { key: 'name_column', name: 'Name', minWidth: 320, maxWidth: 320, isResizable: false },
+      { key: 'graph_column', name: 'Graph', minWidth: suiteGraphWidth + 8, maxWidth: suiteGraphWidth + 8, isResizable: false },
+   ];
+
+   const renderItem = (item: TestItem, index?: number, column?: IColumn) => {
+
+      if (!column) {
+         return null;
+      }
+
+      if (column.name === "Name") {
+         //return <Stack horizontal>{<StepRefStatusIcon stepRef={ref} />}<Text>{jobDetails.nodeByStepId(stepId)?.name}</Text></Stack>;
+         const status = handler.testStatus.get(item.testId);
+         let color = dashboard.darktheme ? "#E0E0E0" : undefined;
+         if (status?.success) {
+            //color = scolors.get(StatusColor.Success)!;
+         } else if (status?.error) {
+            if (!dashboard.darktheme) {
+               color = scolors.get(StatusColor.Failure)!;
+            }
+         }
+         return <Stack style={{ height: 18 }} verticalFill={true} verticalAlign="center"><Text style={{ fontSize: 11, fontFamily: "Horde Open Sans Semibold", color: color }}>{item.name}</Text></Stack>;
+      }
+
+      if (column.name === "Graph") {
+
+         return <Stack style={{ height: 18 }} verticalFill={true} verticalAlign="center"><SuiteGraph handler={handler} testId={item.testId} /></Stack>;
+      }
+
+      return null;
+
+   };
+
+   const testIds = new Set<TestId>();
+
+   const testData = handler.getFilteredTestData();
+   testData.forEach(d => testIds.add(d.testId));
+
+   let tests: GetTestResponse[] = [];
+
+   for (let testId of testIds) {
+      const test = SuiteHandler.tests.get(testId);
+      if (test) {
+         tests.push(test);
+      }
+   }
+
+   tests = tests.sort((a, b) => {
+      const nameA = SuiteHandler.fixedTestNames.get(a.id)!;
+      const nameB = SuiteHandler.fixedTestNames.get(b.id)!
+      return nameA.localeCompare(nameB);
+   });
+
+   const items = tests.map(t => {
+      return {
+         key: `automation_suite_test_${t.displayName ?? t.name}_${id_counter++}`,
+         name: SuiteHandler.fixedTestNames.get(t.id)!,
+         testId: t.id
+      }
+   });
+
+   const wheight = 1000;
+   const flexBasic = windowSize.height < wheight ? "300px" : "400px";
+   const maxHeight = windowSize.height < wheight ? 264 : 364;
+   const metrics = handler.metrics;
+
+   return <Stack key={`key_automation_suite_list_${id_counter++}`} className={hordeClasses.raised} style={{ flexBasis: flexBasic, flexShrink: 0 }}>
+      <Stack >
+         <Stack>
+            {!!metrics && <Stack horizontal tokens={{childrenGap: 24}} style={{paddingBottom: 8}}>
+               <Stack>
+                  <Text>{`Tests: ${metrics.total}`}</Text>
+               </Stack>
+               {!!metrics.success && <Stack>
+                  <Stack>
+                     <Text>{`Success: ${metrics.success} / ${Math.ceil((metrics.success / metrics.total)* 100)}%`}</Text>
+                  </Stack>
+               </Stack>}
+               {!!metrics.error && <Stack>
+                  <Stack>
+                     <Text>{`Error: ${metrics.error} / ${Math.ceil((metrics.error / metrics.total)* 100)}%`}</Text>
+                  </Stack>
+               </Stack>}
+               {!!metrics.skipped && <Stack>
+                  <Stack>
+                     <Text>{`Skipped: ${metrics.skipped} / ${Math.ceil((metrics.skipped / metrics.total)* 100)}%`}</Text>
+                  </Stack>
+               </Stack>}
+               {!!metrics.unspecified && <Stack>
+                  <Stack>
+                     <Text>{`Unspecified: ${metrics.unspecified} / ${Math.ceil((metrics.unspecified / metrics.total)* 100)}%`}</Text>
+                  </Stack>
+               </Stack>}
+            </Stack>}
+         </Stack>
+         <Stack style={{ paddingLeft: 352 }}>
+            <SuiteGraph handler={handler} />
+         </Stack>
+         <Stack>
+            <div style={{ overflowY: 'auto', overflowX: 'hidden', maxHeight: maxHeight }} data-is-scrollable={true}>
+               <DetailsList
+                  styles={{ root: { overflowX: "hidden" } }}
+                  key={`key_automation_suite_detaillist_${id_counter++}`}
+                  className={styles.list}
+                  compact={true}
+                  isHeaderVisible={false}
+                  items={items}
+                  columns={columns}
+                  layoutMode={DetailsListLayoutMode.justified}
+                  onShouldVirtualize={() => true}
+                  selectionMode={SelectionMode.none}
+                  onRenderItemColumn={renderItem}
+               />
+            </div>
+         </Stack>
+      </Stack>
+   </Stack>
+
+});
+
+const SuiteTestViewPanel: React.FC<{ handler: SuiteHandler }> = observer(({ handler }) => {
+
+   return <Stack className={hordeClasses.raised} style={{ height: "100%", marginBottom: 8 }}>
+      <Stack style={{ flexShrink: 0, height: "100%" }}>
+         <SuiteTestView handler={handler} />
+      </Stack>
+   </Stack>
+});
+
+
+
+export const AutomationSuiteDetails: React.FC<{ suite: GetTestSuiteResponse, suiteRefs: GetTestDataRefResponse[], metaData: GetTestMetaResponse, onClose: () => void }> = observer(({ suite, suiteRefs, metaData, onClose }) => {
+
+   const [state, setState] = useState<{ handler?: SuiteHandler }>({});
+
+   let handler = state.handler;
+
+   // subscribe
+   if (!handler) {
+      handler = new SuiteHandler(suite, [...suiteRefs], metaData);
+      setState({ handler: handler });
+      return null;
+   }
+
+   // subscribe
+   if (handler.updated) { }
+
+   const streamId = handler.suiteRefs.find(r => !!r.streamId)?.streamId;
+   const streamName = projectStore.streamById(streamId)?.fullname ?? "Unknown Stream";
+
+   return <Stack>
+      <Modal isOpen={true} isBlocking={true} topOffsetFixed={true} styles={{ main: { padding: 8, width: 1700, backgroundColor: modeColors.background, hasBeenOpened: false, top: "24px", position: "absolute", height: "95vh" } }} className={hordeClasses.modal} onDismiss={() => { onClose() }}>
+         <Stack className="horde-no-darktheme" style={{ height: "93vh" }}>
+            <Stack style={{ height: "100%" }} >
+               <Stack horizontal style={{ paddingLeft: 24, paddingTop: 12, paddingRight: 18 }}>
+                  <Stack tokens={{ childrenGap: 4 }}>
+                     <Text style={{ fontSize: 16, fontFamily: "Horde Open Sans Semibold" }}>{handler.metaData?.projectName} - {suite.name} - {handler.metaName}</Text>
+                     <Text style={{ fontSize: 14, fontFamily: "Horde Open Sans Semibold" }}>{streamName}</Text>
+                  </Stack>
+                  <Stack grow />
+                  <Stack>
+                     <SuiteOperationsBar handler={handler} />
+                  </Stack>
+                  <Stack style={{ paddingLeft: 24, paddingTop: 0 }}>
+                     <IconButton
+                        style={{ height: "14px" }}
+                        iconProps={{ iconName: 'Cancel' }}
+                        onClick={() => onClose()} />
+                  </Stack>
+               </Stack>
+               <Stack style={{ flexShrink: 0, flexGrow: 1, paddingLeft: 24, paddingRight: 24, paddingTop: 12 }}>
+                  <div id="root_suite_test_list" style={{ position: "relative", flexShrink: 0, flexGrow: 1 }}>
+                     {!handler.loaded && <Stack style={{ paddingTop: 16 }} horizontalAlign="center" tokens={{ childrenGap: 18 }}>
+                        <Stack>
+                           <Text variant="mediumPlus">Loading Test Suite</Text>
+                        </Stack>
+                        <Stack>
+                           <Spinner size={SpinnerSize.large} />
+                        </Stack>
+                     </Stack>}
+                     {!!handler.loaded && <Stack tokens={{ childrenGap: 18 }} style={{ height: "100%" }}>
+                        <SuiteTestList handler={handler} />
+                        <SuiteTestViewPanel handler={handler} />
+                     </Stack>}
+                  </div>
+               </Stack>
+            </Stack>
+         </Stack>
+      </Modal >
+   </Stack>
+});
\ No newline at end of file
diff --git a/Engine/Source/Programs/Horde/HordeDashboard/src/components/AutomationSuiteTest.tsx b/Engine/Source/Programs/Horde/HordeDashboard/src/components/AutomationSuiteTest.tsx
new file mode 100644
--- /dev/null
+++ b/Engine/Source/Programs/Horde/HordeDashboard/src/components/AutomationSuiteTest.tsx
@@ -0,0 +1,365 @@
+
+import { FontIcon, Image, Label, ScrollablePane, ScrollbarVisibility, Spinner, SpinnerSize, Stack, Text } from "@fluentui/react";
+import { action, observable } from 'mobx';
+import { observer } from 'mobx-react-lite';
+import React, { useState } from "react";
+import { Link } from "react-router-dom";
+import backend from "../backend";
+import { ArtifactData, GetSuiteTestDataResponse, GetTestDataDetailsResponse, GetTestDataRefResponse, GetTestMetaResponse, GetTestResponse, GetTestSuiteResponse, TestOutcome } from "../backend/Api";
+import dashboard, { StatusColor } from "../backend/Dashboard";
+import { getShortNiceTime } from "../base/utilities/timeUtils";
+
+
+type EventType = "Error" | "Info" | "Warning";
+
+enum ArtifactTag {
+   Approved = "approved",
+   Unapproved = "unapproved",
+   Difference = "difference"
+}
+
+enum EventTag {
+   ImageComparison = "image comparison"
+};
+
+type TestArtifact = {
+   ReferencePath: string;
+   Tag: ArtifactTag;
+}
+
+type TestEvent = {
+   Artifacts?: TestArtifact[];
+   Context: string;
+   Message: string;
+   Tag: EventTag;
+   Type: EventType;
+}
+
+class SuiteTestHandler {
+
+   constructor(suite: GetTestSuiteResponse, test: GetTestResponse, details: GetTestDataDetailsResponse, testRef: GetTestDataRefResponse, metaData: GetTestMetaResponse) {
+      this.suite = suite;
+      this.test = test;
+
+      this.details = details;
+      this.testRef = testRef;
+      this.metaData = metaData;
+
+      this.metaName = `${metaData.platforms.join(" - ")} / ${metaData.configurations.join(" - ")} / ${metaData.buildTargets.join(" - ")} / ${metaData.rhi === "default" ? "Default" : metaData.rhi?.toUpperCase()}`;
+
+      this.load();
+   }
+
+   async load() {
+
+      const suiteTest = this.suiteTest = this.details.suiteTests?.find(t => t.testId === this.test.id);
+      if (!suiteTest) {
+         console.error("SuiteTestHandler: Missing suite test");
+         return;
+      }
+
+      if (!this.details.testDataIds?.length) {
+         console.error("SuiteTestHandler: Suite test has no test id data");
+         return;
+      }
+
+      // @todo: can cache these
+      const requests = this.details.testDataIds.map(id => backend.getTestData(id));
+      const responses = await Promise.all(requests);
+
+      const jobId = this.jobId = responses[0].jobId as string;
+      const stepId = this.stepId = responses[1].stepId as string;
+
+      if (!jobId || !stepId) {
+         console.error("SuiteTestHandler: Missing job or step id");
+         return;
+      }
+
+      const data = responses.map(r => r.data as Record<string, any>);
+      data.forEach(d => {
+
+         const testData = d[suiteTest.uid];
+         if (testData) {
+            if (testData.Events) {
+               for (let e of (testData.Events as TestEvent[])) {
+                  this.events.push(e);
+               }
+            }
+         }
+      });
+
+      let needArtifacts = !!this.events.find(e => e.Tag === EventTag.ImageComparison);
+
+      if (needArtifacts) {
+         this.jobArtifacts = await backend.getJobArtifacts(this.jobId, this.stepId);
+      }
+
+      this.loaded = true;
+
+      this.setUpdated();
+   }
+
+   @action
+   setUpdated() {
+      this.updated++;
+   }
+
+   @observable
+   updated: number = 0;
+
+   events: TestEvent[] = [];
+
+   suite: GetTestSuiteResponse;
+   test: GetTestResponse;
+   suiteTest?: GetSuiteTestDataResponse;
+   details: GetTestDataDetailsResponse;
+   testRef: GetTestDataRefResponse;
+   metaData: GetTestMetaResponse;
+   metaName: string;
+
+   jobId?: string;
+   stepId?: string;
+   jobArtifacts?: ArtifactData[];
+
+   loaded = false;
+
+}
+
+let idCounter = 0;
+
+const EventPanel: React.FC<{ handler: SuiteTestHandler }> = observer(({ handler }) => {
+
+   if (handler.updated) { }
+
+   const events = handler.events.map(e => {
+
+      const artifactMap = new Map<string, string>();
+
+      const order = new Map<string, number>([[ArtifactTag.Approved, 0], [ArtifactTag.Difference, 1], [ArtifactTag.Unapproved, 2]]);
+
+      const imageArtifacts = e.Artifacts?.filter(a => {
+         return (a.Tag === ArtifactTag.Approved || a.Tag === ArtifactTag.Unapproved || a.Tag === ArtifactTag.Difference) && a.ReferencePath.endsWith(".png")
+      }).sort((a, b) => {
+         const oa = order.get(a.Tag) ?? -1;
+         const ob = order.get(b.Tag) ?? -1;
+         return oa - ob;
+      });
+
+      if (e.Tag === EventTag.ImageComparison && handler.jobArtifacts?.length && imageArtifacts?.length) {
+
+         imageArtifacts.forEach(ta => {
+
+            if (ta.Tag !== ArtifactTag.Approved && ta.Tag !== ArtifactTag.Unapproved && ta.Tag !== ArtifactTag.Difference) {
+               return;
+            }
+
+            const artifactName = ta.ReferencePath.replace(/\\/g, '/');
+            const artifact = handler.jobArtifacts!.find(a => a.name.indexOf(artifactName) > -1);
+            if (artifact) {
+               const imageLink = `${backend.serverUrl}/api/v1/artifacts/${artifact.id}/download?Code=${artifact.code}`;
+               artifactMap.set(ta.ReferencePath, imageLink);
+            } else {
+               console.error(`Missing artifact ${artifactName}`);
+            }
+         })
+      }
+
+      const images = imageArtifacts?.map(a => {
+
+         const url = artifactMap.get(a.ReferencePath);
+
+         if (!url) {
+            return <Stack>Missing Image: {a.ReferencePath}</Stack>
+         }
+
+         let tag = "Unknown Tag";
+         switch (a.Tag) {
+            case ArtifactTag.Approved:
+               tag = "Reference";
+               break;
+            case ArtifactTag.Unapproved:
+               tag = "Produced";
+               break;
+            case ArtifactTag.Difference:
+               tag = "Difference";
+               break;
+
+         }
+
+         return <Stack horizontalAlign="center" style={{ paddingBottom: 12 }}>
+            <Image width="100%" style={{ minWidth: 200, maxWidth: 400, minHeight: 120 }} src={url} />
+            <Text>{tag}</Text>
+         </Stack>
+      });
+
+      const scolors = dashboard.getStatusColors();
+      let color = scolors.get(StatusColor.Unspecified)!;
+
+      if (e.Type === "Error") {
+         color = scolors.get(StatusColor.Failure)!
+      }
+
+      return <Stack key={`key_test_suite_events_${idCounter++}`} >
+         <Stack horizontal tokens={{ childrenGap: 12 }}>
+            <Stack>
+               <Stack style={{ width: 6, height: "100%", backgroundColor: color }} />
+            </Stack>
+            <Stack tokens={{ childrenGap: 12 }}>
+               {!!images?.length && <Stack horizontal tokens={{ childrenGap: 48 }}>
+                  {images}
+               </Stack>}
+               <Stack>
+                  <Text style={{ lineBreak: "anywhere", whiteSpace: "pre-wrap", fontSize: 11, lineHeight: "1.65", fontFamily: "Horde Cousine Regular" }}>{e.Message}</Text>
+               </Stack>
+            </Stack>
+         </Stack>
+      </Stack>
+   });
+
+   return <Stack grow>
+      <div style={{ position: 'relative', width: "100%", height: "100%" }} data-is-scrollable>
+         <ScrollablePane scrollbarVisibility={ScrollbarVisibility.auto} onScroll={() => { }}>
+            <Stack style={{ paddingRight: 24 }} grow tokens={{ childrenGap: 24 }}>
+               {events}
+            </Stack>
+         </ScrollablePane>
+      </div>
+   </Stack>
+});
+
+
+const TestInfoBox: React.FC<{ handler: SuiteTestHandler }> = observer(({ handler }) => {
+
+
+   const labelWidth = 120;
+   const textWidth = 120;
+
+   const infoItem = (label: string, textIn: any, link?: string) => {
+      const text = textIn.toString();
+      return <Stack key={`test_info_item_${idCounter++}`} horizontal style={{ width: labelWidth + textWidth + 12 }}>
+         <Stack style={{ width: labelWidth }}>
+            <Label>{label}:</Label>
+         </Stack>
+         <Stack style={{ width: textWidth }}>
+            {!link &&
+               <Text>{text}</Text>}
+            {!!link &&
+               <Link to={link} target="_blank">{text}</Link>}
+         </Stack>
+      </Stack>
+   }
+
+   const errorCount = handler.events?.filter(e => e.Type === "Error").length;
+   const warningCount = handler.events?.filter(e => e.Type === "Warning").length;
+
+   const infoItems: JSX.Element[] = [];
+   if (handler.jobId && handler.stepId) {
+      infoItems.push(infoItem("Horde", "Job Step", `job/${handler.jobId}/?step=${handler.stepId}`));
+   }
+
+   const timestamp = handler.testRef.id.substring(0, 8)
+   const time = getShortNiceTime(new Date(parseInt(timestamp, 16) * 1000), true);
+
+
+   infoItems.push(infoItem("Change", handler.testRef.buildChangeList));
+   infoItems.push(infoItem("Date", time));
+   infoItems.push(infoItem("Errors", errorCount));
+   infoItems.push(infoItem("Warnings", warningCount));
+
+   return <Stack style={{ width: 368 }}>
+      {infoItems}
+   </Stack>;
+});
+
+
+const AutomationSuiteTestInner: React.FC<{ handler: SuiteTestHandler }> = observer(({ handler }) => {
+
+   if (handler.updated) { }
+
+   if (!handler.loaded) {
+      return <Stack style={{ paddingTop: 16 }} horizontalAlign="center" tokens={{ childrenGap: 18 }}>
+         <Stack>
+            <Text variant="mediumPlus">Loading Events</Text>
+         </Stack>
+         <Stack>
+            <Spinner size={SpinnerSize.large} />
+         </Stack>
+      </Stack>
+   }
+
+   const test = handler.test;
+   const suiteTest = handler.suiteTest;
+
+   if (!suiteTest) {
+      return null;
+   }
+
+
+   const scolors = dashboard.getStatusColors();
+
+   let color = scolors.get(StatusColor.Unspecified);
+   if (suiteTest.outcome === TestOutcome.Success) {
+      color = scolors.get(StatusColor.Success)!;
+   }
+   if (suiteTest.outcome === TestOutcome.Failure) {
+      color = scolors.get(StatusColor.Failure)!;
+   }
+
+
+   // @todo: centralize this, and needs to be fixed on client
+   const replace = ["Project.Functional Tests.Tests.", "Project.Maps.AllInPIE."];
+   function fixName(name: string) {
+
+      for (let r of replace) {
+         if (name.startsWith(r)) {
+            name = name.substring(r.length)
+         }
+      }
+      return name;
+   }
+
+   const testName = fixName(test.displayName ?? test.name);
+
+
+   return <Stack style={{ height: "100%" }} >
+      <Stack horizontal tokens={{ childrenGap: 4 }} verticalAlign="center" style={{ paddingBottom: 12 }}>
+         <Stack className="horde-no-darktheme" style={{ paddingTop: 3, paddingRight: 4 }}>
+            <FontIcon style={{ color: color }} iconName="Square" />
+         </Stack>
+         <Stack>
+            <Text style={{ fontSize: 14, fontFamily: "Horde Open Sans Semibold" }}>{testName}</Text>
+         </Stack>
+      </Stack>
+
+      <Stack grow>
+         <Stack horizontal style={{ height: "100%" }}>
+            <TestInfoBox handler={handler} />
+            <EventPanel handler={handler} />
+         </Stack>
+      </Stack>
+   </Stack>
+
+
+});
+
+export const AutomationSuiteTest: React.FC<{ suite: GetTestSuiteResponse, test: GetTestResponse, details: GetTestDataDetailsResponse, testRef: GetTestDataRefResponse, metaData: GetTestMetaResponse, onClose: () => void }> = observer(({ suite, test, details, testRef, metaData, onClose }) => {
+
+   const [state, setState] = useState<{ handler?: SuiteTestHandler }>({});
+
+   let handler = state.handler;
+
+   // subscribe
+   if (!handler) {
+      handler = new SuiteTestHandler(suite, test, details, testRef, metaData);
+      setState({ handler: handler });
+      return null;
+   }
+
+   // subscribe
+   if (handler.updated) { }
+
+   return <Stack style={{ height: "100%" }} >
+      <AutomationSuiteTestInner handler={handler} />
+   </Stack>
+
+});
diff --git a/Engine/Source/Programs/Horde/HordeDashboard/src/components/AutomationView.tsx b/Engine/Source/Programs/Horde/HordeDashboard/src/components/AutomationView.tsx
new file mode 100644
--- /dev/null
+++ b/Engine/Source/Programs/Horde/HordeDashboard/src/components/AutomationView.tsx
@@ -0,0 +1,2462 @@
+// Copyright Epic Games, Inc. All Rights Reserved.
+
+import { isNumber } from '@datadog/browser-core';
+import { Checkbox, DefaultButton, Dropdown, FocusZone, FocusZoneDirection, FontIcon, IContextualMenuItem, IContextualMenuProps, IDropdownOption, Label, PrimaryButton, ScrollablePane, ScrollbarVisibility, Spinner, SpinnerSize, Stack, Text } from '@fluentui/react';
+import * as d3 from "d3";
+import { action, observable } from 'mobx';
+import { observer } from 'mobx-react-lite';
+import moment from 'moment';
+import React, { useState } from 'react';
+import { useHistory, useLocation } from 'react-router-dom';
+import backend from '../backend';
+import { GetTestDataRefResponse, GetTestMetaResponse, GetTestResponse, GetTestStreamResponse, GetTestSuiteResponse, TestOutcome } from '../backend/Api';
+import dashboard, { StatusColor } from '../backend/Dashboard';
+import { projectStore } from '../backend/ProjectStore';
+import { useWindowSize } from '../base/utilities/hooks';
+import { getHumanTime, getShortNiceTime, msecToElapsed } from '../base/utilities/timeUtils';
+import { hordeClasses, modeColors } from '../styles/Styles';
+import { AutomationSuiteDetails } from './AutomationSuiteDetails';
+import { Breadcrumbs } from './Breadcrumbs';
+import { TopNav } from './TopNav';
+
+// some aliases
+type StreamId = string;
+type TestId = string;
+type SuiteId = string;
+type MetaId = string;
+
+type TestStatus = {
+   suite: boolean;
+   refs: GetTestDataRefResponse[];
+   success?: number;
+   error?: number;
+   skip?: number;
+}
+
+type SelectionType = d3.Selection<SVGGElement, unknown, null, undefined>;
+type DivSelectionType = d3.Selection<HTMLDivElement, unknown, null, undefined>;
+
+const defaultQueryWeeks = 2;
+const minuteInWeek = 10080;
+const queryRefreshMS = 2500;
+
+// query states
+type TestDataState = {
+   streams?: string[];
+   automation?: string;
+   platforms?: string[];
+   configurations?: string[];
+   targets?: string[];
+   rhi?: string[];
+   tests?: string[];
+   suites?: string[];
+   weeks?: number;
+}
+
+type FilterState = "Success" | "Failed" | "Consecutive Failures" | "Skipped";
+
+class TestDataHandler {
+
+   constructor(search: URLSearchParams) {
+      this.state = this.stateFromSearch(search);
+      this.load();
+   }
+
+   setAutomation(automation: string) {
+      this.state.automation = automation;
+      if (this.updateSearch()) {
+         this.setUpdated();
+      }
+   }
+
+   setQueryWeeks(weeks: number) {
+      this.state.weeks = weeks;
+      if (this.updateSearch()) {
+         this.setUpdated();
+      }
+   }
+
+   setSuiteRef(suiteRef?: string) {
+      if (this.suiteRef === suiteRef) {
+         return;
+      }
+      this.suiteRef = suiteRef;
+      this.setUpdated();
+   }
+
+   addPlatform(platform: string) {
+
+      const state = this.state;
+
+      if (state.platforms) {
+         if (state.platforms.indexOf(platform) !== -1) {
+            return;
+         }
+      }
+
+      if (!state.platforms) {
+         state.platforms = [];
+      }
+
+      state.platforms.push(platform);
+
+      if (this.updateSearch()) {
+         this.setUpdated();
+      }
+   }
+
+   removePlatform(platform: string) {
+
+      const state = this.state;
+
+      if (state.platforms?.indexOf(platform) === -1) {
+         return;
+      }
+
+      const idx = state.platforms?.indexOf(platform);
+      if (!isNumber(idx) || idx < 0) {
+         return;
+      }
+
+      state.platforms!.splice(idx, 1);
+
+      if (!state.platforms!.length) {
+         state.platforms = undefined;
+      }
+
+      if (this.updateSearch()) {
+         this.setUpdated();
+      }
+   }
+
+   addConfiguration(config: string) {
+
+      const state = this.state;
+
+      if (state.configurations) {
+         if (state.configurations.indexOf(config) !== -1) {
+            return;
+         }
+      }
+
+      if (!state.configurations) {
+         state.configurations = [];
+      }
+
+      state.configurations.push(config);
+
+      if (this.updateSearch()) {
+         this.setUpdated();
+      }
+   }
+
+   removeConfiguration(config: string) {
+
+      const state = this.state;
+
+      if (state.configurations?.indexOf(config) === -1) {
+         return;
+      }
+
+      const idx = state.configurations?.indexOf(config);
+      if (!isNumber(idx) || idx < 0) {
+         return;
+      }
+
+      state.configurations!.splice(idx, 1);
+
+      if (!state.configurations!.length) {
+         state.configurations = undefined;
+      }
+
+      if (this.updateSearch()) {
+         this.setUpdated();
+      }
+   }
+
+   addTarget(target: string) {
+
+      const state = this.state;
+
+      if (state.targets) {
+         if (state.targets.indexOf(target) !== -1) {
+            return;
+         }
+      }
+
+      if (!state.targets) {
+         state.targets = [];
+      }
+
+      state.targets.push(target);
+
+      if (this.updateSearch()) {
+         this.setUpdated();
+      }
+   }
+
+   removeTarget(target: string) {
+
+      const state = this.state;
+
+      if (state.targets?.indexOf(target) === -1) {
+         return;
+      }
+
+      const idx = state.targets?.indexOf(target);
+      if (!isNumber(idx) || idx < 0) {
+         return;
+      }
+
+      state.targets!.splice(idx, 1);
+
+      if (!state.targets!.length) {
+         state.targets = undefined;
+      }
+
+      if (this.updateSearch()) {
+         this.setUpdated();
+      }
+   }
+
+   addRHI(rhi: string) {
+
+      const state = this.state;
+
+      if (state.rhi) {
+         if (state.rhi.indexOf(rhi) !== -1) {
+            return;
+         }
+      }
+
+      if (!state.rhi) {
+         state.rhi = [];
+      }
+
+      state.rhi.push(rhi);
+
+      if (this.updateSearch()) {
+         this.setUpdated();
+      }
+   }
+
+   removeRHI(rhi: string) {
+
+      const state = this.state;
+
+      if (state.rhi?.indexOf(rhi) === -1) {
+         return;
+      }
+
+      const idx = state.rhi?.indexOf(rhi);
+      if (!isNumber(idx) || idx < 0) {
+         return;
+      }
+
+      state.rhi!.splice(idx, 1);
+
+      if (!state.rhi!.length) {
+         state.rhi = undefined;
+      }
+
+      if (this.updateSearch()) {
+         this.setUpdated();
+      }
+   }
+
+   addTest(test: string) {
+
+      const state = this.state;
+
+      if (state.tests) {
+         if (state.tests.indexOf(test) !== -1) {
+            return;
+         }
+      }
+
+      if (!state.tests) {
+         state.tests = [];
+      }
+
+      state.tests.push(test);
+
+      if (this.updateSearch()) {
+         this.setUpdated();
+      }
+   }
+
+   removeTest(test: string) {
+
+      const state = this.state;
+
+      if (state.tests?.indexOf(test) === -1) {
+         return;
+      }
+
+      const idx = state.tests?.indexOf(test);
+      if (!isNumber(idx) || idx < 0) {
+         return;
+      }
+
+      state.tests!.splice(idx, 1);
+
+      if (!state.tests!.length) {
+         state.tests = undefined;
+      }
+
+      if (this.updateSearch()) {
+         this.setUpdated();
+      }
+   }
+
+   addStream(streamId: StreamId) {
+
+      const state = this.state;
+
+      if (state.streams) {
+         if (state.streams.indexOf(streamId) !== -1) {
+            return;
+         }
+      }
+
+      if (!state.streams) {
+         state.streams = [];
+      }
+
+      state.streams.push(streamId);
+
+      if (this.updateSearch()) {
+         this.setUpdated();
+      }
+   }
+
+   removeStream(streamId: string) {
+
+      const state = this.state;
+
+      if (state.streams?.indexOf(streamId) === -1) {
+         return;
+      }
+
+      const idx = state.streams?.indexOf(streamId);
+      if (!isNumber(idx) || idx < 0) {
+         return;
+      }
+
+      state.streams!.splice(idx, 1);
+
+      if (!state.streams!.length) {
+         state.streams = undefined;
+      }
+
+      if (this.updateSearch()) {
+         this.setUpdated();
+      }
+   }
+
+
+   addSuite(suite: string) {
+
+      const state = this.state;
+
+      if (state.suites) {
+         if (state.suites.indexOf(suite) !== -1) {
+            return;
+         }
+      }
+
+      if (!state.suites) {
+         state.suites = [];
+      }
+
+      state.suites.push(suite);
+
+      if (this.updateSearch()) {
+         this.setUpdated();
+      }
+   }
+
+   removeSuite(suite: string) {
+
+      const state = this.state;
+
+      if (state.suites?.indexOf(suite) === -1) {
+         return;
+      }
+
+      const idx = state.suites?.indexOf(suite);
+      if (!isNumber(idx) || idx < 0) {
+         return;
+      }
+
+      state.suites!.splice(idx, 1);
+
+      if (!state.suites!.length) {
+         state.suites = undefined;
+      }
+
+      if (this.updateSearch()) {
+         this.setUpdated();
+      }
+   }
+
+   get currentMeta(): GetTestMetaResponse[] {
+
+      const state = this.state;
+
+      if (!state.automation || !state.streams?.length) {
+         return [];
+      }
+
+      const meta = new Map<string, GetTestMetaResponse>();
+
+      state.streams.forEach(streamId => {
+
+         const stream = this.testStreams.get(streamId);
+
+         if (!stream) {
+            return;
+         }
+
+         stream.testMetadata.forEach(m => {
+
+            if (m.projectName !== state.automation) {
+               return;
+            }
+
+            if (state.targets?.length) {
+               if (!m.buildTargets.find(v => !!state.targets?.find(v2 => v === v2))) {
+                  return;
+               }
+            }
+
+            if (state.platforms?.length) {
+               if (!m.platforms.find(v => !!state.platforms?.find(v2 => v === v2))) {
+                  return;
+               }
+            }
+
+            if (state.configurations?.length) {
+               if (!m.configurations.find(v => !!state.configurations?.find(v2 => v === v2))) {
+                  return;
+               }
+            }
+
+            if (state.rhi?.length) {
+               if (!state.rhi.find(r => m.rhi === r)) {
+                  return;
+               }
+            }
+
+            meta.set(m.id, m);
+
+         });
+
+      });
+
+      return Array.from(meta.values());
+   };
+
+   /// Tests which match current stream
+   get streamTests(): GetTestResponse[] {
+
+      const state = this.state;
+
+      if (!state.automation || !state.streams?.length) {
+         return [];
+      }
+
+      const tests = new Map<string, GetTestResponse>();
+
+      state.streams.forEach(streamId => {
+
+         const stream = this.testStreams.get(streamId);
+         if (!stream) {
+            return;
+         }
+
+         stream.tests.forEach(test => {
+            if (!!test.metadata.find(m => this.metaData.get(m)?.projectName === state.automation)) {
+               tests.set(test.id, test);
+            }
+         });
+
+      });
+
+      return Array.from(tests.values()).sort((a, b) => a.name.localeCompare(b.name));
+   }
+
+   /// Tests that match current meta
+   get metaTests(): GetTestResponse[] {
+
+      const state = this.state;
+
+      if (!state.automation || !state.streams?.length) {
+         return [];
+      }
+
+      const meta = this.currentMeta;
+      const metaIds = new Set<string>();
+      meta.forEach(m => metaIds.add(m.id));
+
+      const tests = new Map<string, GetTestResponse>();
+      const streamTests = this.streamTests;
+
+      streamTests.forEach(test => {
+
+         if (test.metadata.find(id => metaIds.has(id))) {
+            tests.set(test.id, test);
+         }
+      });
+
+      return Array.from(tests.values()).sort((a, b) => a.name.localeCompare(b.name));
+   }
+
+   /// Suites which match current stream
+   get streamSuites(): GetTestSuiteResponse[] {
+
+      const state = this.state;
+
+      if (!state.automation || !state.streams?.length) {
+         return [];
+      }
+
+      const suites = new Map<string, GetTestSuiteResponse>();
+
+      state.streams.forEach(streamId => {
+
+         const stream = this.testStreams.get(streamId);
+         if (!stream) {
+            return;
+         }
+
+         stream.testSuites.forEach(suite => {
+            if (!!suite.metadata.find(m => this.metaData.get(m)?.projectName === state.automation)) {
+               suites.set(suite.id, suite);
+            }
+         });
+
+      });
+
+      return Array.from(suites.values()).sort((a, b) => a.name.localeCompare(b.name));
+   }
+
+   /// Suites that match current meta
+   get metaSuites(): GetTestSuiteResponse[] {
+
+      const state = this.state;
+
+      if (!state.automation || !state.streams?.length) {
+         return [];
+      }
+
+      const meta = this.currentMeta;
+      const metaIds = new Set<string>();
+      meta.forEach(m => metaIds.add(m.id));
+
+      const suites = new Map<string, GetTestResponse>();
+      const streamSuites = this.streamSuites;
+
+      streamSuites.forEach(suite => {
+         if (suite.metadata.find(id => metaIds.has(id))) {
+            suites.set(suite.id, suite);
+         }
+      });
+
+
+      return Array.from(suites.values()).sort((a, b) => a.name.localeCompare(b.name));
+   }
+
+
+   updateSearch(): boolean {
+
+      const state = this.state;
+
+      state.weeks = state.weeks ? state.weeks : undefined;
+
+      state.streams = state.streams?.sort((a, b) => a.localeCompare(b));
+      state.automation = state.automation?.length ? state.automation : undefined;
+      state.platforms = state.platforms?.sort((a, b) => a.localeCompare(b));
+
+      state.configurations = state.configurations?.sort((a, b) => a.localeCompare(b));
+      state.targets = state.targets?.sort((a, b) => a.localeCompare(b));
+      state.rhi = state.rhi?.sort((a, b) => a.localeCompare(b));
+
+      state.tests = state.tests?.sort((a, b) => a.localeCompare(b));
+      state.suites = state.suites?.sort((a, b) => a.localeCompare(b));
+
+      const search = new URLSearchParams();
+
+      const csearch = this.search.toString();
+
+      if (state.automation) {
+         search.append("automation", state.automation);
+      }
+
+      if (state.weeks) {
+         search.append("weeks", state.weeks.toString());
+      }
+
+      state.streams?.forEach(s => {
+         search.append("stream", s);
+      });
+
+      state.platforms?.forEach(p => {
+         search.append("platform", p);
+      });
+
+      state.configurations?.forEach(c => {
+         search.append("configurations", c);
+      });
+
+      state.targets?.forEach(t => {
+         search.append("targets", t);
+      });
+
+      state.rhi?.forEach(r => {
+         search.append("rhi", r);
+      });
+
+      state.tests?.forEach(r => {
+         search.append("test", r);
+      });
+
+      state.suites?.forEach(r => {
+         search.append("suite", r);
+      });
+
+      if (search.toString() !== csearch) {
+
+         if (this.queryTimeoutId) {
+            clearTimeout(this.queryTimeoutId);
+         }
+
+         this.queryTimeoutId = setTimeout(() => {
+            this.queryTimeoutId = undefined;
+            this.query();
+         }, queryRefreshMS);
+
+         this.search = search;
+         return true;
+      }
+
+      return false;
+   }
+
+   private stateFromSearch(search: URLSearchParams): TestDataState {
+
+      const state: TestDataState = {};
+
+      const streams = search.getAll("stream") ?? undefined;
+      const automation = search.get("automation") ?? undefined;
+      const platforms = search.getAll("platform") ?? undefined;
+      const configurations = search.getAll("configurations") ?? undefined;
+      const targets = search.getAll("targets") ?? undefined;
+      const rhi = search.getAll("rhi") ?? undefined;
+
+      const tests = search.getAll("test") ?? undefined;
+      const suites = search.getAll("suite") ?? undefined;
+
+      const weeks = search.get("weeks") ?? undefined;
+
+      state.streams = streams?.sort((a, b) => a.localeCompare(b));
+      state.automation = automation?.length ? automation : undefined;
+      state.platforms = platforms?.sort((a, b) => a.localeCompare(b));
+
+      state.configurations = configurations?.sort((a, b) => a.localeCompare(b));
+      state.targets = targets?.sort((a, b) => a.localeCompare(b));
+      state.rhi = rhi?.sort((a, b) => a.localeCompare(b));
+
+      state.tests = tests?.sort((a, b) => a.localeCompare(b));
+      state.suites = suites?.sort((a, b) => a.localeCompare(b));
+
+      state.weeks = weeks ? parseInt(weeks) : defaultQueryWeeks;
+      if (typeof (state.weeks) !== "number") {
+         state.weeks = defaultQueryWeeks;
+      }
+
+      return state;
+   }
+
+   getMetaString(metaId: string, platform = true, config = true, target = true, rhi = true) {
+
+      const meta = this.metaData.get(metaId);
+      if (!meta) {
+         return "";
+      }
+
+      const elements: string[] = [];
+      if (platform) {
+         elements.push(`${meta.platforms.join(" - ")}`);
+      }
+      if (config) {
+         elements.push(`${meta.configurations.join(" - ")}`);
+      }
+
+      if (target) {
+         elements.push(`${meta.buildTargets.join(" - ")}`);
+      }
+
+      if (rhi) {
+         elements.push(meta.rhi === "default" ? "Default" : meta.rhi.toUpperCase());
+      }
+
+      if (elements.length) {
+         return elements.join(" / ");
+      }
+
+      return "";
+
+   }
+
+   getCommonMeta(metaIds: string[], config = true, target = true, rhi = true) {
+
+      let cconfigs = "";
+      let uniqueConfigs = false;
+
+      let ctargets = "";
+      let uniqueTargets = false;
+
+      let crhi = "";
+      let uniqueRHI = false;
+
+      metaIds.forEach(m => {
+
+         const meta = this.metaData.get(m)!;
+
+         if (config) {
+            const configs = meta.configurations.join(" - ");
+            if (cconfigs && cconfigs !== configs) {
+               uniqueConfigs = true;
+            } else {
+               cconfigs = configs;
+            }
+         }
+
+         if (target) {
+            const targets = meta.buildTargets.join(" - ");
+            if (ctargets && ctargets !== targets) {
+               uniqueTargets = true;
+            } else {
+               ctargets = targets;
+            }
+         }
+
+         if (rhi) {
+            const rhi = meta.rhi === "default" ? "Default" : meta.rhi.toUpperCase()
+
+            if (crhi && crhi !== rhi) {
+               uniqueRHI = true;
+            } else {
+               crhi = rhi;
+            }
+
+         }
+      });
+
+      const elements: string[] = [];
+      if (!uniqueConfigs || !uniqueTargets || !uniqueRHI) {
+         if (config && !uniqueConfigs) {
+            elements.push(cconfigs);
+         }
+         if (target && !uniqueTargets) {
+            elements.push(ctargets);
+         }
+         if (rhi && !uniqueRHI) {
+            elements.push(crhi);
+         }
+      }
+
+      return {
+         commonConfigs: !uniqueConfigs,
+         commonTargets: !uniqueTargets,
+         commonRHI: !uniqueRHI,
+         commonMetaString: elements.length ? elements.join(" / ") : ""
+      }
+   }
+
+   async query() {
+
+      const state = this.state;
+
+      if (!state.automation || !state.streams?.length || (!state.tests?.length && !state.suites?.length)) {
+         return;
+      }
+
+      this.setQueryLoading(true);
+
+      const testIds: Set<string> = new Set();
+      const suiteIds: Set<string> = new Set();
+      const metaIds: Set<string> = new Set(this.currentMeta.map(m => m.id));
+
+      const streamTestMap: Map<string, GetTestResponse> = new Map();
+      this.streamTests.forEach(t => {
+         streamTestMap.set(t.id, t);
+      });
+
+      const streamSuiteMap: Map<string, GetTestResponse> = new Map();
+      this.streamSuites.forEach(s => {
+         streamSuiteMap.set(s.id, s);
+      });
+
+      state.tests?.forEach(name => {
+         streamTestMap.forEach((t, id) => {
+            if (t.name === name || t.displayName === name) {
+               testIds.add(id);
+            }
+         });
+      });
+
+      state.suites?.forEach(name => {
+         streamSuiteMap.forEach((s, id) => {
+            if (s.name === name) {
+               suiteIds.add(id);
+            }
+         });
+      });
+
+
+      // notify user?
+      if (!metaIds.size || (!testIds.size && !suiteIds.size)) {
+         this.refs = [];
+         this.streamStatus = new Map();
+         this.setQueryLoading(false);
+         this.setQueryUpdated();
+         return;
+      }
+
+      let queryWeeks = this.state.weeks;
+      if (typeof (queryWeeks) !== "number") {
+         queryWeeks = defaultQueryWeeks;
+      }
+
+      const minQueryDate = new Date(new Date().valueOf() - (minuteInWeek * queryWeeks * 60000));
+      const maxQueryDate = new Date();
+
+      let refs = await backend.getTestRefs(state.streams, Array.from(metaIds), minQueryDate.toISOString(), maxQueryDate.toISOString(), undefined, undefined, testIds.size ? Array.from(testIds) : undefined, suiteIds.size ? Array.from(suiteIds) : undefined);
+
+      // deduplicate and sort intro stream -> test -> meta -> refs descending by CL
+      const streamStatus = this.streamStatus;
+      streamStatus.clear();
+      const dedupedRefs: GetTestDataRefResponse[] = [];
+      const dupeWarnings = new Set<string>();
+
+      this.changeDates = new Map();
+
+      refs.forEach(r => {
+
+         if (!this.idDates.get(r.id)) {
+            const timestamp = r.id.substring(0, 8)
+            this.idDates.set(r.id, new Date(parseInt(timestamp, 16) * 1000));
+         }
+
+         const date = this.changeDates.get(r.buildChangeList);
+         const rdate = this.idDates.get(r.id)!;
+         if (!date || date.getTime() > rdate.getTime()) {
+            this.changeDates.set(r.buildChangeList, rdate);
+         }
+
+         let streamRefs = streamStatus.get(r.streamId);
+         if (!streamRefs) {
+            streamRefs = new Map();
+            streamStatus.set(r.streamId, streamRefs);
+         }
+
+         const id = r.testId ?? r.suiteId!;
+         let testRefs = streamRefs.get(id);
+         if (!testRefs) {
+            testRefs = new Map();
+            streamRefs.set(id, testRefs);
+         }
+
+         let metaStatus = testRefs.get(r.metaId);
+         if (!metaStatus) {
+            metaStatus = { refs: [], suite: !!r.suiteId };
+            testRefs.set(r.metaId, metaStatus);
+         }
+
+         let cref = metaStatus.refs.find(tr => r.buildChangeList === tr.buildChangeList);
+         if (cref) {
+            // todo: better replace logic, use latest, should elevate errors/skips/etc
+            dupeWarnings.add(id);
+         } else {
+            metaStatus.refs.push(r);
+            dedupedRefs.push(r);
+         }
+      });
+
+      this.minRefDate = new Date();
+      this.maxRefDate = new Date();
+
+      const timeRefs = refs.sort((a, b) => this.idDates.get(a.id)!.getTime() - this.idDates.get(b.id)!.getTime());
+      if (timeRefs.length === 1) {
+         this.minRefDate = this.maxRefDate = this.idDates.get(timeRefs[0].id)!;
+      } else if (timeRefs.length > 1) {
+         this.minRefDate = this.idDates.get(timeRefs[0].id)!;
+         this.maxRefDate = this.idDates.get(timeRefs[timeRefs.length - 1].id)!;
+      }
+
+      for (const [, testMap] of streamStatus) {
+         for (const [, metaMap] of testMap) {
+            for (const [, metaStatus] of metaMap) {
+
+               const suite = metaStatus.suite;
+               metaStatus.refs = metaStatus.refs.sort((a, b) => a.buildChangeList - b.buildChangeList).reverse();
+
+               // errors
+               for (const ref of metaStatus.refs) {
+                  const error = suite ? !!ref.suiteErrorCount : ref.outcome === "Failure";
+                  if (!error) {
+                     break;
+                  }
+                  metaStatus.error = metaStatus.error ? metaStatus.error + 1 : 1;
+               }
+
+               // skip
+               for (const ref of metaStatus.refs) {
+                  const skip = suite ? !!ref.suiteSkipCount : ref.outcome === "Skipped";
+                  if (!skip) {
+                     break;
+                  }
+                  metaStatus.skip = metaStatus.skip ? metaStatus.skip + 1 : 1;
+               }
+
+               // success
+               for (const ref of metaStatus.refs) {
+                  const success = suite ? (!ref.suiteSkipCount && !ref.suiteErrorCount) : (ref.outcome !== "Failure" && ref.outcome !== "Skipped");
+                  if (!success) {
+                     break;
+                  }
+                  metaStatus.success = metaStatus.success ? metaStatus.success + 1 : 1;
+               }
+            }
+         }
+      }
+
+      // generate warnings for non unique tests in stream
+      for (let id of dupeWarnings) {
+         if (this.testMap.get(id)) {
+            console.warn(`Test ${this.testMap.get(id)!.name} has duplicate refs for stream changelists and may require a more specific name`);
+         }
+         if (this.suiteMap.get(id)) {
+            console.warn(`Suite ${this.suiteMap.get(id)!.name} has duplicate refs for stream changelists and may require a more specific name`);
+         }
+      }
+
+      this.refs = dedupedRefs;
+
+      const changelists = this.refs.map(r => r.buildChangeList);
+
+      this.minCL = this.maxCL = 0;
+
+      this.minCL = Math.min(...changelists);
+      this.maxCL = Math.max(...changelists);
+
+      this.hasQueried = true;
+
+      this.setQueryLoading(false);
+      this.setQueryUpdated();
+
+   }
+
+   // Loads immutable data for this view session
+   async load() {
+
+      const allStreams = projectStore.projects.map(p => p.streams).flat().filter(s => !!s).map(s => s!.id);
+      //const streams = testingStreams;
+      const testStreams = await backend.getTestStreams(allStreams);
+
+      // get unique meta, tests, and suites
+      const metaMap: Map<string, GetTestMetaResponse> = new Map();
+
+      const testMap = this.testMap;
+      const suiteMap = this.suiteMap;
+
+      testMap.clear();
+      suiteMap.clear();
+
+      testStreams.forEach(stream => {
+
+         // it looks like there was a misconfiguration, need a way to remove or age out tests
+         stream.tests = stream.tests.filter(t => !t.name.startsWith("EditorAutomation") && !t.name.startsWith("TargetAutomation"));
+
+         this.testStreams.set(stream.streamId, stream);
+
+         const desktop = ["Win64", "Mac", "Linux"];
+         stream.testMetadata.forEach(m => {
+
+            m.platforms = m.platforms.sort((a, b) => {
+               if (desktop.indexOf(a) !== -1) {
+                  return 1;
+               }
+               if (desktop.indexOf(b) !== -1) {
+                  return -1;
+               }
+
+               return a.localeCompare(b);
+
+            });
+
+            if (!metaMap.has(m.id)) {
+               metaMap.set(m.id, m);
+            }
+         });
+
+         stream.tests.forEach(t => {
+            if (!testMap.has(t.id)) {
+               testMap.set(t.id, t);
+            }
+         });
+
+         stream.testSuites.forEach(s => {
+            if (!suiteMap.has(s.id)) {
+               suiteMap.set(s.id, s);
+            }
+         });
+      });
+
+      const metaData = Array.from(metaMap.values());
+      this.tests = Array.from(testMap.values()).map(t => t.displayName ?? t.name).sort((a, b) => a.localeCompare(b));
+      this.suites = Array.from(suiteMap.values()).map(s => s.name).sort((a, b) => a.localeCompare(b));
+
+      this.metaData = metaMap;
+      metaMap.forEach((meta, id) => {
+         if (this.metaNames.has(id)) {
+            return;
+         }
+
+         const metaName = `${meta.platforms.join(" - ")} / ${meta.configurations.join(" - ")} / ${meta.buildTargets.join(" - ")} / ${meta.rhi}`;
+
+         this.metaNames.set(id, metaName);
+
+      });
+
+      const automation = new Set<string>();
+      const platforms = new Set<string>();
+      const configurations = new Set<string>();
+      const targets = new Set<string>();
+      const rhi = new Set<string>();
+
+      metaData.forEach(m => {
+         automation.add(m.projectName);
+         m.platforms.forEach(p => platforms.add(p));
+         m.buildTargets.forEach(t => targets.add(t));
+         m.configurations.forEach(c => configurations.add(c));
+         rhi.add(m.rhi);
+      });
+
+      this.automation = Array.from(automation).sort((a, b) => a.localeCompare(b));
+      this.platforms = Array.from(platforms).sort((a, b) => a.localeCompare(b));
+      this.configurations = Array.from(configurations).sort((a, b) => a.localeCompare(b));
+      this.targets = Array.from(targets).sort((a, b) => a.localeCompare(b));
+      this.rhi = Array.from(rhi).sort((a, b) => a.localeCompare(b));
+
+      this.loaded = true;
+      this.setUpdated();
+
+      this.query()
+   }
+
+   @action
+   setUpdated() {
+      this.updated++;
+   }
+
+   @action
+   setQueryUpdated() {
+      this.queryUpdated++;
+   }
+
+   @action
+   setQueryLoading(value: boolean) {
+      this.queryLoading = value;
+   }
+
+
+   setFilterState(nstate?: FilterState) {
+      if (this.filterState !== nstate) {
+         this.filterState = nstate;
+         this.setUpdated();
+         // need to update the center view
+         this.setQueryUpdated();
+      }
+   }
+
+   getAutomationStreams(automation: string): StreamId[] {
+
+      const streams: Set<StreamId> = new Set();
+
+      this.testStreams.forEach(testStream => {
+         if (testStream.testMetadata.find(t => t.projectName === automation)) {
+            streams.add(testStream.streamId);
+         }
+
+      });
+
+      return Array.from(streams);
+   }
+
+   getFilterState(): FilterState | undefined {
+      return this.filterState;
+   }
+
+   getRef(refId: string) {
+      return this.refs.find(r => r.id === refId);
+   }
+
+   getSuiteRefs(suiteId: string, metaId: string) {
+      return this.refs.filter(r => r.suiteId === suiteId && r.metaId === metaId);
+   }
+
+   getFilteredRefs(streamIdIn?: StreamId, idIn?: TestId | SuiteId, metaIdIn?: MetaId): GetTestDataRefResponse[] {
+
+      let status: TestStatus[] = [];
+
+      for (const [streamId, testMap] of this.streamStatus) {
+         if (streamIdIn && streamId !== streamIdIn) {
+            continue;
+         }
+         for (const [id, metaMap] of testMap) {
+            if (idIn && id !== idIn) {
+               continue;
+            }
+            for (const [metaId, metaStatus] of metaMap) {
+               if (metaIdIn && metaId !== metaIdIn) {
+                  continue;
+               }
+               status.push(metaStatus);
+            }
+         }
+      }
+
+      if (!this.filterState) {
+         return status.map(m => m.refs).flat();
+      }
+
+      if (this.filterState === "Failed") {
+         return status.filter(s => !!s.error).map(m => m.refs).flat();
+      }
+
+      if (this.filterState === "Consecutive Failures") {
+         return status.filter(s => s.error && s.error > 1).map(m => m.refs).flat();
+      }
+
+      if (this.filterState === "Skipped") {
+         return status.filter(s => !!s.skip).map(m => m.refs).flat();
+      }
+
+      if (this.filterState === "Success") {
+         return status.filter(s => !!s.success).map(m => m.refs).flat();
+      }
+
+      return [];
+
+   }
+
+   @observable
+   updated: number = 0;
+
+   @observable
+   queryUpdated: number = 0;
+
+   @observable
+   queryLoading: boolean = false;
+
+   state: TestDataState = {};
+   search: URLSearchParams = new URLSearchParams();
+
+   testStreams: Map<string, GetTestStreamResponse> = new Map();
+
+   metaData: Map<MetaId, GetTestMetaResponse> = new Map();
+   metaNames: Map<MetaId, string> = new Map();
+
+   testMap: Map<TestId, GetTestResponse> = new Map();
+   suiteMap: Map<SuiteId, GetTestSuiteResponse> = new Map();
+
+   automation: string[] = [];
+   platforms: string[] = [];
+   configurations: string[] = [];
+   targets: string[] = [];
+   rhi: string[] = [];
+
+   tests: string[] = [];
+   suites: string[] = [];
+
+
+   private refs: GetTestDataRefResponse[] = [];
+   minCL: number = 0;
+   maxCL: number = 0;
+   minRefDate: Date = new Date();
+   maxRefDate: Date = new Date();
+   idDates = new Map<string, Date>();
+   // quantized so common CL -> lands at same time in views
+   changeDates = new Map<number, Date>();
+
+   private filterState?: FilterState;
+   private streamStatus: Map<StreamId, Map<TestId | SuiteId, Map<MetaId, TestStatus>>> = new Map();
+
+   loaded = false;
+
+   hasQueried = false;
+
+   suiteRef?: string;
+
+   private queryTimeoutId: any = undefined;
+}
+
+const StreamChooser: React.FC<{ handler: TestDataHandler }> = observer(({ handler }) => {
+
+   // subscribe
+   if (handler.updated) { }
+
+   const automation = handler.state.automation;
+
+   if (!automation) {
+      return null;
+   }
+
+   const options: IContextualMenuItem[] = [];
+
+   const projects = projectStore.projects.sort((a, b) => a.order - b.order);
+
+   const streams = handler.getAutomationStreams(automation);
+
+   let anySelected = false;
+
+   projects.forEach(p => {
+
+      let projectStreams = p.streams?.filter(s => !!streams.find(stream => stream === s.id));
+
+      if (!projectStreams?.length) {
+         return;
+      }
+
+      projectStreams = projectStreams.sort((a, b) => a.name.localeCompare(b.name));
+
+      const streamItems: IContextualMenuItem[] = [];
+
+      let anyChecked = false;
+      projectStreams.forEach(s => {
+         const checked = handler.state.streams?.find(id => id === s.id);
+         if (checked) {
+            anyChecked = true;
+            anySelected = true;
+         }
+         streamItems.push({
+            key: `stream_select_${p.id}_${s.id}}`, canCheck: true, isChecked: checked, text: s.name, onClick: (ev, item) => {
+
+               if (!item) {
+                  return;
+               }
+               // don't close
+               ev?.preventDefault()
+
+               !checked ? handler.addStream(s.id) : handler.removeStream(s.id);
+            }
+         });
+      });
+
+      const subMenuProps: IContextualMenuProps = {
+         shouldFocusOnMount: true,
+         subMenuHoverDelay: 0,
+         items: streamItems,
+      };
+
+      options.push({ key: `project_stream_select_${p.id}`, canCheck: true, isChecked: anyChecked, text: p.name, subMenuProps: subMenuProps });
+   })
+
+   const menuProps: IContextualMenuProps = {
+      shouldFocusOnMount: true,
+      subMenuHoverDelay: 0,
+      items: options,
+   };
+
+   return <DefaultButton style={{ width: 220, textAlign: "left" }} menuProps={menuProps} text={anySelected ? "Select" : "None Selected"} />
+});
+
+
+const AutomationChooser: React.FC<{ handler: TestDataHandler }> = observer(({ handler }) => {
+
+   // subscribe
+   if (handler.updated) { }
+
+   const automation = handler.state.automation;
+
+   const options: IContextualMenuItem[] = [];
+
+   handler.automation.forEach(a => {
+      options.push({ key: `automation_${a}`, text: a, onClick: () => handler.setAutomation(a) });
+   });
+
+   const menuProps: IContextualMenuProps = {
+      shouldFocusOnMount: true,
+      subMenuHoverDelay: 0,
+      items: options,
+   };
+
+   return <DefaultButton style={{ width: 220, textAlign: "left" }} menuProps={menuProps} text={automation ? automation : "Select"} />
+});
+
+const TestChooser: React.FC<{ handler: TestDataHandler }> = observer(({ handler }) => {
+
+   // subscribe
+   if (handler.updated) { }
+
+   const ctests: Set<string> = new Set(handler.state.tests ?? []);
+   const metaTests = handler.metaTests;
+   const metaTestIds: Set<string> = new Set(metaTests.map(t => t.id));
+   const streamTests = handler.streamTests;
+
+   if (!streamTests.length) {
+      return null;
+   }
+
+   const tests = streamTests.map(test => {
+      const t = test.name;
+      const valid = metaTestIds.has(test.id);
+      return <Stack style={{ textDecoration: valid ? undefined : "line-through" }}>
+         <Checkbox label={t} checked={!!ctests.has(t)} onChange={(ev, checked) => {
+            if (checked) {
+               handler.addTest(t);
+            } else {
+               handler.removeTest(t);
+            }
+         }} />
+      </Stack>
+   });
+
+   return <Stack>
+      <Stack style={{ paddingTop: 12, paddingBottom: 4 }}>
+         <Label>Tests</Label>
+      </Stack>
+      <Stack tokens={{ childrenGap: 6 }} style={{ paddingLeft: 12 }}>
+         {tests}
+      </Stack>
+   </Stack>
+});
+
+const SuiteChooser: React.FC<{ handler: TestDataHandler }> = observer(({ handler }) => {
+
+   // subscribe
+   if (handler.updated) { }
+
+   const csuites: Set<string> = new Set(handler.state.suites ?? []);
+   const metaSuites = handler.metaSuites;
+   const metaSuiteIds: Set<string> = new Set(metaSuites.map(s => s.id));
+   const streamSuites = handler.streamSuites;
+
+   if (!streamSuites.length) {
+      return null;
+   }
+
+   const suites = streamSuites.map(suite => {
+      const s = suite.name;
+      const valid = metaSuiteIds.has(suite.id);
+      return <Stack style={{ textDecoration: valid ? undefined : "line-through" }}>
+         <Checkbox label={s} checked={!!csuites.has(s)} onChange={(ev, checked) => {
+            if (checked) {
+               handler.addSuite(s);
+            } else {
+               handler.removeSuite(s);
+            }
+         }} />
+      </Stack>
+   });
+
+   return <Stack>
+      <Stack style={{ paddingTop: 12, paddingBottom: 4 }}>
+         <Label>Suites</Label>
+      </Stack>
+      <Stack tokens={{ childrenGap: 6 }} style={{ paddingLeft: 12 }}>
+         {suites}
+      </Stack>
+   </Stack>
+});
+
+
+const PlatformChooser: React.FC<{ handler: TestDataHandler }> = observer(({ handler }) => {
+
+   // subscribe
+   if (handler.updated) { }
+
+   const cplatforms: Set<string> = new Set(handler.state.platforms ?? []);
+
+   const platforms = handler.platforms.map(p => {
+      return <Stack>
+         <Checkbox label={p} checked={!!cplatforms.has(p)} onChange={(ev, checked) => {
+            if (checked) {
+               handler.addPlatform(p);
+            } else {
+               handler.removePlatform(p);
+            }
+         }} />
+      </Stack>
+   });
+
+   return <Stack>
+      <Stack style={{ paddingTop: 0, paddingBottom: 4 }}>
+         <Label>Platforms</Label>
+      </Stack>
+      <Stack tokens={{ childrenGap: 6 }} style={{ paddingLeft: 12 }}>
+         {platforms}
+      </Stack>
+   </Stack>
+});
+
+const ConfigChooser: React.FC<{ handler: TestDataHandler }> = observer(({ handler }) => {
+
+   // subscribe
+   if (handler.updated) { }
+
+   const cconfig: Set<string> = new Set(handler.state.configurations ?? []);
+
+   const configurations = handler.configurations.map(c => {
+      return <Stack>
+         <Checkbox label={c} checked={!!cconfig.has(c)} onChange={(ev, checked) => {
+            if (checked) {
+               handler.addConfiguration(c);
+            } else {
+               handler.removeConfiguration(c);
+            }
+         }} />
+      </Stack>
+   });
+
+   return <Stack>
+      <Stack style={{ paddingTop: 12, paddingBottom: 4 }}>
+         <Label>Configurations</Label>
+      </Stack>
+      <Stack tokens={{ childrenGap: 6 }} style={{ paddingLeft: 12 }}>
+         {configurations}
+      </Stack>
+   </Stack>
+});
+
+const TargetChooser: React.FC<{ handler: TestDataHandler }> = observer(({ handler }) => {
+
+   // subscribe
+   if (handler.updated) { }
+
+   const ctargets: Set<string> = new Set(handler.state.targets ?? []);
+
+   const targets = handler.targets.map(t => {
+      return <Stack>
+         <Checkbox label={t} checked={!!ctargets.has(t)} onChange={(ev, checked) => {
+            if (checked) {
+               handler.addTarget(t);
+            } else {
+               handler.removeTarget(t);
+            }
+         }} />
+      </Stack>
+   });
+
+   return <Stack>
+      <Stack style={{ paddingTop: 12, paddingBottom: 4 }}>
+         <Label>Targets</Label>
+      </Stack>
+      <Stack tokens={{ childrenGap: 6 }} style={{ paddingLeft: 12 }}>
+         {targets}
+      </Stack>
+   </Stack>
+});
+
+const RHIChooser: React.FC<{ handler: TestDataHandler }> = observer(({ handler }) => {
+
+   // subscribe
+   if (handler.updated) { }
+
+   const crhi: Set<string> = new Set(handler.state.rhi ?? []);
+
+   const rhi = handler.rhi.map(r => {
+      return <Stack>
+         <Checkbox label={r === "default" ? "Default" : r.toLocaleUpperCase()} checked={!!crhi.has(r)} onChange={(ev, checked) => {
+            if (checked) {
+               handler.addRHI(r);
+            } else {
+               handler.removeRHI(r);
+            }
+         }} />
+      </Stack>
+   });
+
+   return <Stack>
+      <Stack style={{ paddingTop: 12, paddingBottom: 4 }}>
+         <Label>RHI</Label>
+      </Stack>
+      <Stack tokens={{ childrenGap: 6 }} style={{ paddingLeft: 12 }}>
+         {rhi}
+      </Stack>
+   </Stack>
+});
+
+const AutomationSidebarLeft: React.FC<{ handler: TestDataHandler }> = ({ handler }) => {
+
+   return <Stack style={{ paddingRight: 18 }}>
+      <Stack className={hordeClasses.modal}>
+         <Stack>
+            <Stack style={{ paddingTop: 0, paddingBottom: 4 }}>
+               <Label>Automation</Label>
+            </Stack>
+            <Stack tokens={{ childrenGap: 6 }} style={{ paddingLeft: 12 }}>
+               <AutomationChooser handler={handler} />
+            </Stack>
+         </Stack>
+         {!!handler.state.automation && <Stack>
+            <Stack style={{ paddingTop: 12, paddingBottom: 4 }}>
+               <Label>Streams</Label>
+            </Stack>
+            <Stack tokens={{ childrenGap: 6 }} style={{ paddingLeft: 12 }}>
+               <StreamChooser handler={handler} />
+            </Stack>
+         </Stack>}
+         <Stack>
+            <TestChooser handler={handler} />
+         </Stack>
+         <Stack>
+            <SuiteChooser handler={handler} />
+         </Stack>
+      </Stack>
+   </Stack>
+}
+
+const AutomationSidebarRight: React.FC<{ handler: TestDataHandler }> = ({ handler }) => {
+
+   return <Stack className={hordeClasses.modal}>
+      <Stack>
+         <Stack>
+            <PlatformChooser handler={handler} />
+         </Stack>
+         <Stack>
+            <ConfigChooser handler={handler} />
+         </Stack>
+         <Stack>
+            <TargetChooser handler={handler} />
+         </Stack>
+         <Stack>
+            <RHIChooser handler={handler} />
+         </Stack>
+      </Stack>
+   </Stack>
+}
+
+const AutomationOperationsBar: React.FC<{ handler: TestDataHandler }> = observer(({ handler }) => {
+
+   if (handler.updated) { }
+   if (handler.queryLoading) { }
+
+   const stateItems: IDropdownOption[] = ["All", "Success", "Failed", "Consecutive Failures", "Skipped"].map(state => {
+      return {
+         key: state,
+         text: state
+      };
+   });
+
+   type TimeSelection = {
+      text: string;
+      key: string;
+      weeks: number;
+   }
+
+   const timeSelections: TimeSelection[] = [
+      {
+         text: "Past Week", key: "time_1_week", weeks: 1
+      },
+      {
+         text: "Past 2 Weeks", key: "time_2_weeks", weeks: 2
+      },
+      {
+         text: "Past Month", key: "time_4_weeks", weeks: 4
+      },
+      {
+         text: "Past 2 Months", key: "time_8_weeks", weeks: 8
+      },
+      {
+         text: "Past 3 Months", key: "time_12_weeks", weeks: 12
+      }
+   ]
+
+   const searchDisabled = handler.queryLoading || !handler.state.automation || !handler.state.streams?.length || (!handler.state.tests?.length && !handler.state.suites?.length);
+
+   const selectedKey = timeSelections.find(t => t.weeks === handler.state.weeks);
+
+   return <Stack className={hordeClasses.modal} horizontal style={{ paddingBottom: 18, paddingTop: 2 }}>
+      <Stack grow />
+      <Stack horizontal tokens={{ childrenGap: 24 }} verticalFill={true} verticalAlign={"start"}>
+         <Stack>
+            <Dropdown placeholder="Filter" style={{ width: 180 }} options={stateItems} selectedKey={handler.getFilterState() ?? "All"}
+               onChange={(event, option, index) => {
+
+                  if (option) {
+                     if (option.key === "All") {
+                        handler.setFilterState(undefined);
+                     } else {
+                        handler.setFilterState(option.key as FilterState);
+                     }
+                  }
+               }} />
+         </Stack>
+         <Stack>
+            <Dropdown style={{ width: 180 }} options={timeSelections} selectedKey={selectedKey?.key}
+               onChange={(event, option, index) => {
+
+                  const select = option as TimeSelection;
+                  if (option) {
+                     handler.setQueryWeeks(select.weeks);
+                  }
+               }}
+            />
+         </Stack>
+         {false && <Stack>
+            <PrimaryButton disabled={searchDisabled} text="Search" onClick={() => handler.query()} />
+         </Stack>}
+      </Stack>
+   </Stack>;
+})
+
+export const AutomationView: React.FC = observer(() => {
+
+   const [state, setState] = useState<{ handler?: TestDataHandler, search?: string }>({});
+   const location = useLocation();
+   const history = useHistory();
+
+   let handler = state.handler;
+
+   // subscribe
+   if (!handler) {
+      handler = new TestDataHandler(new URLSearchParams(location.search));
+      setState({ handler: handler, search: handler.search.toString() });
+      return null;
+   }
+
+   // subscribe
+   if (handler.updated) { }
+
+   const csearch = handler.search.toString();
+
+   if (state.search !== csearch) {
+      location.search = csearch;
+      history.replace(location);
+      setState({ ...state, search: csearch });
+      return null;
+   }
+
+   let suiteRefs: GetTestDataRefResponse[] | undefined;
+   let metaData: GetTestMetaResponse | undefined;
+   let suite: GetTestSuiteResponse | undefined;
+
+   if (handler.suiteRef) {
+      const ref = handler.getRef(handler.suiteRef);
+      if (ref?.suiteId) {
+         metaData = handler.metaData.get(ref.metaId);
+         suiteRefs = handler.getSuiteRefs(ref.suiteId, ref.metaId).filter(r => r.streamId === ref.streamId);
+         suite = handler.suiteMap.get(ref.suiteId);
+      }
+   }
+
+   return (
+      <Stack>
+         {(!!suiteRefs?.length && !!metaData && !!suite) && <AutomationSuiteDetails suite={suite} suiteRefs={suiteRefs} metaData={metaData} onClose={() => handler?.setSuiteRef(undefined)} />}
+         <TopNav />
+         <Breadcrumbs items={[{ text: 'Automation' }]} />
+         <Stack horizontalAlign="center" grow styles={{ root: { width: "100%", padding: 12, backgroundColor: modeColors.background } }}>
+            <Stack styles={{ root: { width: 1774 } }}>
+               {!handler.loaded && <Stack>
+               </Stack>}
+               {handler.loaded && <Stack style={{ paddingTop: 8 }}>
+                  <Stack horizontal >
+                     <AutomationSidebarLeft handler={handler} />
+                     <Stack grow>
+                        <Stack style={{ paddingRight: 62 }}>
+                           <AutomationOperationsBar handler={handler} />
+                        </Stack>
+                        <AutomationCenter handler={handler} />
+                     </Stack>
+                     <Stack>
+                        <AutomationSidebarRight handler={handler} />
+                     </Stack>
+                  </Stack>
+               </Stack>}
+            </Stack>
+         </Stack>
+      </Stack>
+
+   );
+});
+
+let id_counter = 0;
+
+const AutomationCenter: React.FC<{ handler: TestDataHandler }> = observer(({ handler }) => {
+
+   const windowSize = useWindowSize();
+
+   if (handler.queryLoading) { }
+   if (handler.queryUpdated) { }
+
+   const refs = handler.getFilteredRefs();
+
+   const height = windowSize.height - 230;
+
+   if (handler.queryLoading) {
+      return <Stack horizontalAlign='center' style={{ paddingTop: 24, height: height }} tokens={{ childrenGap: 24 }}>
+         <Text style={{ fontSize: 24 }}>{!handler.hasQueried ? "Loading Data" : "Refreshing Data"}</Text>
+         <Spinner size={SpinnerSize.large} />
+      </Stack>
+   }
+
+   if (handler.hasQueried && !refs?.length) {
+      return <Stack horizontalAlign='center' style={{ paddingTop: 24, height: height }} tokens={{ childrenGap: 24 }}>
+         <Text style={{ fontSize: 24 }}>No Results</Text>
+      </Stack>
+   }
+
+   const streams: Set<string> = new Set();
+   const testSet: Set<string> = new Set();
+   const suiteSet: Set<string> = new Set();
+
+   refs.forEach(r => {
+
+      const meta = handler.metaData.get(r.metaId);
+
+      if (!meta) {
+         console.warn(`Missing meta data on test ref`);
+         return;
+      }
+
+      streams.add(r.streamId);
+      if (r.testId) {
+         testSet.add(r.testId);
+      }
+      if (r.suiteId) {
+         suiteSet.add(r.suiteId);
+      }
+   });
+
+   const tests = Array.from(testSet).map(tid => handler.testMap.get(tid)!).filter(t => !!t).sort((ta, tb) => ta!.name.localeCompare(tb!.name));
+   const suites = Array.from(suiteSet).map(sid => handler.suiteMap.get(sid)!).filter(s => !!s).sort((sa, sb) => sa!.name.localeCompare(sb!.name));
+
+   const testViews = tests.map(t => {
+      return <Stack key={`test_view_${t.id}_${id_counter++}`}>
+         <AutomationTestView test={t} handler={handler} />
+      </Stack>
+   });
+
+   const suiteViews = suites.map(s => {
+      return <Stack key={`suite_view_${s.id}_${id_counter++}`}>
+         <AutomationSuiteView suite={s} handler={handler} />
+      </Stack>
+   });
+
+   return <Stack style={{ paddingLeft: 12, paddingRight: 24 }}>
+      <FocusZone direction={FocusZoneDirection.vertical}>
+         <div style={{ position: 'relative', height: height }} data-is-scrollable>
+            <ScrollablePane scrollbarVisibility={ScrollbarVisibility.always} onScroll={() => { }}>
+               {!!tests.length &&
+                  <Stack style={{ paddingRight: 24 }}>
+                     <Stack style={{ paddingLeft: 12, paddingTop: 8 }} tokens={{ childrenGap: 8 }}>
+                        {testViews}
+                     </Stack>
+                  </Stack>
+               }
+               {!!suites.length &&
+                  <Stack style={{ paddingRight: 24 }}>
+                     <Stack style={{ paddingLeft: 12, paddingTop: 8 }} tokens={{ childrenGap: 8 }}>
+                        {suiteViews}
+                     </Stack>
+                  </Stack>
+               }
+            </ScrollablePane>
+         </div>
+      </FocusZone>
+   </Stack>
+
+})
+
+class AutomationGraph {
+
+   constructor(id: string, streamId: string, handler: TestDataHandler) {
+
+      this.id = id;
+      this.streamId = streamId;
+
+      this.handler = handler;
+      //this.refs = handler.filteredRefs.filter(r => r.streamId === streamId && (r.testId === id || r.suiteId == id));
+      this.refs = handler.getFilteredRefs(streamId, id);
+
+      this.margin = { top: 0, right: 32, bottom: 0, left: 160 };
+      this.clipId = `automation_clip_path_${id}_${streamId}`;
+
+      this.suite = !!this.refs.find(r => r.suiteId === id);
+   }
+
+   initData() {
+
+      const refs = this.refs;
+      const handler = this.handler;
+
+      const metaSet: Set<string> = new Set();
+      const streamSet: Set<string> = new Set();
+      refs.forEach(r => {
+         metaSet.add(r.metaId);
+         streamSet.add(r.streamId);
+      });
+
+      this.metaIds = Array.from(metaSet).sort((a, b) => {
+         const nameA = handler.metaNames.get(a)!;
+         const nameB = handler.metaNames.get(b)!;
+         return nameA.localeCompare(nameB);
+      });
+
+      const commonMeta = handler.getCommonMeta(this.metaIds);
+
+      this.metaIds.forEach(m => {
+
+         const metaRefs = this.refs.filter(r => r.metaId === m).sort((a, b) => a.buildChangeList - b.buildChangeList).reverse();
+         if (!metaRefs.length) {
+            this.metaStatus.set(m, "Unspecified");
+         } else {
+            const cref = metaRefs[0];
+            if (this.suite) {
+               if (cref.suiteErrorCount) {
+                  this.metaStatus.set(m, "Failure");
+               } else if (cref.suiteWaringCount) {
+                  this.metaStatus.set(m, "Warning");
+               } else if (cref.suiteSkipCount) {
+                  this.metaStatus.set(m, "Skipped");
+               } else {
+                  this.metaStatus.set(m, "Unspecified");
+               }
+
+            } else {
+               this.metaStatus.set(m, cref.outcome ?? "Unspecified");
+            }
+
+         }
+
+         const meta = handler.metaData.get(m)!
+         const elements: string[] = [];
+
+         elements.push(meta.platforms.join(" - "));
+
+         if (!commonMeta.commonConfigs) {
+            elements.push(meta.configurations.join(" - "));
+         }
+
+         if (!commonMeta.commonTargets) {
+            elements.push(meta.buildTargets.join(" - "));
+         }
+
+         if (!commonMeta.commonRHI) {
+            elements.push(meta.rhi === "default" ? "Default" : meta.rhi.toUpperCase());
+         }
+
+         this.metaNames.set(m, `${elements.join(" / ")}`);
+
+      });
+   }
+
+   render(container: HTMLDivElement) {
+
+      if (this.hasRendered && !this.forceRender) {
+         return;
+      }
+
+      this.clear();
+
+      this.hasRendered = true;
+      this.forceRender = false;
+
+      this.initData();
+
+      const handler = this.handler;
+      const refs = this.refs.sort((a, b) => handler.metaNames.get(a.metaId)!.localeCompare(handler.metaNames.get(b.metaId)!));
+      const width = 1254
+
+      const scolors = dashboard.getStatusColors();
+      const colors: Record<string, string> = {
+         "Success": scolors.get(StatusColor.Success)!,
+         "Failure": scolors.get(StatusColor.Failure)!,
+         "Warning": scolors.get(StatusColor.Warnings)!,
+         "Unspecified": scolors.get(StatusColor.Skipped)!,
+         "Skipped": scolors.get(StatusColor.Skipped)!
+      };
+
+      const X = d3.map(refs, (r) => handler.changeDates.get(r.buildChangeList)!.getTime() / 1000);
+      const Y = d3.map(refs, (r) => this.metaIds.indexOf(r.metaId));
+      let Z: (TestOutcome | "Unspecified" | "Warning" | "Success" | "Failure")[] = [];
+      if (this.suite) {
+         Z = d3.map(refs, (r) => {
+
+            if (!r.suiteErrorCount && !r.suiteWaringCount && !r.suiteSkipCount) {
+               // If there are skipped, will be a shape with warning/error/success included??
+               return "Success";
+            }
+
+            if (r.suiteErrorCount) {
+               return "Failure";
+            }
+            if (r.suiteWaringCount) {
+               return "Warning";
+            }
+
+            return "Unspecified";
+
+         });
+      } else {
+         Z = d3.map(refs, (r) => r.outcome ?? "Unspecified");
+      }
+
+      const xDomain = d3.extent(handler.changeDates.values(), d => d.getTime() / 1000);
+      let yDomain: any = Y;
+      yDomain = new d3.InternSet(yDomain);
+
+      const I = d3.range(X.length);
+
+      const yPadding = 1;
+      const height = Math.ceil((yDomain.size + yPadding) * 16) + this.margin.top + this.margin.bottom;
+
+      const xRange = [this.margin.left, width - this.margin.right];
+      let yRange = [this.margin.top, height - this.margin.bottom];
+
+      const xScale = d3.scaleTime(xDomain as any, xRange);
+      const yScale = d3.scalePoint(yDomain, yRange).round(true).padding(yPadding);
+
+      const svg = d3.select(container)
+         .append("svg")
+         .attr("width", width)
+         .attr("height", height + 24)
+         .attr("viewBox", [0, 0, width, height + 24] as any)
+
+
+      const g = svg.append("g")
+         .selectAll()
+         .data(d3.group(I, i => Y[i]))
+         .join("g")
+         .attr("transform", ([y]) => `translate(0,${(yScale(y) as any) + 16})`);
+
+      g.append("line")
+         .attr("stroke", ([y]) => {
+            const status = this.metaStatus.get(this.metaIds[y])!;
+            if (status === "Failure") {
+               return colors[this.metaStatus.get(this.metaIds[y])!]
+            }
+            return dashboard.darktheme ? "#6D6C6B" : "#4D4C4B";
+         })
+         .attr("stroke-width", 1)
+         .attr("stroke-linecap", 4)
+         .attr("stroke-opacity", dashboard.darktheme ? 0.35 : 0.25)
+         .attr("x1", ([, I]: any) => {
+            return this.margin.left;  //xScale(d3.min(I, i => X[i as number]) as any)
+         })
+         .attr("x2", ([, I]: any) => width - this.margin.right);
+
+
+      const radius = 3.5
+      g.selectAll("circle")
+         .data(([, I]: any) => I)
+         .join("circle")
+         .attr("id", i => `circle${refs[i as any].id}`)
+         .attr("cx", i => xScale(X[i as any]))
+         .attr("fill", i => colors[Z[i as any]] ?? "#0000FF")
+         .attr("r", radius);
+
+      g.append("text")
+         .attr("text-anchor", "start")
+         .style("alignment-baseline", "left")
+         .style("font-family", "Horde Open Sans Regular")
+         .style("font-size", 10)
+         .attr("dy", "0.15em") // center stream name
+         .attr("x", ([, I]) => 0)
+         .attr("fill", ([y]) => {
+            const status = this.metaStatus.get(this.metaIds[y])!;
+            if (!dashboard.darktheme && status === "Failure") {
+               return colors[this.metaStatus.get(this.metaIds[y])!]
+            }
+            return dashboard.darktheme ? "#E0E0E0" : "#2D3F5F";
+         })
+         .text(([y]) => this.metaNames.get(this.metaIds[y])!);
+
+      const tooltip = this.tooltip = d3.select(container)
+         .append("div")
+         .attr("id", "tooltip")
+         .style("display", "none")
+         .style("background-color", modeColors.background)
+         .style("border", "solid")
+         .style("border-width", "1px")
+         .style("border-radius", "3px")
+         .style("border-color", dashboard.darktheme ? "#413F3D" : "#2D3F5F")
+         .style("padding", "8px")
+         .style("position", "absolute")
+         .style("pointer-events", "none");
+
+
+      const xAxis = (g: SelectionType) => {
+
+         const dateMin = new Date(xDomain[0]! * 1000);
+         const dateMax = new Date(xDomain[1]! * 1000);
+
+         let ticks: number[] = [];
+         for (const date of d3.timeDays(dateMin, dateMax, 1).reverse()) {
+            ticks.push(date.getTime() / 1000);
+         }
+
+         if (ticks.length > 14) {
+            let nticks = [...ticks];
+            // remove first and last, will be readded 
+            const first = nticks.shift()!;
+            const last = nticks.pop()!;
+
+            const n = Math.floor(nticks.length / 12);
+
+            const rticks: number[] = [];
+            for (let i = 0; i < nticks.length; i = i + n) {
+               rticks.push(nticks[i]);
+            }
+
+            rticks.unshift(first);
+            rticks.push(last);
+            ticks = rticks;
+            
+         }
+         
+         g.attr("transform", `translate(0,16)`)
+            .style("font-family", "Horde Open Sans Regular")
+            .style("font-size", "9px")
+            .call(d3.axisTop(xScale)
+               .tickValues(ticks)
+               //.ticks(d3.timeDays(dateMin, dateMax))
+               .tickFormat(d => {
+                  return getHumanTime(new Date((d as number) * 1000));
+               })
+               .tickSizeOuter(0))
+            .call(g => g.select(".domain").remove())
+            .call(g => g.selectAll(".tick line").attr("stroke-opacity", dashboard.darktheme ? 0.35 : 0.25).clone()
+               .attr("stroke", dashboard.darktheme ? "#6D6C6B" : "#4D4C4B")
+               .attr("y2", height - this.margin.bottom)
+            )
+      }
+
+      // top axis
+      svg.append("g").call(xAxis, xScale)
+
+
+      const closestData = (x: number, y: number): GetTestDataRefResponse | undefined => {
+
+         y -= 16;
+
+         const metaIds = this.metaIds;
+         let closest = this.refs.reduce((best, ref, i) => {
+
+            let absy = Math.abs(yScale(metaIds.indexOf(ref.metaId) as any)! - y)
+            const timeStamp = handler.changeDates.get(ref.buildChangeList)!.getTime() / 1000;
+            const sx = xScale(timeStamp);
+            let absx = Math.abs(sx - x)
+
+            const length = Math.sqrt(absy * absy + absx * absx);
+
+            if (length < best.value) {
+               return { index: i, value: length };
+            }
+            else {
+               return best;
+            }
+         }, { index: 0, value: Number.MAX_SAFE_INTEGER });
+
+         if (closest) {
+            return this.refs[closest.index];
+         }
+
+         return undefined;
+
+      }
+
+      const handleMouseMove = (event: any) => {
+
+         const mouseX = d3.pointer(event)[0];
+         const mouseY = d3.pointer(event)[1];
+
+         const closest = closestData(mouseX, mouseY);
+         if (closest) {
+            svg.selectAll(`circle`).attr("r", radius);
+            svg.select(`#circle${closest.id}`).attr("r", radius * 2);
+
+            const timestamp = closest.id.substring(0, 8)
+
+            const cmeta = handler.getMetaString(closest.metaId, true);
+            const date = getShortNiceTime(new Date(parseInt(timestamp, 16) * 1000), true, true);
+            let desc = `${cmeta} <br/>`;
+            desc += `CL ${closest.buildChangeList} <br/>`;
+            desc += `Duration ${msecToElapsed(moment.duration(closest.duration).asMilliseconds(), true, true)} <br/>`;
+            desc += `${date} <br/>`;
+
+            const timeStamp = handler.changeDates.get(closest.buildChangeList)!.getTime() / 1000;
+            const tx = xScale(timeStamp);
+            const ty = yScale(this.metaIds.indexOf(closest.metaId) as any)!
+
+            this.updateTooltip(true, tx, ty, desc);
+         }
+      }
+
+      const handleMouseLeave = (event: any) => {
+         svg.selectAll(`circle`).attr("r", radius);
+         tooltip.style("display", "none");
+      }
+
+      const handleMouseClick = (event: any) => {
+
+         const mouseX = d3.pointer(event)[0];
+         const mouseY = d3.pointer(event)[1];
+
+         const closest = closestData(mouseX, mouseY);
+         if (closest) {
+
+            if (this.suite) {
+               handler.setSuiteRef(closest.id);
+            } else {
+               // @todo: optimize
+               backend.getTestDetails([closest.id]).then(r => {
+                  backend.getTestData(r[0].testDataIds[0], "jobId,stepId").then(d => {
+                     window.open(`/job/${d.jobId}?step=${d.stepId}`, '_blank')
+                  });
+               });
+
+            }
+         }
+      }
+
+      svg.on("mousemove", (event) => handleMouseMove(event))
+      svg.on("mouseleave", (event) => handleMouseLeave(event))
+      svg.on("click", (event) => handleMouseClick(event))
+
+   }
+
+   updateTooltip(show: boolean, x?: number, y?: number, html?: string) {
+      if (!this.tooltip) {
+         return;
+      }
+
+      x = x ?? 0;
+      y = y ?? 0;
+
+      this.tooltip
+         .style("display", "block")
+         .html(html ?? "")
+         .style("position", `absolute`)
+         .style("width", `max-content`)
+         .style("top", (y - 78) + "px")
+         .style("left", `${x}px`)
+         .style("transform", "translateX(-108%)")
+         .style("font-family", "Horde Open Sans Semibold")
+         .style("font-size", "10px")
+         .style("line-height", "16px")
+         .style("shapeRendering", "crispEdges")
+         .style("stroke", "none")
+
+   }
+
+   clear() {
+      // d3.selectAll("#automation_graph_container > *").remove();
+   }
+
+
+   // test or suite id
+   id: string;
+
+   streamId: string;
+
+   suite = false;
+
+   handler: TestDataHandler;
+   margin: { top: number, right: number, bottom: number, left: number }
+
+   hasRendered = false;
+   forceRender = false;
+
+   refs: GetTestDataRefResponse[] = [];
+   metaIds: string[] = [];
+   // metaId => status string
+   metaStatus: Map<string, string> = new Map();
+   // metaId => meta string
+   metaNames: Map<string, string> = new Map();
+
+   clipId: string;
+
+   tooltip?: DivSelectionType;
+}
+
+// --------------------------------------------------------------------------------------------------
+// TestView
+// --------------------------------------------------------------------------------------------------
+
+class TestGraphRenderer extends AutomationGraph {
+
+}
+
+const TestGraph: React.FC<{ testId: string, streamId: string, handler: TestDataHandler }> = observer(({ testId, streamId, handler }) => {
+
+   const graph_container_id = `${testId}_${streamId}_automation_graph_container`;
+
+   const [container, setContainer] = useState<HTMLDivElement | null>(null);
+   const [state, setState] = useState<{ graph?: TestGraphRenderer }>({});
+
+   if (!state.graph) {
+      setState({ ...state, graph: new TestGraphRenderer(testId, streamId, handler) })
+      return null;
+   }
+
+   if (container) {
+      try {
+         state.graph?.render(container);
+      } catch (err) {
+         console.error(err);
+      }
+
+   }
+
+   return <Stack className={hordeClasses.horde}>
+      <Stack style={{ paddingLeft: 8, paddingTop: 8 }}>
+         <div id={graph_container_id} className="horde-no-darktheme" style={{ shapeRendering: "crispEdges", userSelect: "none", position: "relative" }} ref={(ref: HTMLDivElement) => setContainer(ref)} onMouseEnter={() => { }} onMouseLeave={() => { }} />
+      </Stack>
+   </Stack>;
+
+})
+
+const AutomationTestView: React.FC<{ test: GetTestResponse, handler: TestDataHandler }> = observer(({ test, handler }) => {
+
+   const refs = handler.getFilteredRefs(undefined, test.id);
+   const streamSet = new Set<string>();
+
+   refs.forEach(r => {
+      streamSet.add(r.streamId);
+   });
+
+   const testStreams = Array.from(streamSet).sort((a, b) => a.localeCompare(b));
+
+   const streamMap = new Map<string, Map<string, GetTestDataRefResponse[]>>();
+   refs.forEach(r => {
+
+      let streamRefs = streamMap.get(r.streamId);
+      if (!streamRefs) {
+         streamRefs = new Map();
+         streamMap.set(r.streamId, streamRefs);
+      }
+
+      let metaRefs = streamRefs.get(r.metaId);
+      if (!metaRefs) {
+         metaRefs = [];
+         streamRefs.set(r.metaId, metaRefs);
+      }
+      metaRefs.push(r);
+   });
+
+   let anyFailed = false;
+   testStreams.forEach(s => {
+      const streamRefs = streamMap.get(s)!;
+      for (const [, refs] of streamRefs) {
+         const v = refs.sort((a, b) => a.buildChangeList - b.buildChangeList).reverse()[0];
+         if (v?.outcome !== "Success") {
+            anyFailed = true;
+         }
+      }
+   })
+
+   const testCommonMeta = handler.getCommonMeta(Array.from(new Set(refs.map(r => r.metaId))));
+   const scolors = dashboard.getStatusColors();
+   const streamViews = testStreams.map(s => {
+
+      const streamName = projectStore.streamById(s)?.fullname ?? s;
+      const streamRefs = streamMap.get(s)!;
+      const metaIds = Array.from(new Set<string>(streamRefs.keys()));
+
+      const commonMeta = handler.getCommonMeta(metaIds, !testCommonMeta.commonConfigs, !testCommonMeta.commonTargets, !testCommonMeta.commonRHI);
+
+      let anyPlatformFailed = false;
+      for (const [, refs] of streamRefs) {
+         const v = refs.sort((a, b) => a.buildChangeList - b.buildChangeList).reverse()[0];
+         if (v?.outcome !== "Success") {
+            anyPlatformFailed = true;
+         }
+      }
+
+      let name = streamName;
+      if (commonMeta.commonMetaString) {
+         name += ` - ${commonMeta.commonMetaString}`
+      }
+
+      return <Stack key={`automation_test_view_${s}_${test.id}`}>
+         <Stack horizontal>
+            <Stack className="horde-no-darktheme" style={{ paddingTop: 2, paddingRight: 3 }}>
+               <FontIcon style={{ fontSize: 11, color: anyPlatformFailed ? scolors.get(StatusColor.Failure)! : scolors.get(StatusColor.Success)! }} iconName="Square" />
+            </Stack>
+            <Stack>
+               <Text style={{ fontSize: 11 }}>{name}</Text>
+            </Stack>
+         </Stack>
+         <Stack style={{ paddingTop: 4 }}>
+            <TestGraph testId={test.id} streamId={s} handler={handler} />
+         </Stack>
+      </Stack>
+   });
+
+   let name = test.name;
+   if (testCommonMeta.commonMetaString) {
+      name += ` - ${testCommonMeta.commonMetaString}`
+   }
+
+   return <Stack>
+      <Stack className={hordeClasses.raised}>
+         <Stack style={{ paddingBottom: 12 }}>
+            <Stack horizontal>
+               <Stack className="horde-no-darktheme" style={{ paddingTop: 3, paddingRight: 4 }}>
+                  <FontIcon style={{ color: anyFailed ? scolors.get(StatusColor.Failure)! : scolors.get(StatusColor.Success)! }} iconName="Square" />
+               </Stack>
+               <Stack>
+                  <Text style={{ fontFamily: "Horde Open Sans Semibold" }} variant='medium'>{name}</Text>
+               </Stack>
+            </Stack>
+         </Stack>
+         <Stack style={{ paddingLeft: 12 }} tokens={{ childrenGap: 4 }}>
+            {streamViews}
+         </Stack>
+      </Stack>
+   </Stack>;
+})
+
+// --------------------------------------------------------------------------------------------------
+// SuiteView
+// --------------------------------------------------------------------------------------------------
+
+class SuiteGraphRenderer extends AutomationGraph {
+}
+
+const SuiteGraph: React.FC<{ suiteId: string, streamId: string, handler: TestDataHandler }> = observer(({ suiteId, streamId, handler }) => {
+
+   const graph_container_id = `${suiteId}_${streamId}_automation_suite_graph_container`;
+
+   const [container, setContainer] = useState<HTMLDivElement | null>(null);
+   const [state, setState] = useState<{ graph?: SuiteGraphRenderer }>({});
+
+   if (!state.graph) {
+      setState({ ...state, graph: new SuiteGraphRenderer(suiteId, streamId, handler) })
+      return null;
+   }
+
+   if (container) {
+      try {
+         state.graph?.render(container);
+      } catch (err) {
+         console.error(err);
+      }
+   }
+
+   return <Stack className={hordeClasses.horde}>
+      <Stack style={{ paddingLeft: 8 }}>
+         <div id={graph_container_id} className="horde-no-darktheme" style={{ shapeRendering: "crispEdges", userSelect: "none", position: "relative" }} ref={(ref: HTMLDivElement) => setContainer(ref)} onMouseEnter={() => { }} onMouseLeave={() => { }} />
+      </Stack>
+   </Stack>;
+
+})
+
+
+const AutomationSuiteView: React.FC<{ suite: GetTestResponse, handler: TestDataHandler }> = observer(({ suite, handler }) => {
+
+   const refs = handler.getFilteredRefs(undefined, suite.id);
+   const streamSet = new Set<string>();
+
+   refs.forEach(r => {
+      streamSet.add(r.streamId);
+   });
+
+   const testStreams = Array.from(streamSet).sort((a, b) => a.localeCompare(b));
+
+   const streamMap = new Map<string, Map<string, GetTestDataRefResponse[]>>();
+   refs.forEach(r => {
+
+      let streamRefs = streamMap.get(r.streamId);
+      if (!streamRefs) {
+         streamRefs = new Map();
+         streamMap.set(r.streamId, streamRefs);
+      }
+
+      let metaRefs = streamRefs.get(r.metaId);
+      if (!metaRefs) {
+         metaRefs = [];
+         streamRefs.set(r.metaId, metaRefs);
+      }
+      metaRefs.push(r);
+   });
+
+   let anyFailed = false;
+   testStreams.forEach(s => {
+      const streamRefs = streamMap.get(s)!;
+      for (const [, refs] of streamRefs) {
+         const v = refs.sort((a, b) => a.buildChangeList - b.buildChangeList).reverse()[0];
+         if (v?.suiteErrorCount || v?.suiteSkipCount) {
+            anyFailed = true;
+         }
+      }
+   })
+
+   const scolors = dashboard.getStatusColors();
+   const suiteCommonMeta = handler.getCommonMeta(Array.from(new Set(refs.map(r => r.metaId))));
+
+   const streamViews = testStreams.map(s => {
+      const streamName = projectStore.streamById(s)?.fullname ?? s;
+      const streamRefs = streamMap.get(s)!;
+      const metaIds = Array.from(new Set<string>(streamRefs.keys()));
+
+
+      const commonMeta = handler.getCommonMeta(metaIds, !suiteCommonMeta.commonConfigs, !suiteCommonMeta.commonTargets, !suiteCommonMeta.commonRHI);
+
+      let anyPlatformFailed = false;
+      for (const [, refs] of streamRefs) {
+         const v = refs.sort((a, b) => a.buildChangeList - b.buildChangeList).reverse()[0];
+         if (v?.suiteErrorCount || v?.suiteSkipCount) {
+            anyPlatformFailed = true;
+         }
+      }
+
+      let name = streamName;
+      if (commonMeta.commonMetaString) {
+         name += ` - ${commonMeta.commonMetaString}`
+      }
+
+      return <Stack key={`automation_suite_view_${s}_${suite.id}`}>
+         <Stack horizontal>
+            <Stack className="horde-no-darktheme" style={{ paddingTop: 2, paddingRight: 3 }}>
+               <FontIcon style={{ fontSize: 11, color: anyPlatformFailed ? scolors.get(StatusColor.Failure)! : scolors.get(StatusColor.Success)! }} iconName="Square" />
+            </Stack>
+            <Stack>
+               <Text style={{ fontSize: 11 }}>{name}</Text>
+            </Stack>
+         </Stack>
+         <Stack style={{ paddingTop: 4 }}>
+            <SuiteGraph suiteId={suite.id} streamId={s} handler={handler} />
+         </Stack>
+      </Stack>
+   });
+
+   let name = suite.name;
+   if (suiteCommonMeta.commonMetaString) {
+      name += ` - ${suiteCommonMeta.commonMetaString}`
+   }
+
+   return <Stack>
+      <Stack className={hordeClasses.raised}>
+         <Stack style={{ paddingBottom: 12 }}>
+            <Stack horizontal>
+               <Stack className="horde-no-darktheme" style={{ paddingTop: 3, paddingRight: 4 }}>
+                  <FontIcon style={{ color: anyFailed ? scolors.get(StatusColor.Failure)! : scolors.get(StatusColor.Success)! }} iconName="Square" />
+               </Stack>
+               <Stack>
+                  <Text style={{ fontFamily: "Horde Open Sans Semibold" }} variant='medium'>{name}</Text>
+               </Stack>
+            </Stack>
+         </Stack>
+         <Stack style={{ paddingLeft: 12 }} tokens={{ childrenGap: 4 }}>
+            {streamViews}
+         </Stack>
+      </Stack>
+   </Stack>;
+
+})
\ No newline at end of file
diff --git a/Engine/Source/Programs/Horde/HordeDashboard/src/components/TopNav.tsx b/Engine/Source/Programs/Horde/HordeDashboard/src/components/TopNav.tsx
--- a/Engine/Source/Programs/Horde/HordeDashboard/src/components/TopNav.tsx
+++ b/Engine/Source/Programs/Horde/HordeDashboard/src/components/TopNav.tsx
@@ -362,12 +362,6 @@ export const TopNav: React.FC<{ suppressServer?: boolean }> = observer(({ suppre
          link: `/agents`
       });
 
-      resourceItems.push({
-         key: "admin_devices",
-         text: "Devices",
-         link: `/devices`
-      });
-
       resourceItems.push({
          key: "admin_perforce_servers",
          text: "Perforce Servers",
@@ -384,6 +378,31 @@ export const TopNav: React.FC<{ suppressServer?: boolean }> = observer(({ suppre
          }
       });
 
+      // Automation
+      const automatonItems: IContextualMenuItem[] = [];
+
+      automatonItems.push({
+         key: "admin_devices",
+         text: "Devices",
+         link: `/devices`
+      });
+
+      automatonItems.push({
+         key: "admin_tests",
+         text: "Tests",
+         link: `/automation`
+      });
+
+      subItems.push({
+         itemType: ContextualMenuItemType.Section,
+         key: `admin_automation`,
+         sectionProps: {
+            title: "Automation",
+            items: automatonItems,
+            bottomDivider: true
+         }
+      });
+
       // Monitoring
       const monitoringItems: IContextualMenuItem[] = [];
       monitoringItems.push({
diff --git a/Engine/Source/Programs/Horde/HordeDashboard/tsconfig.json b/Engine/Source/Programs/Horde/HordeDashboard/tsconfig.json
--- a/Engine/Source/Programs/Horde/HordeDashboard/tsconfig.json
+++ b/Engine/Source/Programs/Horde/HordeDashboard/tsconfig.json
@@ -1,6 +1,6 @@
 {
   "compilerOptions": {
-    "target": "es5",
+    "target": "es2015",
     "lib": [
       "dom",
       "dom.iterable",
