Horde: css fix for automation home
#skipci
#rnx

diff --git a/Engine/Source/Programs/Horde/HordeDashboard/src/components/AutomationView.tsx b/Engine/Source/Programs/Horde/HordeDashboard/src/components/AutomationView.tsx
--- a/Engine/Source/Programs/Horde/HordeDashboard/src/components/AutomationView.tsx
+++ b/Engine/Source/Programs/Horde/HordeDashboard/src/components/AutomationView.tsx
@@ -1655,7 +1655,7 @@ export const AutomationView: React.FC = observer(() => {
    }
 
    return (
-      <Stack>
+      <Stack className={hordeClasses.horde}>
          {(!!suiteRefs?.length && !!metaData && !!suite) && <AutomationSuiteDetails suite={suite} suiteRefs={suiteRefs} metaData={metaData} onClose={() => handler?.setSuiteRef(undefined)} />}
          <TopNav />
          <Breadcrumbs items={[{ text: 'Automation' }]} />
