Horde: Add support for referencing individual exports by alias (eg. content addresses or content ids).

Since bundles have more defined restrictions on size (ie. expectation is that the server can read them in memory), we do not require aliases to match node hashes exactly - allowing internal representation to diverge from the public representation.

diff --git a/Engine/Source/Programs/Horde/Horde.Build.Tests/TestSetup.cs b/Engine/Source/Programs/Horde/Horde.Build.Tests/TestSetup.cs
--- a/Engine/Source/Programs/Horde/Horde.Build.Tests/TestSetup.cs
+++ b/Engine/Source/Programs/Horde/Horde.Build.Tests/TestSetup.cs
@@ -157,6 +157,8 @@ namespace Horde.Build.Tests
 			settings.AdminClaimType = HordeClaimTypes.Role;
 			settings.AdminClaimValue = "app-horde-admins";
 			settings.WithAws = true;
+
+			settings.ForceConfigUpdateOnStartup = true;
 		}
 
 		protected override void ConfigureServices(IServiceCollection services)
@@ -252,6 +254,7 @@ namespace Horde.Build.Tests
 
 			services.AddSingleton<ILegacyStorageClient, BasicStorageClient>();
 
+			services.AddSingleton<ConfigUpdateService>();
 			services.AddSingleton<StorageService>();
 
 			services.AddSingleton<ISingletonDocument<AgentSoftwareChannels>>(new SingletonDocumentStub<AgentSoftwareChannels>());
diff --git a/Engine/Source/Programs/Horde/Horde.Build/Commands/Bundles/BundleWriteTest.cs b/Engine/Source/Programs/Horde/Horde.Build/Commands/Bundles/BundleWriteTest.cs
--- a/Engine/Source/Programs/Horde/Horde.Build/Commands/Bundles/BundleWriteTest.cs
+++ b/Engine/Source/Programs/Horde/Horde.Build/Commands/Bundles/BundleWriteTest.cs
@@ -1,6 +1,7 @@
 // Copyright Epic Games, Inc. All Rights Reserved.
 
 using System;
+using System.Collections.Generic;
 using System.Diagnostics;
 using System.IO;
 using System.Security.Cryptography;
@@ -21,6 +22,7 @@ namespace Horde.Build.Commands.Bundles
 			public Task DeleteRefAsync(RefName name, CancellationToken cancellationToken = default) => Task.CompletedTask;
 			public Task<Stream> ReadBlobAsync(BlobLocator locator, CancellationToken cancellationToken = default) => throw new NotImplementedException();
 			public Task<Stream> ReadBlobRangeAsync(BlobLocator locator, int offset, int length, CancellationToken cancellationToken = default) => throw new NotImplementedException();
+			public IAsyncEnumerable<NodeHandle> FindNodesAsync(Utf8String alias, CancellationToken cancellationToken = default) => throw new NotImplementedException();
 			public Task<NodeHandle?> TryReadRefTargetAsync(RefName name, DateTime cacheTime = default, CancellationToken cancellationToken = default) => throw new NotImplementedException();
 			public Task<BlobLocator> WriteBlobAsync(Stream stream, Utf8String prefix = default, CancellationToken cancellationToken = default) => Task.FromResult(BlobLocator.Create(HostId.Empty));
 			public Task<NodeHandle> WriteRefAsync(RefName name, Bundle bundle, int exportIdx, Utf8String prefix = default, RefOptions? options = null, CancellationToken cancellationToken = default) => throw new NotImplementedException();
diff --git a/Engine/Source/Programs/Horde/Horde.Build/Defaults/globals.json b/Engine/Source/Programs/Horde/Horde.Build/Defaults/globals.json
--- a/Engine/Source/Programs/Horde/Horde.Build/Defaults/globals.json
+++ b/Engine/Source/Programs/Horde/Horde.Build/Defaults/globals.json
@@ -17,6 +17,10 @@
 				"id": "default-backend",
 				"type":"FileSystem",
 				"baseDir": "Storage" // Default base directory is C:\ProgramData\HordeServer
+			},
+			{
+				"id": "memory-backend", // Used for automated tests
+				"type": "Memory"
 			}
 		],
 		"namespaces": [
@@ -43,6 +47,10 @@
 				"id": "horde-tools",
 				"prefix": "Tools/",
 				"backend": "default-backend"
+			},
+			{
+				"id": "memory",
+				"backend": "memory-backend"
 			}
 		]
 	},
diff --git a/Engine/Source/Programs/Horde/Horde.Build/Storage/StorageConfig.cs b/Engine/Source/Programs/Horde/Horde.Build/Storage/StorageConfig.cs
--- a/Engine/Source/Programs/Horde/Horde.Build/Storage/StorageConfig.cs
+++ b/Engine/Source/Programs/Horde/Horde.Build/Storage/StorageConfig.cs
@@ -133,6 +133,11 @@ namespace Horde.Build.Storage
 		/// </summary>
 		public double GcDelayHrs { get; set; } = 6.0;
 
+		/// <summary>
+		/// Support querying exports by their aliases
+		/// </summary>
+		public bool EnableAliases { get; set; }
+
 		/// <summary>
 		/// Access list for this namespace
 		/// </summary>
diff --git a/Engine/Source/Programs/Horde/Horde.Build/Storage/StorageController.cs b/Engine/Source/Programs/Horde/Horde.Build/Storage/StorageController.cs
--- a/Engine/Source/Programs/Horde/Horde.Build/Storage/StorageController.cs
+++ b/Engine/Source/Programs/Horde/Horde.Build/Storage/StorageController.cs
@@ -1,6 +1,7 @@
 // Copyright Epic Games, Inc. All Rights Reserved.
 
 using System;
+using System.Collections.Generic;
 using System.IO;
 using System.Linq;
 using System.Text.RegularExpressions;
@@ -39,6 +40,47 @@ namespace Horde.Build.Storage
 		public bool? SupportsRedirects { get; set; }
 	}
 
+	/// <summary>
+	/// Response object for finding a node
+	/// </summary>
+	public class FindNodeResponse
+	{
+		/// <summary>
+		/// Hash of the target node
+		/// </summary>
+		public IoHash Hash { get; set; }
+
+		/// <summary>
+		/// Locator for the target blob
+		/// </summary>
+		public BlobLocator Blob { get; set; }
+
+		/// <summary>
+		/// Export index for the ref
+		/// </summary>
+		public int ExportIdx { get; set; }
+
+		/// <summary>
+		/// Constructor
+		/// </summary>
+		public FindNodeResponse(NodeHandle target)
+		{
+			Hash = target.Hash;
+			Blob = target.Locator.Blob;
+			ExportIdx = target.Locator.ExportIdx;
+		}
+	}
+	/// <summary>
+	/// Response object for searching for nodes with a given alias
+	/// </summary>
+	public class FindNodesResponse
+	{
+		/// <summary>
+		/// Hash of the target node
+		/// </summary>
+		public List<FindNodeResponse> Nodes { get; set; } = new List<FindNodeResponse>();
+	}
+
 	/// <summary>
 	/// Request object for writing a ref
 	/// </summary>
@@ -232,6 +274,37 @@ namespace Horde.Build.Storage
 #pragma warning restore CA2000 // Dispose objects before losing scope
 		}
 
+		/// <summary>
+		/// Retrieves data from the storage service. 
+		/// </summary>
+		/// <param name="namespaceId">Namespace to fetch from</param>
+		/// <param name="alias">Alias of the node to find</param>
+		/// <param name="cancellationToken">Cancellation token for the operation</param>
+		[HttpGet]
+		[Route("/api/v1/storage/{namespaceId}/nodes")]
+		public async Task<ActionResult<FindNodesResponse>> FindNodesAsync(NamespaceId namespaceId, string alias, CancellationToken cancellationToken = default)
+		{
+			if (!await _storageService.AuthorizeAsync(namespaceId, User, AclAction.ReadBlobs, null, cancellationToken))
+			{
+				return Forbid(AclAction.WriteBlobs);
+			}
+
+			IStorageClientImpl client = await _storageService.GetClientAsync(namespaceId, cancellationToken);
+
+			FindNodesResponse response = new FindNodesResponse();
+			await foreach (NodeHandle handle in client.FindNodesAsync(alias, cancellationToken))
+			{
+				response.Nodes.Add(new FindNodeResponse(handle));
+			}
+
+			if (response.Nodes.Count == 0)
+			{
+				return NotFound();
+			}
+
+			return response;
+		}
+
 		/// <summary>
 		/// Writes a ref to the storage service.
 		/// </summary>
diff --git a/Engine/Source/Programs/Horde/Horde.Build/Storage/StorageService.cs b/Engine/Source/Programs/Horde/Horde.Build/Storage/StorageService.cs
--- a/Engine/Source/Programs/Horde/Horde.Build/Storage/StorageService.cs
+++ b/Engine/Source/Programs/Horde/Horde.Build/Storage/StorageService.cs
@@ -5,9 +5,11 @@ using System.Buffers;
 using System.Buffers.Binary;
 using System.Collections;
 using System.Collections.Generic;
+using System.Composition;
 using System.IO;
 using System.Linq;
 using System.Reflection;
+using System.Runtime.CompilerServices;
 using System.Runtime.InteropServices;
 using System.Security.Claims;
 using System.Threading;
@@ -109,7 +111,7 @@ namespace Horde.Build.Storage
 				_outer = outer;
 				Config = config;
 				Backend = backend;
-				SupportsRedirects = backend is IStorageBackendWithRedirects;
+				SupportsRedirects = (backend is IStorageBackendWithRedirects) && !config.EnableAliases;
 
 				_prefix = config.Prefix;
 				if (_prefix.Length > 0 && !_prefix.EndsWith("/", StringComparison.Ordinal))
@@ -172,11 +174,46 @@ namespace Horde.Build.Storage
 			/// <inheritdoc/>
 			public override async Task<BlobLocator> WriteBlobAsync(Stream stream, Utf8String prefix = default, CancellationToken cancellationToken = default)
 			{
-				BlobLocator locator = await _outer.AddBlobAsync(NamespaceId, prefix, cancellationToken);
+				BlobLocator locator;
+				if (Config.EnableAliases)
+				{
+					using (MemoryStream memoryStream = new MemoryStream())
+					{
+						// Read the blob into memory
+						await stream.CopyToAsync(memoryStream, cancellationToken);
 
-				string path = GetBlobPath(locator.BlobId);
-				await Backend.WriteAsync(path, stream, cancellationToken);
+						// Reset the position back to zero and read the header
+						memoryStream.Position = 0;
+						BundleHeader header = await BundleHeader.FromStreamAsync(memoryStream, cancellationToken);
+
+						List<ExportInfo> exports = new List<ExportInfo>();
+						for (int idx = 0; idx < header.Exports.Count; idx++)
+						{
+							BundleExport export = header.Exports[idx];
+							if (!export.Alias.IsEmpty)
+							{
+								exports.Add(new ExportInfo(export.Alias.ToString(), export.Hash, idx));
+							}
+						}
+
+						// Add the blob record
+						locator = await _outer.AddBlobAsync(NamespaceId, prefix, exports, cancellationToken);
+
+						// Write it to the backend
+						string path = GetBlobPath(locator.BlobId);
+						memoryStream.Position = 0;
+						await Backend.WriteAsync(path, memoryStream, cancellationToken);
+					}
+				}
+				else
+				{
+					// Add the blob record
+					locator = await _outer.AddBlobAsync(NamespaceId, prefix, null, cancellationToken);
 
+					// Write it to the backend
+					string path = GetBlobPath(locator.BlobId);
+					await Backend.WriteAsync(path, stream, cancellationToken);
+				}
 				return locator;
 			}
 
@@ -189,7 +226,7 @@ namespace Horde.Build.Storage
 					return null;
 				}
 
-				BlobLocator locator = await _outer.AddBlobAsync(NamespaceId, prefix, cancellationToken);
+				BlobLocator locator = await _outer.AddBlobAsync(NamespaceId, prefix, null, cancellationToken);
 				string path = GetBlobPath(locator.BlobId);
 
 				Uri? url = await redirectBackend.GetWriteRedirectAsync(path, cancellationToken);
@@ -209,6 +246,13 @@ namespace Horde.Build.Storage
 
 			#endregion
 
+			#region Nodes
+
+			/// <inheritdoc/>
+			public override IAsyncEnumerable<NodeHandle> FindNodesAsync(Utf8String alias, CancellationToken cancellationToken = default) => _outer.FindNodesAsync(NamespaceId, alias, cancellationToken);
+
+			#endregion
+
 			#region Refs
 
 			/// <inheritdoc/>
@@ -258,6 +302,29 @@ namespace Horde.Build.Storage
 			}
 		}
 
+		class ExportInfo
+		{
+			[BsonElement("alias")]
+			public string Alias { get; set; } = String.Empty;
+
+			[BsonElement("hash")]
+			public IoHash Hash { get; set; }
+
+			[BsonElement("idx")]
+			public int Index { get; set; }
+
+			public ExportInfo()
+			{
+			}
+
+			public ExportInfo(string alias, IoHash hash, int index)
+			{
+				Alias = alias;
+				Hash = hash;
+				Index = index;
+			}
+		}
+
 		class BlobInfo
 		{
 			public ObjectId Id { get; set; }
@@ -271,9 +338,12 @@ namespace Horde.Build.Storage
 			[BsonElement("blob")]
 			public BlobId BlobId { get; set; }
 
-			[BsonElement("imp")]
+			[BsonElement("imp"), BsonIgnoreIfNull]
 			public List<ObjectId>? Imports { get; set; }
 
+			[BsonElement("exp"), BsonIgnoreIfNull]
+			public List<ExportInfo>? Exports { get; set; }
+
 			[BsonIgnore]
 			public BlobLocator Locator => new BlobLocator(HostId, BlobId);
 
@@ -408,6 +478,7 @@ namespace Horde.Build.Storage
 			List<MongoIndex<BlobInfo>> blobIndexes = new List<MongoIndex<BlobInfo>>();
 			blobIndexes.Add(keys => keys.Ascending(x => x.Imports));
 			blobIndexes.Add(keys => keys.Ascending(x => x.NamespaceId).Ascending(x => x.BlobId), unique: true);
+			blobIndexes.Add(keys => keys.Ascending(x => x.NamespaceId).Ascending($"{nameof(BlobInfo.Exports)}.{nameof(ExportInfo.Alias)}"));
 			_blobCollection = mongoService.GetCollection<BlobInfo>("Storage.Blobs", blobIndexes);
 
 			List<MongoIndex<RefInfo>> refIndexes = new List<MongoIndex<RefInfo>>();
@@ -689,7 +760,7 @@ namespace Horde.Build.Storage
 		#region Blobs
 
 		/// <inheritdoc/>
-		async Task<BlobLocator> AddBlobAsync(NamespaceId namespaceId, Utf8String prefix = default, CancellationToken cancellationToken = default)
+		async Task<BlobLocator> AddBlobAsync(NamespaceId namespaceId, Utf8String prefix = default, List<ExportInfo>? exports = null, CancellationToken cancellationToken = default)
 		{
 			HostId hostId = HostId.Empty;
 
@@ -697,6 +768,7 @@ namespace Horde.Build.Storage
 			BlobId blobId = (prefix.Length > 0) ? new BlobId($"{prefix}/{id}") : new BlobId(id.ToString());
 
 			BlobInfo blobInfo = new BlobInfo(id, namespaceId, hostId, blobId);
+			blobInfo.Exports = exports;
 			await _blobCollection.InsertOneAsync(blobInfo, new InsertOneOptions { }, cancellationToken);
 
 			return new BlobLocator(hostId, blobId);
@@ -776,6 +848,34 @@ namespace Horde.Build.Storage
 
 		#endregion
 
+		#region Nodes
+
+		/// <summary>
+		/// Finds nodes with the given type and hash
+		/// </summary>
+		/// <param name="namespaceId">Namespace to search</param>
+		/// <param name="alias">Alias for the node</param>
+		/// <param name="cancellationToken">Cancellation token for the operation</param>
+		/// <returns>Sequence of thandles</returns>
+		async IAsyncEnumerable<NodeHandle> FindNodesAsync(NamespaceId namespaceId, Utf8String alias, [EnumeratorCancellation] CancellationToken cancellationToken = default)
+		{
+			await foreach (BlobInfo blobInfo in _blobCollection.Find(x => x.NamespaceId == namespaceId && x.Exports!.Any(y => y.Alias == alias)).ToAsyncEnumerable(cancellationToken))
+			{
+				if (blobInfo.Exports != null)
+				{
+					foreach (ExportInfo exportInfo in blobInfo.Exports)
+					{
+						if (exportInfo.Alias == alias)
+						{
+							yield return new NodeHandle(exportInfo.Hash, blobInfo.Locator, exportInfo.Index);
+						}
+					}
+				}
+			}
+		}
+
+		#endregion
+
 		#region Refs
 
 		record RefCacheKey(NamespaceId NamespaceId, RefName Name);
diff --git a/Engine/Source/Programs/Shared/EpicGames.Horde/Storage/Backends/FileStorageClient.cs b/Engine/Source/Programs/Shared/EpicGames.Horde/Storage/Backends/FileStorageClient.cs
--- a/Engine/Source/Programs/Shared/EpicGames.Horde/Storage/Backends/FileStorageClient.cs
+++ b/Engine/Source/Programs/Shared/EpicGames.Horde/Storage/Backends/FileStorageClient.cs
@@ -1,9 +1,11 @@
 // Copyright Epic Games, Inc. All Rights Reserved.
 
 using System;
+using System.Collections.Generic;
 using System.Diagnostics;
 using System.Globalization;
 using System.IO;
+using System.Linq;
 using System.Threading;
 using System.Threading.Tasks;
 using EpicGames.Core;
@@ -72,6 +74,16 @@ namespace EpicGames.Horde.Storage.Backends
 
 		#endregion
 
+		#region Nodes
+
+		/// <inheritdoc/>
+		public override IAsyncEnumerable<NodeHandle> FindNodesAsync(Utf8String alias, CancellationToken cancellationToken = default)
+		{
+			throw new NotSupportedException("File storage client does not currently support aliases.");
+		}
+
+		#endregion
+
 		#region Refs
 
 		/// <inheritdoc/>
diff --git a/Engine/Source/Programs/Shared/EpicGames.Horde/Storage/Backends/HttpStorageClient.cs b/Engine/Source/Programs/Shared/EpicGames.Horde/Storage/Backends/HttpStorageClient.cs
--- a/Engine/Source/Programs/Shared/EpicGames.Horde/Storage/Backends/HttpStorageClient.cs
+++ b/Engine/Source/Programs/Shared/EpicGames.Horde/Storage/Backends/HttpStorageClient.cs
@@ -1,14 +1,18 @@
 // Copyright Epic Games, Inc. All Rights Reserved.
 
 using System;
+using System.Collections.Generic;
 using System.IO;
 using System.Net;
 using System.Net.Http;
 using System.Net.Http.Headers;
 using System.Net.Http.Json;
+using System.Runtime.CompilerServices;
 using System.Runtime.InteropServices;
 using System.Threading;
 using System.Threading.Tasks;
+using System.Web;
+using System.Xml.Linq;
 using EpicGames.Core;
 using EpicGames.Horde.Storage;
 using Microsoft.Extensions.Caching.Memory;
@@ -33,6 +37,18 @@ namespace EpicGames.Horde.Storage.Backends
 			public bool? SupportsRedirects { get; set; }
 		}
 
+		class FindNodeResponse
+		{
+			public IoHash Hash { get; set; }
+			public BlobLocator Blob { get; set; }
+			public int ExportIdx { get; set; }
+		}
+
+		class FindNodesResponse
+		{
+			public List<FindNodeResponse> Nodes { get; set; } = new List<FindNodeResponse>();
+		}
+
 		class ReadRefResponse
 		{
 			public IoHash Hash { get; set; }
@@ -177,6 +193,32 @@ namespace EpicGames.Horde.Storage.Backends
 
 		#endregion
 
+		#region Nodes
+
+		/// <inheritdoc/>
+		public override async IAsyncEnumerable<NodeHandle> FindNodesAsync(Utf8String alias, [EnumeratorCancellation] CancellationToken cancellationToken = default)
+		{
+			_logger.LogDebug("Finding nodes with alias {Alias}", alias);
+			using (HttpClient httpClient = _createClient())
+			{
+				using (HttpRequestMessage request = new HttpRequestMessage(HttpMethod.Get, $"nodes?alias={HttpUtility.UrlEncode(alias.ToString())}"))
+				{
+					using (HttpResponseMessage response = await httpClient.SendAsync(request, cancellationToken))
+					{
+						response.EnsureSuccessStatusCode();
+
+						FindNodesResponse? message = await response.Content.ReadFromJsonAsync<FindNodesResponse>(cancellationToken: cancellationToken);
+						foreach (FindNodeResponse node in message!.Nodes)
+						{
+							yield return new NodeHandle(node.Hash, node.Blob, node.ExportIdx);
+						}
+					}
+				}
+			}
+		}
+
+		#endregion
+
 		#region Refs
 
 		/// <inheritdoc/>
diff --git a/Engine/Source/Programs/Shared/EpicGames.Horde/Storage/Backends/MemoryStorageClient.cs b/Engine/Source/Programs/Shared/EpicGames.Horde/Storage/Backends/MemoryStorageClient.cs
--- a/Engine/Source/Programs/Shared/EpicGames.Horde/Storage/Backends/MemoryStorageClient.cs
+++ b/Engine/Source/Programs/Shared/EpicGames.Horde/Storage/Backends/MemoryStorageClient.cs
@@ -5,6 +5,7 @@ using System.Buffers;
 using System.Collections.Concurrent;
 using System.Collections.Generic;
 using System.IO;
+using System.Runtime.CompilerServices;
 using System.Threading;
 using System.Threading.Tasks;
 using EpicGames.Core;
@@ -18,6 +19,8 @@ namespace EpicGames.Horde.Storage.Backends
 	/// </summary>
 	public class MemoryStorageClient : StorageClientBase
 	{
+		record class ExportEntry(NodeHandle Handle, ExportEntry? Next);
+
 		/// <summary>
 		/// Map of blob id to blob data
 		/// </summary>
@@ -28,6 +31,11 @@ namespace EpicGames.Horde.Storage.Backends
 		/// </summary>
 		readonly ConcurrentDictionary<RefName, NodeHandle> _refs = new ConcurrentDictionary<RefName, NodeHandle>();
 
+		/// <summary>
+		/// Content addressed data lookup
+		/// </summary>
+		readonly ConcurrentDictionary<Utf8String, ExportEntry> _exports = new ConcurrentDictionary<Utf8String, ExportEntry>();
+
 		/// <inheritdoc cref="_blobs"/>
 		public IReadOnlyDictionary<BlobLocator, Bundle> Blobs => _blobs;
 
@@ -61,12 +69,42 @@ namespace EpicGames.Horde.Storage.Backends
 		public override async Task<BlobLocator> WriteBlobAsync(Stream stream, Utf8String prefix = default, CancellationToken cancellationToken = default)
 		{
 			BlobLocator locator = BlobLocator.Create(HostId.Empty, prefix);
-			_blobs[locator] = await Bundle.FromStreamAsync(stream, cancellationToken);
+			Bundle bundle = await Bundle.FromStreamAsync(stream, cancellationToken);
+			_blobs[locator] = bundle;
+
+			for (int idx = 0; idx < bundle.Header.Exports.Count; idx++)
+			{
+				BundleExport export = bundle.Header.Exports[idx];
+				if(!export.Alias.IsEmpty)
+				{
+					NodeHandle handle = new NodeHandle(export.Hash, locator, idx);
+					_exports.AddOrUpdate(export.Alias, _ => new ExportEntry(handle, null), (_, entry) => new ExportEntry(handle, entry));
+				}
+			}
+
 			return locator;
 		}
 
 		#endregion
 
+		#region Nodes
+
+		/// <inheritdoc/>
+		public override async IAsyncEnumerable<NodeHandle> FindNodesAsync(Utf8String alias, [EnumeratorCancellation] CancellationToken cancellationToken = default)
+		{
+			if(_exports.TryGetValue(alias, out ExportEntry? entry))
+			{
+				for (; entry != null; entry = entry.Next)
+				{
+					cancellationToken.ThrowIfCancellationRequested();
+					await Task.Yield();
+					yield return entry.Handle;
+				}
+			}
+		}
+
+		#endregion
+
 		#region Refs
 
 		/// <inheritdoc/>
diff --git a/Engine/Source/Programs/Shared/EpicGames.Horde/Storage/Bundle.cs b/Engine/Source/Programs/Shared/EpicGames.Horde/Storage/Bundle.cs
--- a/Engine/Source/Programs/Shared/EpicGames.Horde/Storage/Bundle.cs
+++ b/Engine/Source/Programs/Shared/EpicGames.Horde/Storage/Bundle.cs
@@ -5,13 +5,13 @@ using K4os.Compression.LZ4;
 using System;
 using System.Buffers;
 using System.Buffers.Binary;
+using System.Collections;
 using System.Collections.Generic;
 using System.Diagnostics;
 using System.Globalization;
 using System.IO;
 using System.IO.Compression;
 using System.Linq;
-using System.Reflection;
 using System.Text;
 using System.Threading;
 using System.Threading.Tasks;
@@ -29,6 +29,11 @@ namespace EpicGames.Horde.Storage
 		/// </summary>
 		Initial = 0,
 
+		/// <summary>
+		/// Added the <see cref="BundleExport.Alias"/> property
+		/// </summary>
+		ExportAliases = 1,
+
 		/// <summary>
 		/// Last item in the enum. Used for <see cref="Latest"/>
 		/// </summary>
@@ -231,7 +236,7 @@ namespace EpicGames.Horde.Storage
 			CompressionFormat = (BundleCompressionFormat)reader.ReadUnsignedVarInt();
 			Types = reader.ReadVariableLengthArray(() => new BundleType(reader));
 			Imports = reader.ReadVariableLengthArray(() => new BundleImport(reader));
-			Exports = reader.ReadVariableLengthArray(() => new BundleExport(reader));
+			Exports = reader.ReadVariableLengthArray(() => new BundleExport(reader, version));
 
 			if (CompressionFormat == BundleCompressionFormat.None)
 			{
@@ -510,7 +515,7 @@ namespace EpicGames.Horde.Storage
 		/// Uncompressed length of this node
 		/// </summary>
 		public int Length { get; }
-		
+
 		/// <summary>
 		/// Nodes referenced by this export. Indices in this array correspond to a lookup table consisting
         /// of the imported nodes in the order they are declared in the header, followed by nodes listed in the
@@ -518,36 +523,53 @@ namespace EpicGames.Horde.Storage
         /// </summary>
 		public IReadOnlyList<int> References { get; }
 
+		/// <summary>
+		/// Alias for the content of this node. Used for content addressing.
+		/// </summary>
+		public Utf8String Alias { get; }
+
 		/// <summary>
 		/// Constructor
 		/// </summary>
-		public BundleExport(int typeIdx, IoHash hash, int length, IReadOnlyList<int> references)
+		public BundleExport(int typeIdx, IoHash hash, int length, IReadOnlyList<int> references, Utf8String alias = default)
 		{
 			TypeIdx = typeIdx;
 			Hash = hash;
 			Length = length;
 			References = references;
+			Alias = alias;
 		}
 
 		/// <summary>
 		/// Deserialize an export
 		/// </summary>
 		/// <param name="reader">Reader to deserialize from</param>
-		public BundleExport(IMemoryReader reader)
+		/// <param name="version">Version number of the bundle</param>
+		public BundleExport(IMemoryReader reader, BundleVersion version)
 		{
 			TypeIdx = (int)reader.ReadUnsignedVarInt();
 			Hash = reader.ReadIoHash();
 			Length = (int)reader.ReadUnsignedVarInt();
 
 			int numReferences = (int)reader.ReadUnsignedVarInt();
-			int[] references = new int[numReferences];
-
-			for (int idx = 0; idx < numReferences; idx++)
+			if (numReferences == 0)
 			{
-				references[idx] = (int)reader.ReadUnsignedVarInt();
+				References = Array.Empty<int>();
+			}
+			else
+			{
+				int[] references = new int[numReferences];
+				for (int idx = 0; idx < numReferences; idx++)
+				{
+					references[idx] = (int)reader.ReadUnsignedVarInt();
+				}
+				References = references;
 			}
 
-			References = references;
+			if (version >= BundleVersion.ExportAliases)
+			{
+				Alias = reader.ReadUtf8String();
+			}
 		}
 
 		/// <summary>
@@ -565,6 +587,8 @@ namespace EpicGames.Horde.Storage
 			{
 				writer.WriteUnsignedVarInt(References[idx]);
 			}
+
+			writer.WriteUtf8String(Alias);
 		}
 	}
 
diff --git a/Engine/Source/Programs/Shared/EpicGames.Horde/Storage/IStorageClient.cs b/Engine/Source/Programs/Shared/EpicGames.Horde/Storage/IStorageClient.cs
--- a/Engine/Source/Programs/Shared/EpicGames.Horde/Storage/IStorageClient.cs
+++ b/Engine/Source/Programs/Shared/EpicGames.Horde/Storage/IStorageClient.cs
@@ -148,6 +148,18 @@ namespace EpicGames.Horde.Storage
 
 		#endregion
 
+		#region Nodes
+
+		/// <summary>
+		/// Finds nodes with the given alias. Unlike refs, aliases do not serve as GC roots.
+		/// </summary>
+		/// <param name="alias">Alias for the node</param>
+		/// <param name="cancellationToken">Cancellation token for the operation</param>
+		/// <returns>Nodes matching the given handle</returns>
+		IAsyncEnumerable<NodeHandle> FindNodesAsync(Utf8String alias, CancellationToken cancellationToken = default);
+
+		#endregion
+
 		#region Refs
 
 		/// <summary>
@@ -236,6 +248,13 @@ namespace EpicGames.Horde.Storage
 
 			#endregion
 
+			#region Nodes
+
+			/// <inheritdoc/>
+			public IAsyncEnumerable<NodeHandle> FindNodesAsync(Utf8String alias, CancellationToken cancellationToken = default) => _inner.FindNodesAsync(alias, cancellationToken);
+
+			#endregion
+
 			#region Refs
 
 			/// <inheritdoc/>
diff --git a/Engine/Source/Programs/Shared/EpicGames.Horde/Storage/StorageClientBase.cs b/Engine/Source/Programs/Shared/EpicGames.Horde/Storage/StorageClientBase.cs
--- a/Engine/Source/Programs/Shared/EpicGames.Horde/Storage/StorageClientBase.cs
+++ b/Engine/Source/Programs/Shared/EpicGames.Horde/Storage/StorageClientBase.cs
@@ -1,6 +1,7 @@
 // Copyright Epic Games, Inc. All Rights Reserved.
 
 using System;
+using System.Collections.Generic;
 using System.IO;
 using System.Threading;
 using System.Threading.Tasks;
@@ -33,6 +34,13 @@ namespace EpicGames.Horde.Storage
 
 		#endregion
 
+		#region Nodes
+
+		/// <inheritdoc/>
+		public abstract IAsyncEnumerable<NodeHandle> FindNodesAsync(Utf8String alias, CancellationToken cancellationToken = default);
+
+		#endregion
+
 		#region Refs
 
 		/// <inheritdoc/>
