Horde: Fix file size not being reset correctly when replicating Perforce data.

diff --git a/Engine/Source/Programs/Horde/Horde.Build/Perforce/PerforceReplicator.cs b/Engine/Source/Programs/Horde/Horde.Build/Perforce/PerforceReplicator.cs
--- a/Engine/Source/Programs/Horde/Horde.Build/Perforce/PerforceReplicator.cs
+++ b/Engine/Source/Programs/Horde/Horde.Build/Perforce/PerforceReplicator.cs
@@ -162,6 +162,7 @@ namespace Horde.Build.Perforce
 				handle._fileWriter.Reset();
 				handle._path = path;
 				handle._size = size;
+				handle._sizeWritten = 0;
 				handle._flags = flags;
 
 				_openHandles.Add(fd, handle);
