Horde: Redact things that look like workspace root directories when hashing issues. Allows matching issues which are generated on different agent types.

diff --git a/Engine/Source/Programs/Horde/Horde.Build.Tests/IssueServiceTests.cs b/Engine/Source/Programs/Horde/Horde.Build.Tests/IssueServiceTests.cs
--- a/Engine/Source/Programs/Horde/Horde.Build.Tests/IssueServiceTests.cs
+++ b/Engine/Source/Programs/Horde/Horde.Build.Tests/IssueServiceTests.cs
@@ -1273,6 +1273,24 @@ namespace Horde.Build.Tests
 			Assert.AreEqual("Warnings in Update Version Files and Compile UnrealHeaderTool Win64", issue.Summary);
 		}
 
+		[TestMethod]
+		public async Task HashedIssueTest4()
+		{
+			IJob job = CreateJob(_mainStreamId, 120, "Compile Test", _graph);
+
+			await ParseEventsAsync(job, 0, 0, new[] { "Assertion failed: 1 == 2 [File:D:\\build\\++UE5\\Sync\\Engine\\Source\\Runtime\\Core\\Tests\\Misc\\AssertionMacrosTest.cpp] [Line: 119]" });
+			await UpdateCompleteStep(job, 0, 0, JobStepOutcome.Failure);
+
+			await ParseEventsAsync(job, 0, 1, new[] { "Assertion failed: 1 == 2 [File:C:\\build\\++UE5+Inc\\Sync\\Engine\\Source\\Runtime\\Core\\Tests\\Misc\\AssertionMacrosTest.cpp] [Line: 119]" });
+			await UpdateCompleteStep(job, 0, 1, JobStepOutcome.Failure);
+
+			List<IIssue> issues = await IssueCollection.FindIssuesAsync();
+			Assert.AreEqual(1, issues.Count);
+
+			IIssue issue = issues[0];
+			Assert.AreEqual("Errors in Update Version Files and Compile UnrealHeaderTool Win64", issue.Summary);
+		}
+
 		[TestMethod]
 		public async Task SymbolIssueTest()
 		{
diff --git a/Engine/Source/Programs/Horde/Horde.Build/Issues/Handlers/HashedIssueHandler.cs b/Engine/Source/Programs/Horde/Horde.Build/Issues/Handlers/HashedIssueHandler.cs
--- a/Engine/Source/Programs/Horde/Horde.Build/Issues/Handlers/HashedIssueHandler.cs
+++ b/Engine/Source/Programs/Horde/Horde.Build/Issues/Handlers/HashedIssueHandler.cs
@@ -80,6 +80,7 @@ namespace Horde.Build.Issues.Handlers
 		static bool TryGetHash(string message, out Md5Hash hash)
 		{
 			string sanitized = message.ToUpperInvariant();
+			sanitized = Regex.Replace(sanitized, @"(?<![a-zA-Z])(?:[A-Z]:|/)[^ :]+[/\\]SYNC[/\\]", "{root}/"); // Redact things that look like workspace roots; may be different between agents
 			sanitized = Regex.Replace(sanitized, @"0[xX][0-9a-fA-F]+", "H"); // Redact hex strings
 			sanitized = Regex.Replace(sanitized, @"\d[\d.,:]*", "n"); // Redact numbers and timestamp like things
 
