Horde: Enable aliases for the in-memory test backend, to fix tests.

diff --git a/Engine/Source/Programs/Horde/Horde.Build/Defaults/globals.json b/Engine/Source/Programs/Horde/Horde.Build/Defaults/globals.json
--- a/Engine/Source/Programs/Horde/Horde.Build/Defaults/globals.json
+++ b/Engine/Source/Programs/Horde/Horde.Build/Defaults/globals.json
@@ -50,7 +50,8 @@
 			},
 			{
 				"id": "memory",
-				"backend": "memory-backend"
+				"backend": "memory-backend",
+				"enableAliases": true
 			}
 		]
 	},
