Horde: Use managed process instances to start Redis/Mongo for running tests, to ensure they're terminated after the process exits. Would previously leave an instance running if terminated through the debugger, etc...

diff --git a/Engine/Source/Programs/Horde/Horde.Build.Tests/DatabaseRunner.cs b/Engine/Source/Programs/Horde/Horde.Build.Tests/DatabaseRunner.cs
--- a/Engine/Source/Programs/Horde/Horde.Build.Tests/DatabaseRunner.cs
+++ b/Engine/Source/Programs/Horde/Horde.Build.Tests/DatabaseRunner.cs
@@ -9,33 +9,43 @@ using System.Net.NetworkInformation;
 using System.Net.Sockets;
 using System.Reflection;
 using System.Runtime.InteropServices;
+using System.Threading.Tasks;
 using EpicGames.Core;
 
 namespace Horde.Build.Tests;
 
-public abstract class DatabaseRunner
+public abstract class DatabaseRunner : IDisposable
 {
 	private readonly string _binName;
 	private readonly int _defaultPort;
 	private readonly string _name;
-	private readonly bool _printStdErr;
-	private readonly bool _printStdOut;
 	private readonly bool _reuseProcess;
 	protected string TempDir { get; }
-	private Process? _proc;
 
-	protected DatabaseRunner(string name, string binName, int defaultPort, bool reuseProcess, bool printStdOut = false,
-		bool printStdErr = true)
+	private ManagedProcessGroup? _processGroup;
+	private ManagedProcess? _process;
+
+	protected DatabaseRunner(string name, string binName, int defaultPort, bool reuseProcess)
 	{
 		_name = name;
 		_binName = binName;
 		_defaultPort = defaultPort;
 		_reuseProcess = reuseProcess;
-		_printStdOut = printStdOut;
-		_printStdErr = printStdErr;
 		TempDir = GetTemporaryDirectory();
 	}
 
+	/// <inheritdoc/>
+	public void Dispose()
+	{
+		Dispose(true);
+		GC.SuppressFinalize(this);
+	}
+
+	protected virtual void Dispose(bool disposing)
+	{
+		Stop();
+	}
+
 	protected int Port { get; private set; } = -1;
 
 	protected abstract string GetArguments();
@@ -57,62 +67,51 @@ public abstract class DatabaseRunner
 			throw new Exception("Failed finding process to re-use! See stdout for info.");
 		}
 
-		if (_proc != null)
+		if (_process != null)
 		{
 			return;
 		}
 
 		Port = GetAvailablePort();
 
-		Process p = new();
-		if (_printStdOut)
-		{
-			p.OutputDataReceived += (_, args) => Console.WriteLine("{0} stdout: {1}", _name, args.Data);
-		}
-
-		if (_printStdErr)
-		{
-			p.ErrorDataReceived += (_, args) => Console.WriteLine("{0} stderr: {1}", _name, args.Data);
-		}
-
-		p.StartInfo.FileName = GetBinaryPath();
-		p.StartInfo.WorkingDirectory = TempDir;
-		p.StartInfo.Arguments = GetArguments();
-		p.StartInfo.UseShellExecute = false;
-		p.StartInfo.CreateNoWindow = true;
-		p.StartInfo.RedirectStandardOutput = true;
-		p.StartInfo.RedirectStandardError = true;
-
-		if (!p.Start())
-		{
-			throw new Exception("Process start failed!");
-		}
-
-		p.BeginOutputReadLine();
-		p.BeginErrorReadLine();
-
+		_processGroup = new ManagedProcessGroup();
+		_process = new ManagedProcess(_processGroup, GetBinaryPath(), GetArguments(), TempDir, null, ProcessPriorityClass.Normal);
+		Task.Run(() => RelayOutputAsync(_process));
+		
 		// Try detect when main .NET process exits and kill the runner
 		AppDomain.CurrentDomain.ProcessExit += (sender, eventArgs) =>
 		{
 			Console.WriteLine("Main process exiting!");
 			Stop();
 		};
+	}
 
-		_proc = p;
+	static async Task RelayOutputAsync(ManagedProcess process)
+	{
+		for (; ; )
+		{
+			string? line = await process.ReadLineAsync();
+			if (line == null)
+			{
+				break;
+			}
+//			Console.WriteLine("{0} output: {1}", _name, line);
+		}
 	}
 
 	public void Stop()
 	{
-		if (_proc != null)
+		if (_process != null)
 		{
-			_proc.Kill(true);
-
-			// Waiting for exit blocks excessively even though the kill was sent. Anti-virus interfering?
-			// Process eventually shuts down but takes 2-3 min in Redis case. ReuseProcess flags circumvents this.
-			//Proc.WaitForExit();
-			_proc = null;
-			DeleteDirectory(TempDir);
+			_process.Dispose();
+			_process = null;
+		}
+		if (_processGroup != null)
+		{
+			_processGroup.Dispose();
+			_processGroup = null;
 		}
+		DeleteDirectory(TempDir);
 	}
 
 	public (string Host, int Port) GetListenAddress()
@@ -182,13 +181,18 @@ public abstract class DatabaseRunner
 
 	private static void DeleteDirectory(string path)
 	{
-		DirectoryInfo dir = new(path) { Attributes = FileAttributes.Normal };
-		foreach (FileSystemInfo info in dir.GetFileSystemInfos("*", SearchOption.AllDirectories))
+		DirectoryInfo dir = new DirectoryInfo(path);
+		if (dir.Exists)
 		{
-			info.Attributes = FileAttributes.Normal;
-		}
+			dir.Attributes = FileAttributes.Normal;
 
-		dir.Delete(true);
+			foreach (FileSystemInfo info in dir.GetFileSystemInfos("*", SearchOption.AllDirectories))
+			{
+				info.Attributes = FileAttributes.Normal;
+			}
+
+			dir.Delete(true);
+		}
 	}
 }
 
diff --git a/Engine/Source/Programs/Horde/Horde.Build.Tests/DatabaseRunnerTest.cs b/Engine/Source/Programs/Horde/Horde.Build.Tests/DatabaseRunnerTest.cs
--- a/Engine/Source/Programs/Horde/Horde.Build.Tests/DatabaseRunnerTest.cs
+++ b/Engine/Source/Programs/Horde/Horde.Build.Tests/DatabaseRunnerTest.cs
@@ -11,7 +11,7 @@ namespace Horde.Build.Tests
         [TestMethod]
         public void RunMongoDbTest()
         {
-	        MongoDbRunnerLocal runner = new MongoDbRunnerLocal();
+	        using MongoDbRunnerLocal runner = new MongoDbRunnerLocal();
 	        runner.Start();
             Thread.Sleep(100);
             runner.Stop();
@@ -20,7 +20,7 @@ namespace Horde.Build.Tests
         [TestMethod]
         public void RunRedisTest()
         {
-	        RedisRunner runner = new RedisRunner();
+	        using RedisRunner runner = new RedisRunner();
 	        runner.Start();
 	        Thread.Sleep(100);
 	        runner.Stop();
