Horde: Adding user preference and job tab configuration to show all preflights
#rnx

diff --git a/Engine/Source/Programs/Horde/Horde.Build/Streams/StreamConfig.cs b/Engine/Source/Programs/Horde/Horde.Build/Streams/StreamConfig.cs
--- a/Engine/Source/Programs/Horde/Horde.Build/Streams/StreamConfig.cs
+++ b/Engine/Source/Programs/Horde/Horde.Build/Streams/StreamConfig.cs
@@ -339,6 +339,11 @@ namespace Horde.Build.Streams
 		/// </summary>
 		public bool ShowNames { get; set; }
 
+		/// <summary>
+		/// Whether to show all user preflights 
+		/// </summary>
+		public bool? ShowPreflights { get; set; }
+
 		/// <summary>
 		/// Names of jobs to include on this page. If there is only one name specified, the name column does not need to be displayed.
 		/// </summary>
diff --git a/Engine/Source/Programs/Horde/HordeDashboard/src/backend/Api.ts b/Engine/Source/Programs/Horde/HordeDashboard/src/backend/Api.ts
--- a/Engine/Source/Programs/Horde/HordeDashboard/src/backend/Api.ts
+++ b/Engine/Source/Programs/Horde/HordeDashboard/src/backend/Api.ts
@@ -2399,9 +2399,11 @@ export type GetStreamTabResponse = {
 /**Describes a job page */
 export type GetJobsTabResponse = GetStreamTabResponse & {
 
-	/**  Whether to show names on the page */
+	/** Whether to show names on the page */
 	showNames: boolean;
 
+	/** Whether to show preflights */
+	showPreflights?: boolean;
 
 	/** List of templates to show on the page */
 	templates?: string[];
@@ -3008,7 +3010,8 @@ export enum DashboardPreference {
 	ColorRunning = "ColorRunning",
 	LocalCache = "LocalCache",
 	LeftAlignLog = "LeftAlignLog",
-	CompactViews = "CompactViews"
+	CompactViews = "CompactViews",
+	ShowPreflights = "ShowPreflights"
 }
 
 export type DashboardSettings = {
diff --git a/Engine/Source/Programs/Horde/HordeDashboard/src/backend/Dashboard.ts b/Engine/Source/Programs/Horde/HordeDashboard/src/backend/Dashboard.ts
--- a/Engine/Source/Programs/Horde/HordeDashboard/src/backend/Dashboard.ts
+++ b/Engine/Source/Programs/Horde/HordeDashboard/src/backend/Dashboard.ts
@@ -239,6 +239,14 @@ export class Dashboard {
         return this.preferences.get(DashboardPreference.LeftAlignLog) === 'true';
     }
 
+    setShowPreflights(value: boolean | undefined) {
+        this.setPreference(DashboardPreference.ShowPreflights, value ? "true" : "false");
+    }
+
+    get showPreflights(): boolean {
+        return this.preferences.get(DashboardPreference.ShowPreflights) === 'true';
+    }
+
     setCompactViews(value: boolean | undefined) {
         this.setPreference(DashboardPreference.CompactViews, value ? "true" : "false");
     }
diff --git a/Engine/Source/Programs/Horde/HordeDashboard/src/components/DashboardView.tsx b/Engine/Source/Programs/Horde/HordeDashboard/src/components/DashboardView.tsx
--- a/Engine/Source/Programs/Horde/HordeDashboard/src/components/DashboardView.tsx
+++ b/Engine/Source/Programs/Horde/HordeDashboard/src/components/DashboardView.tsx
@@ -361,6 +361,9 @@ const GeneralPanel: React.FC = observer(() => {
                         <Toggle label="24 Hour Clock" inlineLabel={true} defaultChecked={dashboard.display24HourClock} onChange={(ev, checked) => {
                            dashboard.setDisplay24HourClock(checked ? true : false);
                         }} />
+                        <Toggle label="Show All Preflights" inlineLabel={true} defaultChecked={dashboard.showPreflights} onChange={(ev, checked) => {
+                           dashboard.setShowPreflights(checked ? true : false);
+                        }} />
                         <Toggle label="Prefer Compact Views" inlineLabel={true} defaultChecked={dashboard.compactViews} onChange={(ev, checked) => {
                            dashboard.setCompactViews(checked ? true : false);
                         }} />
diff --git a/Engine/Source/Programs/Horde/HordeDashboard/src/components/JobView.tsx b/Engine/Source/Programs/Horde/HordeDashboard/src/components/JobView.tsx
--- a/Engine/Source/Programs/Horde/HordeDashboard/src/components/JobView.tsx
+++ b/Engine/Source/Programs/Horde/HordeDashboard/src/components/JobView.tsx
@@ -550,7 +550,7 @@ const JobList: React.FC<{ tab: string; filter: JobFilterSimple, controller: Call
       return null;
    };
 
-   const forcePreflights = tab?.toLowerCase() === "swarm" || tab?.toLowerCase() === "presubmit" || tab?.toLowerCase() === "services";
+   const forcePreflights = dashboard.showPreflights || !!jobTab?.showPreflights || tab?.toLowerCase() === "swarm" || tab?.toLowerCase() === "presubmit" || tab?.toLowerCase() === "services";
    let preflightStartedByUserId: string | undefined;
 
    if (!forcePreflights && !filter.showOthersPreflights) {
