Horde: Default preflight setting to dashboard preference
#rnx
#skipci

diff --git a/Engine/Source/Programs/Horde/HordeDashboard/src/components/StreamView.tsx b/Engine/Source/Programs/Horde/HordeDashboard/src/components/StreamView.tsx
--- a/Engine/Source/Programs/Horde/HordeDashboard/src/components/StreamView.tsx
+++ b/Engine/Source/Programs/Horde/HordeDashboard/src/components/StreamView.tsx
@@ -5,6 +5,7 @@ import { observer } from 'mobx-react-lite';
 import React, { useEffect, useState } from 'react';
 import { useHistory, useLocation, useParams } from 'react-router-dom';
 import { useBackend } from '../backend';
+import dashboard from '../backend/Dashboard';
 import { JobFilterSimple } from '../base/utilities/filter';
 import { useWindowSize } from '../base/utilities/hooks';
 import { modeColors, hordeClasses } from '../styles/Styles';
@@ -62,7 +63,7 @@ const StreamViewInner: React.FC = observer(() => {
    const location = useLocation();
    const query = useQuery();
 
-   const [filter, setFilter] = useState<JobFilterSimple>({ showOthersPreflights: false });
+   const [filter, setFilter] = useState<JobFilterSimple>({ showOthersPreflights: dashboard.showPreflights });
 
    const [shown, setShown] = useState(query.get("newbuild") ? true : false);
    const [findJobsShown, setFindJobsShown] = useState(false);
