Horde: Support inviting users to Slack channels using admin tokens.

diff --git a/Engine/Source/Programs/Horde/Horde.Build/Issues/WorkflowConfig.cs b/Engine/Source/Programs/Horde/Horde.Build/Issues/WorkflowConfig.cs
--- a/Engine/Source/Programs/Horde/Horde.Build/Issues/WorkflowConfig.cs
+++ b/Engine/Source/Programs/Horde/Horde.Build/Issues/WorkflowConfig.cs
@@ -90,6 +90,11 @@ namespace Horde.Build.Issues
 		/// </summary>
 		public bool AllowMentions { get; set; } = true;
 
+		/// <summary>
+		/// Uses the admin.conversations.invite API to invite users to the channel
+		/// </summary>
+		public bool InviteUsersAsAdmin { get; set; }
+
 		/// <summary>
 		/// Additional node annotations implicit in this workflow
 		/// </summary>
diff --git a/Engine/Source/Programs/Horde/Horde.Build/Notifications/Sinks/SlackNotificationSink.cs b/Engine/Source/Programs/Horde/Horde.Build/Notifications/Sinks/SlackNotificationSink.cs
--- a/Engine/Source/Programs/Horde/Horde.Build/Notifications/Sinks/SlackNotificationSink.cs
+++ b/Engine/Source/Programs/Horde/Horde.Build/Notifications/Sinks/SlackNotificationSink.cs
@@ -232,6 +232,9 @@ namespace Horde.Build.Notifications.Sinks
 		readonly HttpClient _httpClient;
 		readonly SlackClient _slackClient;
 
+		readonly HttpClient? _adminHttpClient;
+		readonly SlackClient? _adminSlackClient;
+
 		/// <summary>
 		/// Map of email address to Slack user ID.
 		/// </summary>
@@ -264,6 +267,13 @@ namespace Horde.Build.Notifications.Sinks
 			_httpClient.DefaultRequestHeaders.Add("Authorization", $"Bearer {_settings.SlackToken ?? ""}");
 			_slackClient = new SlackClient(_httpClient, _logger);
 
+			if (!String.IsNullOrEmpty(_settings.SlackAdminToken))
+			{
+				_adminHttpClient = new HttpClient();
+				_adminHttpClient.DefaultRequestHeaders.Add("Authorization", $"Bearer {_settings.SlackAdminToken ?? ""}");
+				_adminSlackClient = new SlackClient(_adminHttpClient, _logger);
+			}
+
 			_jsonSerializerOptions = new JsonSerializerOptions();
 			Startup.ConfigureJsonSerializer(_jsonSerializerOptions);
 
@@ -858,7 +868,7 @@ namespace Horde.Build.Notifications.Sinks
 			return null;
 		}
 
-		async Task InviteUsersAsync(string channel, IEnumerable<UserId> userIds)
+		async Task InviteUsersAsync(string channel, IEnumerable<UserId> userIds, bool inviteUsersAsAdmin)
 		{
 			List<string> slackUserIds = new List<string>();
 			foreach (UserId userId in userIds)
@@ -876,7 +886,14 @@ namespace Horde.Build.Notifications.Sinks
 
 			if (slackUserIds.Count > 0)
 			{
-				await _slackClient.InviteUsersAsync(channel, slackUserIds);
+				if (inviteUsersAsAdmin && _adminSlackClient != null)
+				{
+					await _adminSlackClient.AdminInviteUsersAsync(channel, slackUserIds);
+				}
+				else
+				{
+					await _slackClient.InviteUsersAsync(channel, slackUserIds);
+				}
 			}
 		}
 
@@ -1063,7 +1080,7 @@ namespace Horde.Build.Notifications.Sinks
 
 				if (_environment.IsProduction() && workflow.AllowMentions)
 				{
-					await InviteUsersAsync(state.Channel, inviteUserIds);
+					await InviteUsersAsync(state.Channel, inviteUserIds, workflow.InviteUsersAsAdmin);
 				}
 			}
 
diff --git a/Engine/Source/Programs/Horde/Horde.Build/ServerSettings.cs b/Engine/Source/Programs/Horde/Horde.Build/ServerSettings.cs
--- a/Engine/Source/Programs/Horde/Horde.Build/ServerSettings.cs
+++ b/Engine/Source/Programs/Horde/Horde.Build/ServerSettings.cs
@@ -435,15 +435,20 @@ namespace Horde.Build
 		public string? ScheduleTimeZone { get; set; }
 
 		/// <summary>
-		/// Token for interacting with Slack
+		/// Bot token for interacting with Slack (xoxb-*)
 		/// </summary>
 		public string? SlackToken { get; set; }
 
 		/// <summary>
-		/// Token for opening a socket to slack
+		/// Token for opening a socket to slack (xapp-*)
 		/// </summary>
 		public string? SlackSocketToken { get; set; }
 
+		/// <summary>
+		/// Admin user token for Slack (xoxp-*). This is only required when using the admin endpoints to invite users.
+		/// </summary>
+		public string? SlackAdminToken { get; set; }
+
 		/// <summary>
 		/// Filtered list of slack users to send notifications to. Should be Slack user ids, separated by commas.
 		/// </summary>
diff --git a/Engine/Source/Programs/Shared/EpicGames.Slack/SlackClient.cs b/Engine/Source/Programs/Shared/EpicGames.Slack/SlackClient.cs
--- a/Engine/Source/Programs/Shared/EpicGames.Slack/SlackClient.cs
+++ b/Engine/Source/Programs/Shared/EpicGames.Slack/SlackClient.cs
@@ -391,6 +391,22 @@ namespace EpicGames.Slack
 			await SendRequestAsync<SlackResponse>(ConversationsInviteUrl, message, ShouldLogError);
 		}
 
+		const string AdminConversationsInviteUrl = "https://slack.com/api/admin.conversations.invite";
+
+		/// <summary>
+		/// Invite a set of users to a channel using admin powers.
+		/// </summary>
+		/// <param name="channel">Channel identifier to invite the user to</param>
+		/// <param name="userIds">The user id</param>
+		public async Task AdminInviteUsersAsync(string channel, IEnumerable<string> userIds)
+		{
+			object message = new { channel_id = channel, user_ids = String.Join(",", userIds) };
+
+			static bool ShouldLogError(SlackResponse response) => !response.Ok && !String.Equals(response.Error, "already_in_channel", StringComparison.Ordinal);
+
+			await SendRequestAsync<SlackResponse>(AdminConversationsInviteUrl, message, ShouldLogError);
+		}
+
 		#endregion
 
 		#region Users
