Horde: Use admin endpoint to invite restricted users to channels if the InviteRestrictedUsers workflow flag is set, and regular invites fail.

diff --git a/Engine/Source/Programs/Horde/Horde.Build/Issues/WorkflowConfig.cs b/Engine/Source/Programs/Horde/Horde.Build/Issues/WorkflowConfig.cs
--- a/Engine/Source/Programs/Horde/Horde.Build/Issues/WorkflowConfig.cs
+++ b/Engine/Source/Programs/Horde/Horde.Build/Issues/WorkflowConfig.cs
@@ -93,7 +93,7 @@ namespace Horde.Build.Issues
 		/// <summary>
 		/// Uses the admin.conversations.invite API to invite users to the channel
 		/// </summary>
-		public bool InviteUsersAsAdmin { get; set; }
+		public bool InviteRestrictedUsers { get; set; }
 
 		/// <summary>
 		/// Additional node annotations implicit in this workflow
diff --git a/Engine/Source/Programs/Horde/Horde.Build/Notifications/Sinks/SlackNotificationSink.cs b/Engine/Source/Programs/Horde/Horde.Build/Notifications/Sinks/SlackNotificationSink.cs
--- a/Engine/Source/Programs/Horde/Horde.Build/Notifications/Sinks/SlackNotificationSink.cs
+++ b/Engine/Source/Programs/Horde/Horde.Build/Notifications/Sinks/SlackNotificationSink.cs
@@ -270,7 +270,7 @@ namespace Horde.Build.Notifications.Sinks
 			if (!String.IsNullOrEmpty(_settings.SlackAdminToken))
 			{
 				_adminHttpClient = new HttpClient();
-				_adminHttpClient.DefaultRequestHeaders.Add("Authorization", $"Bearer {_settings.SlackAdminToken ?? ""}");
+				_adminHttpClient.DefaultRequestHeaders.Add("Authorization", $"Bearer {_settings.SlackAdminToken}");
 				_adminSlackClient = new SlackClient(_adminHttpClient, _logger);
 			}
 
@@ -884,15 +884,26 @@ namespace Horde.Build.Notifications.Sinks
 				}
 			}
 
-			if (slackUserIds.Count > 0)
+			foreach (string slackUserId in slackUserIds)
 			{
-				if (inviteUsersAsAdmin && _adminSlackClient != null)
+				string? errorCode = await _slackClient.TryInviteUsersAsync(channel, new[] { slackUserId });
+				if (errorCode != null)
 				{
-					await _adminSlackClient.AdminInviteUsersAsync(channel, slackUserIds);
-				}
-				else
-				{
-					await _slackClient.InviteUsersAsync(channel, slackUserIds);
+					if (errorCode.Equals("user_is_restricted", StringComparison.Ordinal) && inviteUsersAsAdmin && _adminSlackClient != null)
+					{
+						try
+						{
+							await _adminSlackClient.AdminInviteUsersAsync(channel, new[] { slackUserId });
+						}
+						catch (SlackException ex)
+						{
+							_logger.LogWarning(ex, "Unable to invite user {UserId} to {Channel} (as admin): {Error}", slackUserId, channel, ex.Code);
+						}
+					}
+					else
+					{
+						_logger.LogWarning("Unable to invite user {UserId} to {Channel}: {Error}", slackUserId, channel, errorCode);
+					}
 				}
 			}
 		}
@@ -1080,7 +1091,7 @@ namespace Horde.Build.Notifications.Sinks
 
 				if (_environment.IsProduction() && workflow.AllowMentions)
 				{
-					await InviteUsersAsync(state.Channel, inviteUserIds, workflow.InviteUsersAsAdmin);
+					await InviteUsersAsync(state.Channel, inviteUserIds, workflow.InviteRestrictedUsers);
 				}
 			}
 
diff --git a/Engine/Source/Programs/Shared/EpicGames.Slack/SlackClient.cs b/Engine/Source/Programs/Shared/EpicGames.Slack/SlackClient.cs
--- a/Engine/Source/Programs/Shared/EpicGames.Slack/SlackClient.cs
+++ b/Engine/Source/Programs/Shared/EpicGames.Slack/SlackClient.cs
@@ -391,6 +391,23 @@ namespace EpicGames.Slack
 			await SendRequestAsync<SlackResponse>(ConversationsInviteUrl, message, ShouldLogError);
 		}
 
+		/// <summary>
+		/// Attempt to invite a set of users to a channel, returning the error code on failure
+		/// </summary>
+		/// <param name="channel">Channel identifier to invite the user to</param>
+		/// <param name="userIds">The user id</param>
+		public async Task<string?> TryInviteUsersAsync(string channel, IEnumerable<string> userIds)
+		{
+			InviteMessage message = new InviteMessage();
+			message.Channel = channel;
+			message.Users = String.Join(",", userIds);
+
+			static bool ShouldLogError(SlackResponse response) => false;
+
+			SlackResponse response = await SendRequestAsync<SlackResponse>(ConversationsInviteUrl, message, ShouldLogError);
+			return response.Ok ? null : response.Error;
+		}
+
 		const string AdminConversationsInviteUrl = "https://slack.com/api/admin.conversations.invite";
 
 		/// <summary>
