Horde: Support test meta data variations and expiring data sets
#rnx

diff --git a/Engine/Source/Programs/Horde/Horde.Build.Tests/TestDataTests.cs b/Engine/Source/Programs/Horde/Horde.Build.Tests/TestDataTests.cs
--- a/Engine/Source/Programs/Horde/Horde.Build.Tests/TestDataTests.cs
+++ b/Engine/Source/Programs/Horde/Horde.Build.Tests/TestDataTests.cs
@@ -18,6 +18,7 @@ using Microsoft.AspNetCore.Mvc;
 using System.Security.Claims;
 using Horde.Build.Logs;
 using System.Linq;
+using System.Threading;
 
 namespace Horde.Build.Tests
 {
@@ -160,6 +161,7 @@ namespace Horde.Build.Tests
 		[TestMethod]
 		public async Task SimpleReportTest()
 		{
+			await TestDataService.StartAsync(CancellationToken.None);
 			string[] streamIds = new string[] { _mainStreamId.ToString(), _releaseStreamId.ToString() };
 
 			IJob job = CreateJob(_mainStreamId, 105, "Test Build", _graph);
@@ -193,6 +195,8 @@ namespace Horde.Build.Tests
 			await TestDataCollection.AddAsync(job2, step2, data1.ToArray());
 			await TestDataCollection.AddAsync(job3, step3, data2.ToArray());
 
+			await TestDataService.TickForTestingAsync();
+
 			ActionResult<List<GetTestMetaResponse>> metaResult = await TestDataController.GetTestMetaAsync();
 			Assert.IsNotNull(metaResult);
 			Assert.IsNotNull(metaResult.Value);
@@ -237,6 +241,8 @@ namespace Horde.Build.Tests
 		[TestMethod]
 		public async Task SessionReportTest()
 		{
+			await TestDataService.StartAsync(CancellationToken.None);
+
 			string[] streamIds = new string[] { _mainStreamId.ToString(), _releaseStreamId.ToString() };
 
 			IJob job = CreateJob(_mainStreamId, 105, "Test Build", _graph);
@@ -276,6 +282,8 @@ namespace Horde.Build.Tests
 			await TestDataCollection.AddAsync(job2, step2, data.ToArray());
 			await TestDataCollection.AddAsync(job3, step3, data2.ToArray());
 
+			await TestDataService.TickForTestingAsync();
+
 			ActionResult<List<GetTestMetaResponse>> metaResult = await TestDataController.GetTestMetaAsync();
 			Assert.IsNotNull(metaResult);
 			Assert.IsNotNull(metaResult.Value);
@@ -290,6 +298,7 @@ namespace Horde.Build.Tests
 			Assert.AreEqual("Development", meta[0].Configurations[0]);
 			Assert.AreEqual("EngineTest", meta[0].ProjectName);
 			Assert.AreEqual("default", meta[0].RHI);
+			Assert.AreEqual("default", meta[0].Variation);
 
 			ActionResult<List<GetTestStreamResponse>> streamResult = await TestDataController.GetTestStreamsAsync(streamIds);
 			Assert.IsNotNull(streamResult);
diff --git a/Engine/Source/Programs/Horde/Horde.Build/Jobs/TestData/ITestData.cs b/Engine/Source/Programs/Horde/Horde.Build/Jobs/TestData/ITestData.cs
--- a/Engine/Source/Programs/Horde/Horde.Build/Jobs/TestData/ITestData.cs
+++ b/Engine/Source/Programs/Horde/Horde.Build/Jobs/TestData/ITestData.cs
@@ -74,6 +74,12 @@ namespace Horde.Build.Jobs.TestData
 		/// The rendering hardware interface used for the test
 		/// </summary>
 		string RHI { get; }
+
+		/// <summary>
+		/// The variation of the meta data, for example address sanitizing
+		/// </summary>
+		string Variation { get; }
+
 	}
 
 	/// <summary>
diff --git a/Engine/Source/Programs/Horde/Horde.Build/Jobs/TestData/ITestDataCollection.cs b/Engine/Source/Programs/Horde/Horde.Build/Jobs/TestData/ITestDataCollection.cs
--- a/Engine/Source/Programs/Horde/Horde.Build/Jobs/TestData/ITestDataCollection.cs
+++ b/Engine/Source/Programs/Horde/Horde.Build/Jobs/TestData/ITestDataCollection.cs
@@ -103,9 +103,10 @@ namespace Horde.Build.Jobs.TestData
 		/// <param name="configurations"></param>
 		/// <param name="buildTargets"></param>
 		/// <param name="rhi"></param>
+		/// <param name="variation"></param>
 		/// <param name="metaIds"></param>
 		/// <returns></returns>
-		Task<List<ITestMeta>> FindTestMeta(string[]? projectNames = null, string[]? platforms = null, string[]? configurations = null, string[]? buildTargets = null, string? rhi = null, TestMetaId[]? metaIds = null);
+		Task<List<ITestMeta>> FindTestMeta(string[]? projectNames = null, string[]? platforms = null, string[]? configurations = null, string[]? buildTargets = null, string? rhi = null, string? variation = null, TestMetaId[]? metaIds = null);
 
 		/// <summary>
 		/// Delete the test data
@@ -113,5 +114,22 @@ namespace Horde.Build.Jobs.TestData
 		/// <param name="id">Unique id of the test data</param>
 		/// <returns>Async task</returns>
 		Task DeleteAsync(ObjectId id);
+
+		/// <summary>
+		/// Prune test data from collection
+		/// </summary>
+		/// <returns></returns>
+		public Task PruneDataAsync();
+
+		/// <summary>
+		/// Run a tick on the collection
+		/// </summary>
+		/// <returns></returns>
+		public Task<bool> UpdateAsync(int retainMonths);
+
+		/// <summary>
+		/// Upgrades the collection
+		/// </summary>
+		public Task UpgradeAsync();
 	}
 }
diff --git a/Engine/Source/Programs/Horde/Horde.Build/Jobs/TestData/TestDataCollection.cs b/Engine/Source/Programs/Horde/Horde.Build/Jobs/TestData/TestDataCollection.cs
--- a/Engine/Source/Programs/Horde/Horde.Build/Jobs/TestData/TestDataCollection.cs
+++ b/Engine/Source/Programs/Horde/Horde.Build/Jobs/TestData/TestDataCollection.cs
@@ -34,8 +34,12 @@ namespace Horde.Build.Jobs.TestData
 	/// </summary>
 	public class TestDataCollection : ITestDataCollection
 	{
+		internal interface ITestExpire
+		{
+			DateTime? LastSeenUtc { get; }		
+		}
 
-		class TestMetaDocument : ITestMeta
+		class TestMetaDocument : ITestMeta, ITestExpire
 		{
 			[BsonRequired, BsonId]
 			public TestMetaId Id { get; set; }
@@ -58,6 +62,17 @@ namespace Horde.Build.Jobs.TestData
 			[BsonRequired, BsonElement("r")]
 			public string RHI { get; set; }
 
+			[BsonElement("v")]
+			public string Variation { get; set; }
+
+			[BsonIgnoreIfNull, BsonElement("s")]
+			public DateTime? LastSeenUtc { get; set; }
+
+			[BsonIgnoreIfNull, BsonElement("vn")]
+			public int? Version { get; set; }
+
+			public static int CurrentVersion = 1;
+
 			private TestMetaDocument()
 			{
 				Platforms = new List<string>();
@@ -65,10 +80,11 @@ namespace Horde.Build.Jobs.TestData
 				BuildTargets = new List<string>();
 
 				RHI = String.Empty;
+				Variation = String.Empty;
 				ProjectName = String.Empty;
 			}
 
-			public TestMetaDocument(List<string> platforms, List<string> configurations, List<string> buildTargets, string projectName, string? rhi)
+			public TestMetaDocument(List<string> platforms, List<string> configurations, List<string> buildTargets, string projectName, string? rhi, string? variation)
 			{
 				Id = TestMetaId.GenerateNewId();
 				Platforms = platforms;
@@ -76,14 +92,17 @@ namespace Horde.Build.Jobs.TestData
 				BuildTargets = buildTargets;
 				ProjectName = projectName;
 				RHI = rhi ?? "default";
+				Variation = variation ?? "default";
+				Version = CurrentVersion;
+				LastSeenUtc = DateTime.UtcNow;
 			}
 		}
 
-		class TestDocument : ITest
+		class TestDocument : ITest, ITestExpire
 		{
 			[BsonRequired, BsonId]
 			public TestId Id { get; set; }
-			
+
 			[BsonRequired, BsonElement("n")]
 			public string Name { get; set; }
 
@@ -97,6 +116,14 @@ namespace Horde.Build.Jobs.TestData
 			public List<TestMetaId> Metadata { get; set; }
 			IReadOnlyList<TestMetaId> ITest.Metadata => Metadata;
 
+			[BsonIgnoreIfNull, BsonElement("vn")]
+			public int? Version { get; set; }
+
+			[BsonIgnoreIfNull, BsonElement("s")]
+			public DateTime? LastSeenUtc { get; set; }
+
+			public static int CurrentVersion = 1;
+
 			private TestDocument()
 			{
 				Name = String.Empty;
@@ -110,14 +137,16 @@ namespace Horde.Build.Jobs.TestData
 				Metadata = meta;
 				SuiteName = suiteName;
 				DisplayName = displayName;
+				Version = CurrentVersion;
+				LastSeenUtc = DateTime.UtcNow;
 			}
 		}
 
-		class TestSuiteDocument : ITestSuite
+		class TestSuiteDocument : ITestSuite, ITestExpire
 		{
 			[BsonRequired, BsonId]
 			public TestSuiteId Id { get; set; }
-			
+
 			[BsonRequired, BsonElement("n")]
 			public string Name { get; set; }
 
@@ -129,6 +158,14 @@ namespace Horde.Build.Jobs.TestData
 			public List<TestId> Tests { get; set; }
 			IReadOnlyList<TestId> ITestSuite.Tests => Tests;
 
+			[BsonIgnoreIfNull, BsonElement("s")]
+			public DateTime? LastSeenUtc { get; set; }
+
+			[BsonIgnoreIfNull, BsonElement("vn")]
+			public int? Version { get; set; }
+
+			public static int CurrentVersion = 1;
+
 			private TestSuiteDocument()
 			{
 				Name = String.Empty;
@@ -142,6 +179,8 @@ namespace Horde.Build.Jobs.TestData
 				Name = name;
 				Metadata = metaIds;
 				Tests = tests;
+				Version = CurrentVersion;
+				LastSeenUtc = DateTime.UtcNow;
 			}
 		}
 
@@ -559,7 +598,7 @@ namespace Horde.Build.Jobs.TestData
 		}
 
 		/// <inheritdoc/>
-		public async Task<List<ITestMeta>> FindTestMeta(string[]? projectNames = null, string[]? platforms = null, string[]? configurations = null, string[]? buildTargets = null, string? rhi = null, TestMetaId[]? metaIds = null)
+		public async Task<List<ITestMeta>> FindTestMeta(string[]? projectNames = null, string[]? platforms = null, string[]? configurations = null, string[]? buildTargets = null, string? rhi = null, string? variation = null, TestMetaId[]? metaIds = null)
 		{
 			FilterDefinitionBuilder<TestMetaDocument> filterBuilder = Builders<TestMetaDocument>.Filter;
 			FilterDefinition<TestMetaDocument> filter = FilterDefinition<TestMetaDocument>.Empty;
@@ -624,11 +663,16 @@ namespace Horde.Build.Jobs.TestData
 				}
 			}
 
-			if (rhi != null)
+			if (rhi != null && rhi.Length > 0)
 			{
 				filter &= filterBuilder.Eq(x => x.RHI, rhi);
 			}
 
+			if (variation != null && variation.Length > 0)
+			{
+				filter &= filterBuilder.Eq(x => x.Variation, variation);
+			}
+
 			List<TestMetaDocument> result = await _testMeta.Find(filter).ToListAsync();
 
 			return result.ConvertAll<ITestMeta>(x => x);			
@@ -872,15 +916,20 @@ namespace Horde.Build.Jobs.TestData
 			metaData.TryGetValue("RHI", out rhi);
 			rhi = rhi ?? "default";
 
-			TestMetaDocument meta = new TestMetaDocument(platforms, configurations, buildTargets, projectName, rhi);
+			string? variation;
+			metaData.TryGetValue("Variation", out variation);
+			variation = variation ?? "default";
+
+			TestMetaDocument meta = new TestMetaDocument(platforms, configurations, buildTargets, projectName, rhi, variation);
 
 			FilterDefinitionBuilder<TestMetaDocument> filterBuilder = Builders<TestMetaDocument>.Filter;
 			FilterDefinition<TestMetaDocument> filter = filterBuilder.Eq(x => x.ProjectName, meta.ProjectName);
 			filter &= filterBuilder.Eq(x => x.RHI, meta.RHI);
+			filter &= filterBuilder.Eq(x => x.Variation, meta.Variation);
 			filter &= filterBuilder.Eq(x => x.Platforms, meta.Platforms);
 			filter &= filterBuilder.Eq(x => x.BuildTargets, meta.BuildTargets);
 			filter &= filterBuilder.Eq(x => x.Configurations, meta.Configurations);
-
+			
 			TestMetaDocument? result = await _testMeta.Find(filter).FirstOrDefaultAsync();
 
 			if (result == null)
@@ -888,7 +937,14 @@ namespace Horde.Build.Jobs.TestData
 				result = meta;
 				await _testMeta.InsertOneAsync(result);
 			}
-			
+			else
+			{
+				UpdateDefinitionBuilder<TestMetaDocument> updateBuilder = Builders<TestMetaDocument>.Update;
+				List<UpdateDefinition<TestMetaDocument>> updates = new List<UpdateDefinition<TestMetaDocument>>();
+				updates.Add(updateBuilder.Set(x => x.LastSeenUtc, DateTime.UtcNow));
+				await _testMeta.FindOneAndUpdateAsync(x => x.Id == result.Id, updateBuilder.Combine(updates));
+			}
+
 			return result;
 
 		}
@@ -923,17 +979,19 @@ namespace Horde.Build.Jobs.TestData
 			{
 				test = tests[0];
 
+				UpdateDefinitionBuilder<TestDocument> updateBuilder = Builders<TestDocument>.Update;
+				List<UpdateDefinition<TestDocument>> updates = new List<UpdateDefinition<TestDocument>>();
+
+				updates.Add(updateBuilder.Set(x => x.LastSeenUtc, DateTime.UtcNow));
+
 				if (!test.Metadata.Contains(metaData.Id))
 				{
 					test.Metadata.Add(metaData.Id);
 
-					FilterDefinitionBuilder<TestDocument> ebuilder = Builders<TestDocument>.Filter;
-					FilterDefinition<TestDocument> efilter = ebuilder.Eq(x => x.Id, test.Id);
-
-					UpdateDefinitionBuilder<TestDocument> updateBuilder = Builders<TestDocument>.Update;
-					UpdateDefinition<TestDocument> update = updateBuilder.Set(x => x.Metadata, test.Metadata);
-					await _tests.UpdateOneAsync(efilter, update);
+					updates.Add(updateBuilder.Set(x => x.Metadata, test.Metadata));
 				}
+
+				await _tests.FindOneAndUpdateAsync(x => x.Id == test.Id, updateBuilder.Combine(updates));
 			}
 
 			return test;
@@ -960,31 +1018,27 @@ namespace Horde.Build.Jobs.TestData
 			{
 				suite = suiteTests[0];
 
-				List<TestId> newIds = suite.Tests.Union(testIds).ToList();
+				UpdateDefinitionBuilder<TestSuiteDocument> updateBuilder = Builders<TestSuiteDocument>.Update;
+				List<UpdateDefinition<TestSuiteDocument>> updates = new List<UpdateDefinition<TestSuiteDocument>>();
 
-				if (newIds.Count != suite.Tests.Count || !suite.Metadata.Contains(metaData.Id))
-				{
-					List<UpdateDefinition<TestSuiteDocument>> updates = new List<UpdateDefinition<TestSuiteDocument>>();
-					FilterDefinitionBuilder<TestSuiteDocument> ebuilder = Builders<TestSuiteDocument>.Filter;
-					FilterDefinition<TestSuiteDocument> efilter = ebuilder.Eq(x => x.Id, suite.Id);
-
-					UpdateDefinitionBuilder<TestSuiteDocument> updateBuilder = Builders<TestSuiteDocument>.Update;
+				updates.Add(updateBuilder.Set(x => x.LastSeenUtc, DateTime.UtcNow));
 
-					if (newIds.Count != suite.Tests.Count)
-					{
-						suite.Tests = newIds;
-						updates.Add(updateBuilder.Set(x => x.Tests, suite.Tests));
-					}
+				List<TestId> newIds = suite.Tests.Union(testIds).ToList();
 
-					if (!suite.Metadata.Contains(metaData.Id))
-					{
-						suite.Metadata.Add(metaData.Id);
-						updates.Add(updateBuilder.Set(x => x.Metadata, suite.Metadata));
-					}
+				if (newIds.Count != suite.Tests.Count)
+				{
+					suite.Tests = newIds;
+					updates.Add(updateBuilder.Set(x => x.Tests, suite.Tests));
+				}
 
-					await _testSuites.UpdateOneAsync(efilter, updateBuilder.Combine(updates));
+				if (!suite.Metadata.Contains(metaData.Id))
+				{
+					suite.Metadata.Add(metaData.Id);
+					updates.Add(updateBuilder.Set(x => x.Metadata, suite.Metadata));
 				}
 
+				await _testSuites.FindOneAndUpdateAsync(x => x.Id == suite.Id, updateBuilder.Combine(updates));
+
 			}
 
 			return suite;
@@ -1048,6 +1102,319 @@ namespace Horde.Build.Jobs.TestData
 			}
 		}
 
+		/// <inheritdoc/>
+		public async Task PruneDataAsync()
+		{
+
+			HashSet<TestMetaId> metaIds = new HashSet<TestMetaId>();
+			foreach (TestMetaId metaId in await (await _testMeta.DistinctAsync(x => x.Id, Builders<TestMetaDocument>.Filter.Empty)).ToListAsync())
+			{
+				metaIds.Add(metaId);
+			}
+
+			_logger.LogInformation("Test data pruning to {MetaIdCount} meta ids", metaIds.Count);
+
+			// Update tests
+			HashSet<TestId> testIds = new HashSet<TestId>();
+			HashSet<TestId> testDeletes = new HashSet<TestId>();
+			{
+				List<TestDocument> tests = await _tests.Find(Builders<TestDocument>.Filter.Empty).ToListAsync();
+				foreach (TestDocument test in tests)
+				{
+					List<TestMetaId> metas = test.Metadata.Where(x => metaIds.Contains(x)).ToList();
+
+					if (metas.Count == 0)
+					{
+						testDeletes.Add(test.Id);
+						continue;
+					}
+					else if (metas.Count != test.Metadata.Count)
+					{
+						UpdateDefinitionBuilder<TestDocument> updateBuilder = Builders<TestDocument>.Update;
+						List<UpdateDefinition<TestDocument>> updates = new List<UpdateDefinition<TestDocument>>();
+						updates.Add(updateBuilder.Set(x => x.Metadata, metas));
+						await _tests.FindOneAndUpdateAsync(x => x.Id == test.Id, updateBuilder.Combine(updates));
+					}
+
+					testIds.Add(test.Id);
+				}
+			}
+
+			// Update suites
+			HashSet<TestSuiteId> suiteIds = new HashSet<TestSuiteId>();
+			HashSet<TestSuiteId> suiteDeletes = new HashSet<TestSuiteId>();
+			{
+				List<TestSuiteDocument> suites = await _testSuites.Find(Builders<TestSuiteDocument>.Filter.Empty).ToListAsync();
+				foreach (TestSuiteDocument suite in suites)
+				{
+					List<TestMetaId> metas = suite.Metadata.Where(x => metaIds.Contains(x)).ToList();
+					List<TestId> tests = suite.Tests.Where(x => testIds.Contains(x)).ToList();
+
+					if (metas.Count == 0 || tests.Count == 0)
+					{
+						suiteDeletes.Add(suite.Id);
+						continue;
+					}
+					else if (metas.Count != suite.Metadata.Count || tests.Count != suite.Tests.Count)
+					{
+						UpdateDefinitionBuilder<TestSuiteDocument> updateBuilder = Builders<TestSuiteDocument>.Update;
+						List<UpdateDefinition<TestSuiteDocument>> updates = new List<UpdateDefinition<TestSuiteDocument>>();
+						if (metas.Count != suite.Metadata.Count)
+						{
+							updates.Add(updateBuilder.Set(x => x.Metadata, metas));
+						}
+
+						if (tests.Count != suite.Tests.Count)
+						{
+							updates.Add(updateBuilder.Set(x => x.Tests, tests));
+						}
+
+						await _testSuites.FindOneAndUpdateAsync(x => x.Id == suite.Id, updateBuilder.Combine(updates));
+					}
+
+					suiteIds.Add(suite.Id);
+				}
+			}
+
+			if (testDeletes.Count > 0)
+			{
+				FilterDefinitionBuilder<TestDocument> filterBuilder = Builders<TestDocument>.Filter;
+				FilterDefinition<TestDocument> filter = filterBuilder.Empty;
+				filter &= filterBuilder.In(x => x.Id, testDeletes);
+				await _tests.DeleteManyAsync(filter);
+			}
+
+			if (suiteDeletes.Count > 0)
+			{
+				FilterDefinitionBuilder<TestSuiteDocument> filterBuilder = Builders<TestSuiteDocument>.Filter;
+				FilterDefinition<TestSuiteDocument> filter = filterBuilder.Empty;
+				filter &= filterBuilder.In(x => x.Id, suiteDeletes);
+				await _testSuites.DeleteManyAsync(filter);
+			}
+
+			if (testDeletes.Count > 0 || suiteDeletes.Count > 0)
+			{
+				_logger.LogInformation("Pruned {TestDeleteCount} tests and {TestSuiteCOunt} suites", testDeletes.Count, suiteDeletes.Count);
+			}
+
+			int testStreamsDeleted = 0;
+
+			// Update test streams
+			{
+				List<TestStreamDocument> testStreams = await _testStreams.Find(Builders<TestStreamDocument>.Filter.Empty).ToListAsync();
+				foreach (TestStreamDocument stream in testStreams)
+				{
+					List<TestId> tests = stream.Tests.Where(x => testIds.Contains(x)).ToList();
+					List<TestSuiteId> suites = stream.TestSuites.Where(x => suiteIds.Contains(x)).ToList();
+
+					if (tests.Count == 0 && suites.Count == 0)
+					{
+						await _testStreams.DeleteOneAsync(x => x.Id == stream.Id);
+						testStreamsDeleted++;
+					}
+					else if (tests.Count != stream.Tests.Count || suites.Count != stream.TestSuites.Count)
+					{
+						UpdateDefinitionBuilder<TestStreamDocument> updateBuilder = Builders<TestStreamDocument>.Update;
+						List<UpdateDefinition<TestStreamDocument>> updates = new List<UpdateDefinition<TestStreamDocument>>();
+
+						if (tests.Count != stream.Tests.Count)
+						{
+							updates.Add(updateBuilder.Set(x => x.Tests, tests));
+						}
+
+						if (suites.Count != stream.TestSuites.Count)
+						{
+							updates.Add(updateBuilder.Set(x => x.TestSuites, suites));
+						}
+
+						await _testStreams.FindOneAndUpdateAsync(x => x.Id == stream.Id, updateBuilder.Combine(updates));
+					}
+				}
+			}
+
+			if (testStreamsDeleted > 0)
+			{
+				_logger.LogInformation("Pruned {TestStreamDeleteCount} test streams", testStreamsDeleted);
+			}
+
+		}
+
+		static async Task<long> ExpireCollection<DocType>(IMongoCollection<DocType> collection, DateTime expireTime) where DocType : ITestExpire
+		{
+			FilterDefinitionBuilder<DocType> filterBuilder = Builders<DocType>.Filter;
+			FilterDefinition<DocType> filter = filterBuilder.Empty;
+			filter &= filterBuilder.Lte(x => x.LastSeenUtc!, expireTime);
+			DeleteResult result = await collection.DeleteManyAsync(filter);
+
+			return result.DeletedCount;
+
+		}
+
+		async Task ExpireTestData(int retainMonths)
+		{
+			if (retainMonths <= 0)
+			{
+				return;
+			}
+
+			DateTime expireTime = DateTime.Now.AddMonths(-retainMonths);
+
+			// expire and prune root test data
+			long testsExpired = await ExpireCollection(_tests, expireTime);
+			long suitesExpired = await ExpireCollection(_testSuites, expireTime);
+			long metaExpired = await ExpireCollection(_testMeta, expireTime);
+
+			// Check whether we need to prune
+			if (testsExpired > 0 || suitesExpired > 0 || metaExpired > 0)
+			{
+				_logger.LogInformation("Expired {TestsExpired} tests, {SuitesExpired} suites, {MetaExpired} meta", testsExpired, suitesExpired, metaExpired);
+				await PruneDataAsync();
+			}
+
+			// expire test data documents
+			{
+				ObjectId dataExpireTime = ObjectId.GenerateNewId(expireTime);
+
+				FilterDefinition<TestDataDocument> filter = FilterDefinition<TestDataDocument>.Empty;
+				FilterDefinitionBuilder<TestDataDocument> filterBuilder = Builders<TestDataDocument>.Filter;
+				filter &= filterBuilder.Lte(x => x.Id!, dataExpireTime);
+				DeleteResult result = await _testDataDocuments.DeleteManyAsync(filter);
+				if (result.DeletedCount > 0)
+				{
+					_logger.LogInformation("Expired {NumDeleted} test data documents", result.DeletedCount);
+				}
+			}
+
+			TestRefId refExpireTime = TestRefId.GenerateNewId(expireTime);
+
+			// expire test refs
+			{
+				FilterDefinition<TestDataRefDocument> filter = FilterDefinition<TestDataRefDocument>.Empty;
+				FilterDefinitionBuilder<TestDataRefDocument> filterBuilder = Builders<TestDataRefDocument>.Filter;
+				filter &= filterBuilder.Lte(x => x.Id!, refExpireTime);
+				DeleteResult result = await _testRefs.DeleteManyAsync(filter);
+				if (result.DeletedCount > 0)
+				{
+					_logger.LogInformation("Expired {NumDeleted} test ref documents", result.DeletedCount);
+				}
+			}
+
+			// expire test details
+			{
+				FilterDefinition<TestDataDetailsDocument> filter = FilterDefinition<TestDataDetailsDocument>.Empty;
+				FilterDefinitionBuilder<TestDataDetailsDocument> filterBuilder = Builders<TestDataDetailsDocument>.Filter;
+				filter &= filterBuilder.Lte(x => x.Id!, refExpireTime);
+				DeleteResult result = await _testDetails.DeleteManyAsync(filter);
+				if (result.DeletedCount > 0)
+				{
+					_logger.LogInformation("Expired {NumDeleted} test detail documents", result.DeletedCount);
+				}
+
+			}
+
+		}
+
+		/// <inheritdoc/>
+		public async Task<bool> UpdateAsync(int retainMonths)
+		{
+			await ExpireTestData(retainMonths);
+
+			return true;		
+		}
+
+		/// <inheritdoc/>
+		public async Task UpgradeAsync()
+		{
+			// Upgrade test meta data
+			{
+				FilterDefinition<TestMetaDocument> filter = Builders<TestMetaDocument>.Filter.Empty;
+				filter &= Builders<TestMetaDocument>.Filter.Where(x => x.Version == null || x.Version < TestMetaDocument.CurrentVersion);
+
+				List<TestMetaDocument> results = await _testMeta.Find(filter).ToListAsync();
+
+				if (results.Count > 0)
+				{
+					_logger.LogInformation("Found {Count} test meta documents to upgrade", results.Count);
+				}
+
+				foreach (TestMetaDocument meta in results)
+				{
+					// unversioned => 1
+					if (meta.Version == null && TestMetaDocument.CurrentVersion == 1)
+					{
+						UpdateDefinitionBuilder<TestMetaDocument> updateBuilder = Builders<TestMetaDocument>.Update;
+						List<UpdateDefinition<TestMetaDocument>> updates = new List<UpdateDefinition<TestMetaDocument>>();
+
+						updates.Add(updateBuilder.Set(x => x.Version, TestMetaDocument.CurrentVersion));
+						updates.Add(updateBuilder.Set(x => x.LastSeenUtc, DateTime.UtcNow));
+
+						if (meta.Variation == null || meta.Variation.Length == 0)
+						{
+							updates.Add(updateBuilder.Set(x => x.Variation, "default"));
+						}
+
+						await _testMeta.FindOneAndUpdateAsync(x => x.Id == meta.Id, updateBuilder.Combine(updates));
+					}
+				}
+			}
+
+			// Upgrade tests
+			{
+				FilterDefinition<TestDocument> filter = Builders<TestDocument>.Filter.Empty;
+				filter &= Builders<TestDocument>.Filter.Where(x => x.Version == null || x.Version < TestDocument.CurrentVersion);
+
+				List<TestDocument> results = await _tests.Find(filter).ToListAsync();
+
+				if (results.Count > 0)
+				{
+					_logger.LogInformation("Found {Count} test documents to upgrade", results.Count);
+				}
+
+				foreach (TestDocument test in results)
+				{
+					// unversioned => 1
+					if (test.Version == null && TestDocument.CurrentVersion == 1)
+					{
+						UpdateDefinitionBuilder<TestDocument> updateBuilder = Builders<TestDocument>.Update;
+						List<UpdateDefinition<TestDocument>> updates = new List<UpdateDefinition<TestDocument>>();
+
+						updates.Add(updateBuilder.Set(x => x.Version, TestDocument.CurrentVersion));
+						updates.Add(updateBuilder.Set(x => x.LastSeenUtc, DateTime.UtcNow));
+
+						await _tests.FindOneAndUpdateAsync(x => x.Id == test.Id, updateBuilder.Combine(updates));
+					}
+				}
+			}
+
+			// Upgrade test suites
+			{
+				FilterDefinition<TestSuiteDocument> filter = Builders<TestSuiteDocument>.Filter.Empty;
+				filter &= Builders<TestSuiteDocument>.Filter.Where(x => x.Version == null || x.Version < TestSuiteDocument.CurrentVersion);
+
+				List<TestSuiteDocument> results = await _testSuites.Find(filter).ToListAsync();
+
+				if (results.Count > 0)
+				{
+					_logger.LogInformation("Found {Count} test suite documents to upgrade", results.Count);
+				}
+
+				foreach (TestSuiteDocument suite in results)
+				{
+					// unversioned => 1
+					if (suite.Version == null && TestSuiteDocument.CurrentVersion == 1)
+					{
+						UpdateDefinitionBuilder<TestSuiteDocument> updateBuilder = Builders<TestSuiteDocument>.Update;
+						List<UpdateDefinition<TestSuiteDocument>> updates = new List<UpdateDefinition<TestSuiteDocument>>();
+
+						updates.Add(updateBuilder.Set(x => x.Version, TestSuiteDocument.CurrentVersion));
+						updates.Add(updateBuilder.Set(x => x.LastSeenUtc, DateTime.UtcNow));
+
+						await _testSuites.FindOneAndUpdateAsync(x => x.Id == suite.Id, updateBuilder.Combine(updates));
+					}
+				}
+			}
+		}
+
 		// --- (Minimal) Parsing of Automation framework tests, geared towards generating indexes for aggregate queries
 
 		class AutomatedTestSessionData
diff --git a/Engine/Source/Programs/Horde/Horde.Build/Jobs/TestData/TestDataMessages.cs b/Engine/Source/Programs/Horde/Horde.Build/Jobs/TestData/TestDataMessages.cs
--- a/Engine/Source/Programs/Horde/Horde.Build/Jobs/TestData/TestDataMessages.cs
+++ b/Engine/Source/Programs/Horde/Horde.Build/Jobs/TestData/TestDataMessages.cs
@@ -159,6 +159,11 @@ namespace Horde.Build.Jobs.TestData
 		/// </summary>
 		public string RHI { get; set; }
 
+		/// <summary>
+		/// The varation of the test meta data, for example address sanitizing
+		/// </summary>
+		public string Variation { get; set; }
+
 		internal GetTestMetaResponse(ITestMeta meta)
 		{
 			Id = meta.Id.ToString();
@@ -166,7 +171,8 @@ namespace Horde.Build.Jobs.TestData
 			Configurations = meta.Configurations.Select(p => p).ToList();
 			BuildTargets = meta.BuildTargets.Select(p => p).ToList();
 			ProjectName = meta.ProjectName;
-			RHI = meta.RHI;	
+			RHI = meta.RHI;
+			Variation = meta.Variation;
 		}
 	}
 
diff --git a/Engine/Source/Programs/Horde/Horde.Build/Jobs/TestData/TestDataService.cs b/Engine/Source/Programs/Horde/Horde.Build/Jobs/TestData/TestDataService.cs
--- a/Engine/Source/Programs/Horde/Horde.Build/Jobs/TestData/TestDataService.cs
+++ b/Engine/Source/Programs/Horde/Horde.Build/Jobs/TestData/TestDataService.cs
@@ -2,10 +2,15 @@
 
 using System;
 using System.Collections.Generic;
+using System.Threading;
 using System.Threading.Tasks;
 using EpicGames.Core;
 using Horde.Build.Streams;
 using Horde.Build.Utilities;
+using HordeCommon;
+using Microsoft.Extensions.Hosting;
+using Microsoft.Extensions.Logging;
+using Microsoft.Extensions.Options;
 
 namespace Horde.Build.Jobs.TestData
 {
@@ -19,24 +24,75 @@ namespace Horde.Build.Jobs.TestData
 	/// <summary>
 	/// Device management service
 	/// </summary>
-	public sealed class TestDataService : IDisposable
+	public sealed class TestDataService : IHostedService, IDisposable
 	{
 		
-		readonly ITestDataCollection _testData;		
+		readonly ITestDataCollection _testData;
+		readonly IOptionsMonitor<ServerSettings> _settings;
+		readonly ITicker _ticker;
+		readonly ILogger<TestDataService> _logger;
 
 		/// <summary>
 		/// Device service constructor
 		/// </summary>
-		public TestDataService(ITestDataCollection testData)
+		public TestDataService(ITestDataCollection testData, IOptionsMonitor<ServerSettings> settings, IClock clock, ILogger<TestDataService> logger)
 		{
-			_testData = testData;		
+			_testData = testData;
+			_settings = settings;
+			_logger = logger;
+			_ticker = clock.AddSharedTicker<TestDataService>(TimeSpan.FromMinutes(10.0), TickAsync, logger);
+		}
+
+		/// <inheritdoc/>
+		public async Task StartAsync(CancellationToken cancellationToken)
+		{
+			await _ticker.StartAsync();
+
+			try
+			{
+				await _testData.UpgradeAsync();
+			}
+			catch (Exception ex)
+			{
+				_logger.LogError(ex, "Exception while upgrading test data collection: {Message}", ex.Message);
+			}
+		}
+
+		/// <inheritdoc/>
+		public async Task StopAsync(CancellationToken cancellationToken)
+		{
+			await _ticker.StopAsync();
 		}
 
 		/// <inheritdoc/>
 		public void Dispose()
-		{		
+		{
+			_ticker.Dispose();
 		}
 
+		/// <summary>
+		/// Ticks service
+		/// </summary>
+		async ValueTask TickAsync(CancellationToken stoppingToken)
+		{
+			if (!stoppingToken.IsCancellationRequested)
+			{
+				try
+				{
+					await _testData.UpdateAsync(_settings.CurrentValue.TestDataRetainMonths);					
+				}
+				catch (Exception ex)
+				{
+					_logger.LogError(ex, "Exception while ticking test data collection: {Message}", ex.Message);
+				}
+
+			}
+		}
+
+		internal async Task TickForTestingAsync()
+		{
+			await TickAsync(CancellationToken.None);			
+		}
 
 		/// <summary>
 		/// Find test streams
@@ -77,11 +133,12 @@ namespace Horde.Build.Jobs.TestData
 		/// <param name="configurations"></param>
 		/// <param name="buildTargets"></param>
 		/// <param name="rhi"></param>
+		/// <param name="variation"></param>
 		/// <param name="metaIds"></param>
 		/// <returns></returns>
-		public async Task<List<ITestMeta>> FindTestMeta(string[]? projectNames = null, string[]? platforms = null, string[]? configurations = null, string[]? buildTargets = null, string? rhi = null, TestMetaId[]? metaIds = null)
+		public async Task<List<ITestMeta>> FindTestMeta(string[]? projectNames = null, string[]? platforms = null, string[]? configurations = null, string[]? buildTargets = null, string? rhi = null, string? variation = null, TestMetaId[]? metaIds = null)
 		{
-			return await _testData.FindTestMeta(projectNames, platforms, configurations, buildTargets, rhi, metaIds);
+			return await _testData.FindTestMeta(projectNames, platforms, configurations, buildTargets, rhi, variation, metaIds);
 		}
 
 		/// <summary>
diff --git a/Engine/Source/Programs/Horde/Horde.Build/ServerSettings.cs b/Engine/Source/Programs/Horde/Horde.Build/ServerSettings.cs
--- a/Engine/Source/Programs/Horde/Horde.Build/ServerSettings.cs
+++ b/Engine/Source/Programs/Horde/Horde.Build/ServerSettings.cs
@@ -356,6 +356,11 @@ namespace Horde.Build
 		/// </summary>
 		public bool EnableScheduleInTestData { get; set; }
 
+		/// <summary>
+		/// The number of months to retain test data
+		/// </summary>
+		public int TestDataRetainMonths { get; set; } = 6;
+
 		/// <summary>
 		/// Interval between rebuilding the schedule queue with a DB query.
 		/// </summary>
diff --git a/Engine/Source/Programs/Horde/Horde.Build/Startup.cs b/Engine/Source/Programs/Horde/Horde.Build/Startup.cs
--- a/Engine/Source/Programs/Horde/Horde.Build/Startup.cs
+++ b/Engine/Source/Programs/Horde/Horde.Build/Startup.cs
@@ -643,7 +643,8 @@ namespace Horde.Build
 				}
 				services.AddHostedService(provider => provider.GetRequiredService<ConfigUpdateService>());
 				services.AddHostedService<TelemetryService>();
-				services.AddHostedService(provider => provider.GetRequiredService<DeviceService>());				
+				services.AddHostedService(provider => provider.GetRequiredService<DeviceService>());
+				services.AddHostedService(provider => provider.GetRequiredService<TestDataService>());
 			}
 
 			services.AddHostedService(provider => provider.GetRequiredService<IExternalIssueService>());
diff --git a/Engine/Source/Programs/Horde/HordeDashboard/src/backend/Backend.ts b/Engine/Source/Programs/Horde/HordeDashboard/src/backend/Backend.ts
--- a/Engine/Source/Programs/Horde/HordeDashboard/src/backend/Backend.ts
+++ b/Engine/Source/Programs/Horde/HordeDashboard/src/backend/Backend.ts
@@ -1089,10 +1089,10 @@ export class Backend {
         });
     }
 
-    getTestMetadata(auttomationProjects?: string[], platforms?: string[], targets?: string[], configurations?: string[]): Promise<GetTestMetaResponse[]> {
+    getTestMetadata(automationProjects?: string[], platforms?: string[], targets?: string[], configurations?: string[]): Promise<GetTestMetaResponse[]> {
         return new Promise<GetTestMetaResponse[]>((resolve, reject) => {
             this.backend.get(`/api/v2/testdata/metadata`, {
-                params: { project: auttomationProjects, platform: platforms, target: targets, configuration: configurations }
+                params: { project: automationProjects, platform: platforms, target: targets, configuration: configurations }
             }).then((value) => {
                 resolve(value.data as GetTestMetaResponse[]);
             }).catch((reason) => {
