Horde: Suppress systemic issue creation after non-systemic error events to prevent superfluous issue creation
#rnx

diff --git a/Engine/Source/Programs/Horde/Horde.Build.Tests/IssueServiceTests.cs b/Engine/Source/Programs/Horde/Horde.Build.Tests/IssueServiceTests.cs
--- a/Engine/Source/Programs/Horde/Horde.Build.Tests/IssueServiceTests.cs
+++ b/Engine/Source/Programs/Horde/Horde.Build.Tests/IssueServiceTests.cs
@@ -2283,6 +2283,60 @@ namespace Horde.Build.Tests
 
 		}
 
+		[TestMethod]
+		public async Task SystemicIssuesTest()
+		{
+			// #1
+			// Scenario: Job step fails with systemic XGE error
+			// Expected: A systemic issue is created
+			{
+				string[] lines =
+				{
+					@"BUILD FAILED: Command failed (Result:1): C:\Program Files (x86)\Incredibuild\xgConsole.exe ""d:\build\++UE5\Sync\Engine\Programs\AutomationTool\Saved\Logs\UAT_XGE.xml"" /Rebuild /NoLogo /ShowAgent /ShowTime /no_watchdog_thread. See logfile for details: 'xgConsole-2023.01.03-23.39.48.txt'"
+				};
+
+				IJob job = CreateJob(_mainStreamId, 105, "Test Build", _graph);
+				await ParseEventsAsync(job, 0, 0, lines);
+				await UpdateCompleteStep(job, 0, 0, JobStepOutcome.Failure);
+
+				List<IIssue> issues = await IssueCollection.FindIssuesAsync();
+				Assert.AreEqual(1, issues.Count);
+
+				IIssue issue = issues[0];
+
+				Assert.AreEqual(IssueSeverity.Error, issue.Severity);
+				Assert.AreEqual(1, issue.Fingerprints.Count);
+				Assert.AreEqual("Systemic", issue.Fingerprints[0].Type);
+			}
+
+			// #2
+			// Scenario: Job step fails at CL 120 with a linker issue, additionally there is a systemic xgConsole.exe error
+			// Expected: Creates a linker issue with severity error due to fatal warnings, does not create a systemic error
+			{
+				string[] lines =
+				{
+					@"ld: warning: direct access in function 'void Eigen::internal::evaluateProductBlockingSizesHeuristic<Eigen::half, Eigen::half, 1, long>(long&, long&, long&, long)' from file '../../EngineTest/Intermediate/Build/Mac/x86_64/EngineTest/Development/ORT/inverse.cc.o' to global weak symbol 'guard variable for Eigen::internal::manage_caching_sizes(Eigen::Action, long*, long*, long*)::m_cacheSizes' from file '../../EngineTest/Intermediate/Build/Mac/x86_64/EngineTest/Development/DynamicMesh/Module.DynamicMesh.4_of_5.cpp.o' means the weak symbol cannot be overridden at runtime. This was likely caused by different translation units being compiled with different visibility settings.",
+					@"ld: fatal warning(s) induced error (-fatal_warnings)",
+					@"clang: error: linker command failed with exit code 1 (use -v to see invocation)",
+					@"BUILD FAILED: Command failed (Result:1): C:\Program Files (x86)\Incredibuild\xgConsole.exe ""d:\build\++UE5\Sync\Engine\Programs\AutomationTool\Saved\Logs\UAT_XGE.xml"" /Rebuild /NoLogo /ShowAgent /ShowTime /no_watchdog_thread. See logfile for details: 'xgConsole-2023.01.03-23.39.48.txt'"
+				};
+
+				IJob job = CreateJob(_mainStreamId, 120, "Test Build", _graph);
+				await ParseEventsAsync(job, 0, 0, lines);
+				await UpdateCompleteStep(job, 0, 0, JobStepOutcome.Failure);
+
+				List<IIssue> issues = await IssueCollection.FindIssuesAsync();
+				Assert.AreEqual(1, issues.Count);
+
+				IIssue issue = issues[0];
+
+				Assert.AreEqual(IssueSeverity.Error, issue.Severity);
+				Assert.AreEqual(1, issue.Fingerprints.Count);
+				Assert.AreEqual("Hashed", issue.Fingerprints[0].Type);
+			}
+		}
+
+
 
 		private async Task ParseAsync(LogId logId, string[] lines)
 		{
diff --git a/Engine/Source/Programs/Horde/Horde.Build/Issues/Handlers/SystemicIssueHandler.cs b/Engine/Source/Programs/Horde/Horde.Build/Issues/Handlers/SystemicIssueHandler.cs
--- a/Engine/Source/Programs/Horde/Horde.Build/Issues/Handlers/SystemicIssueHandler.cs
+++ b/Engine/Source/Programs/Horde/Horde.Build/Issues/Handlers/SystemicIssueHandler.cs
@@ -6,6 +6,7 @@ using EpicGames.Core;
 using Horde.Build.Jobs;
 using Horde.Build.Jobs.Graphs;
 using Horde.Build.Logs;
+using HordeCommon;
 using Microsoft.Extensions.Logging;
 using MongoDB.Driver;
 
@@ -22,7 +23,7 @@ namespace Horde.Build.Issues.Handlers
 		public override string Type => "Systemic";
 
 		/// <inheritdoc/>
-		public override int Priority => 10;
+		public override int Priority => 255;
 
 		/// <summary>
 		///  Known systemic errors
@@ -49,12 +50,34 @@ namespace Horde.Build.Issues.Handlers
 		public override void TagEvents(IJob job, INode node, IReadOnlyNodeAnnotations annotations, IReadOnlyList<IssueEvent> stepEvents)
 		{
 			NewIssueFingerprint? fingerprint = null;
+			bool nonSystemicError = false;
 			foreach (IssueEvent stepEvent in stepEvents)
 			{
-				if (stepEvent.EventId != null && MatchEvent(stepEvent.EventData))
+				if (stepEvent.EventId == null)
+				{					
+					continue;
+				}
+
+				if (MatchEvent(stepEvent.EventData))
 				{
-					fingerprint ??= new NewIssueFingerprint(Type, new[] { $"step:{job.StreamId}:{job.TemplateId}:{node.Name}" }, null, new[] { $"{NodeNameKey}={node.Name}" });
-					stepEvent.Fingerprint = fingerprint;
+					if (nonSystemicError)
+					{						
+						stepEvent.Ignored = true;
+					}
+					else
+					{
+						fingerprint ??= new NewIssueFingerprint(Type, new[] { $"step:{job.StreamId}:{job.TemplateId}:{node.Name}" }, null, new[] { $"{NodeNameKey}={node.Name}" });
+						stepEvent.Fingerprint = fingerprint;
+					}
+				}
+				else
+				{					
+					if (stepEvent.Severity == EventSeverity.Error)
+					{
+						// We've seen a non-systemic error event, so ignore this systemic event to prevent superfluous issues from being created
+						nonSystemicError = true;
+					}
+					
 				}
 			}
 		}
