Horde: Fix systemic issue handler priority
#rnx

diff --git a/Engine/Source/Programs/Horde/Horde.Build/Issues/Handlers/SystemicIssueHandler.cs b/Engine/Source/Programs/Horde/Horde.Build/Issues/Handlers/SystemicIssueHandler.cs
--- a/Engine/Source/Programs/Horde/Horde.Build/Issues/Handlers/SystemicIssueHandler.cs
+++ b/Engine/Source/Programs/Horde/Horde.Build/Issues/Handlers/SystemicIssueHandler.cs
@@ -23,7 +23,7 @@ namespace Horde.Build.Issues.Handlers
 		public override string Type => "Systemic";
 
 		/// <inheritdoc/>
-		public override int Priority => 255;
+		public override int Priority => 10;
 
 		/// <summary>
 		///  Known systemic errors
