Horde: Factoring meta data variations into UX
#skipci
#rnx

diff --git a/Engine/Source/Programs/Horde/HordeDashboard/src/backend/Api.ts b/Engine/Source/Programs/Horde/HordeDashboard/src/backend/Api.ts
--- a/Engine/Source/Programs/Horde/HordeDashboard/src/backend/Api.ts
+++ b/Engine/Source/Programs/Horde/HordeDashboard/src/backend/Api.ts
@@ -4310,6 +4310,10 @@ export type GetTestMetaResponse = {
 
 	/// The rendering hardware interface being used with the test		
 	rhi: string;
+
+	/// The test variation identifier (or "default")
+	variation: string;
+
 }
 
 
diff --git a/Engine/Source/Programs/Horde/HordeDashboard/src/components/AutomationSuiteDetails.tsx b/Engine/Source/Programs/Horde/HordeDashboard/src/components/AutomationSuiteDetails.tsx
--- a/Engine/Source/Programs/Horde/HordeDashboard/src/components/AutomationSuiteDetails.tsx
+++ b/Engine/Source/Programs/Horde/HordeDashboard/src/components/AutomationSuiteDetails.tsx
@@ -68,8 +68,8 @@ class SuiteHandler {
 
       this.metaData = metaData;
 
-      this.metaName = `${metaData.platforms.join(" - ")} / ${metaData.configurations.join(" - ")} / ${metaData.buildTargets.join(" - ")} / ${metaData.rhi === "default" ? "Default" : metaData.rhi?.toUpperCase()}`;
-
+      this.metaName = `${metaData.platforms.join(" - ")} / ${metaData.configurations.join(" - ")} / ${metaData.buildTargets.join(" - ")} / ${metaData.rhi === "default" ? "Default" : metaData.rhi?.toUpperCase()} / ${metaData.variation === "default" ? "Default" : metaData.variation?.toUpperCase()}`;
+         
       this.load();
    }
 
diff --git a/Engine/Source/Programs/Horde/HordeDashboard/src/components/AutomationSuiteTest.tsx b/Engine/Source/Programs/Horde/HordeDashboard/src/components/AutomationSuiteTest.tsx
--- a/Engine/Source/Programs/Horde/HordeDashboard/src/components/AutomationSuiteTest.tsx
+++ b/Engine/Source/Programs/Horde/HordeDashboard/src/components/AutomationSuiteTest.tsx
@@ -45,7 +45,7 @@ class SuiteTestHandler {
       this.testRef = testRef;
       this.metaData = metaData;
 
-      this.metaName = `${metaData.platforms.join(" - ")} / ${metaData.configurations.join(" - ")} / ${metaData.buildTargets.join(" - ")} / ${metaData.rhi === "default" ? "Default" : metaData.rhi?.toUpperCase()}`;
+      this.metaName = `${metaData.platforms.join(" - ")} / ${metaData.configurations.join(" - ")} / ${metaData.buildTargets.join(" - ")} / ${metaData.rhi === "default" ? "Default" : metaData.rhi?.toUpperCase()} / ${metaData.variation === "default" ? "Default" : metaData.variation?.toUpperCase()}`;
 
       this.load();
    }
diff --git a/Engine/Source/Programs/Horde/HordeDashboard/src/components/AutomationView.tsx b/Engine/Source/Programs/Horde/HordeDashboard/src/components/AutomationView.tsx
--- a/Engine/Source/Programs/Horde/HordeDashboard/src/components/AutomationView.tsx
+++ b/Engine/Source/Programs/Horde/HordeDashboard/src/components/AutomationView.tsx
@@ -48,6 +48,7 @@ type TestDataState = {
    configurations?: string[];
    targets?: string[];
    rhi?: string[];
+   variation?: string[];
    tests?: string[];
    suites?: string[];
    weeks?: number;
@@ -264,6 +265,52 @@ class TestDataHandler {
       }
    }
 
+   addVariation(variation: string) {
+
+      const state = this.state;
+
+      if (state.variation) {
+         if (state.variation.indexOf(variation) !== -1) {
+            return;
+         }
+      }
+
+      if (!state.variation) {
+         state.variation = [];
+      }
+
+      state.variation.push(variation);
+
+      if (this.updateSearch()) {
+         this.setUpdated();
+      }
+   }
+
+   removeVariation(variation: string) {
+
+      const state = this.state;
+
+      if (state.variation?.indexOf(variation) === -1) {
+         return;
+      }
+
+      const idx = state.variation?.indexOf(variation);
+      if (!isNumber(idx) || idx < 0) {
+         return;
+      }
+
+      state.variation!.splice(idx, 1);
+
+      if (!state.variation!.length) {
+         state.variation = undefined;
+      }
+
+      if (this.updateSearch()) {
+         this.setUpdated();
+      }
+   }
+
+
    addTest(test: string) {
 
       const state = this.state;
@@ -448,6 +495,12 @@ class TestDataHandler {
                }
             }
 
+            if (state.variation?.length) {
+               if (!state.variation.find(r => m.variation === r)) {
+                  return;
+               }
+            }
+
             meta.set(m.id, m);
 
          });
@@ -581,6 +634,7 @@ class TestDataHandler {
       state.configurations = state.configurations?.sort((a, b) => a.localeCompare(b));
       state.targets = state.targets?.sort((a, b) => a.localeCompare(b));
       state.rhi = state.rhi?.sort((a, b) => a.localeCompare(b));
+      state.variation = state.variation?.sort((a, b) => a.localeCompare(b));
 
       state.tests = state.tests?.sort((a, b) => a.localeCompare(b));
       state.suites = state.suites?.sort((a, b) => a.localeCompare(b));
@@ -617,6 +671,10 @@ class TestDataHandler {
          search.append("rhi", r);
       });
 
+      state.variation?.forEach(v => {
+         search.append("variation", v);
+      });
+
       state.tests?.forEach(r => {
          search.append("test", r);
       });
@@ -653,6 +711,7 @@ class TestDataHandler {
       const configurations = search.getAll("configurations") ?? undefined;
       const targets = search.getAll("targets") ?? undefined;
       const rhi = search.getAll("rhi") ?? undefined;
+      const variation = search.getAll("var") ?? undefined;
 
       const tests = search.getAll("test") ?? undefined;
       const suites = search.getAll("suite") ?? undefined;
@@ -666,6 +725,7 @@ class TestDataHandler {
       state.configurations = configurations?.sort((a, b) => a.localeCompare(b));
       state.targets = targets?.sort((a, b) => a.localeCompare(b));
       state.rhi = rhi?.sort((a, b) => a.localeCompare(b));
+      state.variation = variation?.sort((a, b) => a.localeCompare(b));
 
       state.tests = tests?.sort((a, b) => a.localeCompare(b));
       state.suites = suites?.sort((a, b) => a.localeCompare(b));
@@ -678,7 +738,7 @@ class TestDataHandler {
       return state;
    }
 
-   getMetaString(metaId: string, platform = true, config = true, target = true, rhi = true) {
+   getMetaString(metaId: string, platform = true, config = true, target = true, rhi = true, variation = true) {
 
       const meta = this.metaData.get(metaId);
       if (!meta) {
@@ -701,6 +761,10 @@ class TestDataHandler {
          elements.push(meta.rhi === "default" ? "Default" : meta.rhi.toUpperCase());
       }
 
+      if (variation) {
+         elements.push(meta.variation === "default" ? "Default" : meta.variation.toUpperCase());
+      }
+
       if (elements.length) {
          return elements.join(" / ");
       }
@@ -709,7 +773,7 @@ class TestDataHandler {
 
    }
 
-   getCommonMeta(metaIds: string[], config = true, target = true, rhi = true) {
+   getCommonMeta(metaIds: string[], config = true, target = true, rhi = true, variation = true) {
 
       let cconfigs = "";
       let uniqueConfigs = false;
@@ -720,6 +784,9 @@ class TestDataHandler {
       let crhi = "";
       let uniqueRHI = false;
 
+      let cvariation = "";
+      let uniqueVariation = false;
+
       metaIds.forEach(m => {
 
          const meta = this.metaData.get(m)!;
@@ -750,8 +817,18 @@ class TestDataHandler {
             } else {
                crhi = rhi;
             }
+         }
 
+         if (variation) {
+            const variation = meta.variation === "default" ? "Default" : meta.variation.toUpperCase()
+
+            if (cvariation && cvariation !== variation) {
+               uniqueVariation = true;
+            } else {
+               cvariation = variation;
+            }
          }
+
       });
 
       const elements: string[] = [];
@@ -765,12 +842,16 @@ class TestDataHandler {
          if (rhi && !uniqueRHI) {
             elements.push(crhi);
          }
+         if (variation && !uniqueVariation) {
+            elements.push(cvariation);
+         }
       }
 
       return {
          commonConfigs: !uniqueConfigs,
          commonTargets: !uniqueTargets,
          commonRHI: !uniqueRHI,
+         commonVariation: !uniqueVariation,
          commonMetaString: elements.length ? elements.join(" / ") : ""
       }
    }
@@ -1025,7 +1106,7 @@ class TestDataHandler {
             return;
          }
 
-         const metaName = `${meta.platforms.join(" - ")} / ${meta.configurations.join(" - ")} / ${meta.buildTargets.join(" - ")} / ${meta.rhi}`;
+         const metaName = `${meta.platforms.join(" - ")} / ${meta.configurations.join(" - ")} / ${meta.buildTargets.join(" - ")} / ${meta.rhi} / ${meta.variation}`;
 
          this.metaNames.set(id, metaName);
 
@@ -1036,6 +1117,7 @@ class TestDataHandler {
       const configurations = new Set<string>();
       const targets = new Set<string>();
       const rhi = new Set<string>();
+      const variation = new Set<string>();
 
       metaData.forEach(m => {
          automation.add(m.projectName);
@@ -1043,6 +1125,7 @@ class TestDataHandler {
          m.buildTargets.forEach(t => targets.add(t));
          m.configurations.forEach(c => configurations.add(c));
          rhi.add(m.rhi);
+         variation.add(m.variation)
       });
 
       this.automation = Array.from(automation).sort((a, b) => a.localeCompare(b));
@@ -1050,6 +1133,7 @@ class TestDataHandler {
       this.configurations = Array.from(configurations).sort((a, b) => a.localeCompare(b));
       this.targets = Array.from(targets).sort((a, b) => a.localeCompare(b));
       this.rhi = Array.from(rhi).sort((a, b) => a.localeCompare(b));
+      this.variation = Array.from(variation).sort((a, b) => a.localeCompare(b));
 
       this.loaded = true;
       this.setUpdated();
@@ -1178,11 +1262,11 @@ class TestDataHandler {
    configurations: string[] = [];
    targets: string[] = [];
    rhi: string[] = [];
+   variation: string[] = [];
 
    tests: string[] = [];
    suites: string[] = [];
 
-
    private refs: GetTestDataRefResponse[] = [];
    minCL: number = 0;
    maxCL: number = 0;
@@ -1490,6 +1574,35 @@ const RHIChooser: React.FC<{ handler: TestDataHandler }> = observer(({ handler }
    </Stack>
 });
 
+const VariationChooser: React.FC<{ handler: TestDataHandler }> = observer(({ handler }) => {
+
+   // subscribe
+   if (handler.updated) { }
+
+   const cvariation: Set<string> = new Set(handler.state.variation ?? []);
+
+   const variation = handler.variation.map(v => {
+      return <Stack>
+         <Checkbox label={v === "default" ? "Default" : v.toLocaleUpperCase()} checked={!!cvariation.has(v)} onChange={(ev, checked) => {
+            if (checked) {
+               handler.addVariation(v);
+            } else {
+               handler.removeVariation(v);
+            }
+         }} />
+      </Stack>
+   });
+
+   return <Stack>
+      <Stack style={{ paddingTop: 12, paddingBottom: 4 }}>
+         <Label>Variation</Label>
+      </Stack>
+      <Stack tokens={{ childrenGap: 6 }} style={{ paddingLeft: 12 }}>
+         {variation}
+      </Stack>
+   </Stack>
+});
+
 const AutomationSidebarLeft: React.FC<{ handler: TestDataHandler }> = ({ handler }) => {
 
    return <Stack style={{ paddingRight: 18 }}>
@@ -1536,6 +1649,9 @@ const AutomationSidebarRight: React.FC<{ handler: TestDataHandler }> = ({ handle
          <Stack>
             <RHIChooser handler={handler} />
          </Stack>
+         <Stack>
+            <VariationChooser handler={handler} />
+         </Stack>
       </Stack>
    </Stack>
 }
@@ -1850,6 +1966,10 @@ class AutomationGraph {
             elements.push(meta.rhi === "default" ? "Default" : meta.rhi.toUpperCase());
          }
 
+         if (!commonMeta.commonVariation) {
+            elements.push(meta.variation === "default" ? "Default" : meta.variation.toUpperCase());
+         }
+
          this.metaNames.set(m, `${elements.join(" / ")}`);
 
       });
@@ -2016,9 +2136,9 @@ class AutomationGraph {
             rticks.unshift(first);
             rticks.push(last);
             ticks = rticks;
-            
+
          }
-         
+
          g.attr("transform", `translate(0,16)`)
             .style("font-family", "Horde Open Sans Regular")
             .style("font-size", "9px")
@@ -2268,7 +2388,7 @@ const AutomationTestView: React.FC<{ test: GetTestResponse, handler: TestDataHan
       const streamRefs = streamMap.get(s)!;
       const metaIds = Array.from(new Set<string>(streamRefs.keys()));
 
-      const commonMeta = handler.getCommonMeta(metaIds, !testCommonMeta.commonConfigs, !testCommonMeta.commonTargets, !testCommonMeta.commonRHI);
+      const commonMeta = handler.getCommonMeta(metaIds, !testCommonMeta.commonConfigs, !testCommonMeta.commonTargets, !testCommonMeta.commonRHI, !testCommonMeta.commonVariation);
 
       let anyPlatformFailed = false;
       for (const [, refs] of streamRefs) {
