Horde: Refactoring of the runtime config system, which adds support for including other JSON files and provides access via regular C# IOptions/IOptionsMonitor etc... DI types.

Combined config state is serialized back into a JSON object and compressed, then stored in Redis and broadcast to other server instances via pub/sub events. Update notifications are supported via IChangeTokens and callbacks on IOptionsMonitor.

Exists in parallel with the legacy config system for the time being.

diff --git a/Engine/Source/Programs/Horde/Horde.Build/Configuration/ConfigContext.cs b/Engine/Source/Programs/Horde/Horde.Build/Configuration/ConfigContext.cs
new file mode 100644
--- /dev/null
+++ b/Engine/Source/Programs/Horde/Horde.Build/Configuration/ConfigContext.cs
@@ -0,0 +1,99 @@
+// Copyright Epic Games, Inc. All Rights Reserved.
+
+using System;
+using System.Collections.Generic;
+using System.Diagnostics;
+using System.Text.Json;
+
+namespace Horde.Build.Configuration
+{
+	/// <summary>
+	/// Context for reading a tree of config files
+	/// </summary>
+	[DebuggerDisplay("{CurrentFile}")]
+	class ConfigContext
+	{
+		/// <summary>
+		/// Options for serializing config files
+		/// </summary>
+		public JsonSerializerOptions JsonOptions { get; }
+
+		/// <summary>
+		/// Stack of included files
+		/// </summary>
+		public Stack<Uri> IncludeStack { get; } = new Stack<Uri>();
+
+		/// <summary>
+		/// Stack of properties
+		/// </summary>
+		public Stack<string> ScopeStack { get; } = new Stack<string>();
+
+		/// <summary>
+		/// Map of property path to the file declaring a value for it
+		/// </summary>
+		public Dictionary<string, Uri> PropertyPathToFile = new Dictionary<string, Uri>(StringComparer.OrdinalIgnoreCase);
+
+		/// <summary>
+		/// Sources to read config files from
+		/// </summary>
+		public IReadOnlyDictionary<string, IConfigSource> Sources { get; }
+
+		/// <summary>
+		/// Tracks files read as part of the configuration
+		/// </summary>
+		public Dictionary<Uri, IConfigData> Files { get; } = new Dictionary<Uri, IConfigData>();
+
+		/// <summary>
+		/// Uri of the current file
+		/// </summary>
+		public Uri CurrentFile => (IncludeStack.Count > 0) ? IncludeStack.Peek() : null!;
+
+		/// <summary>
+		/// Current property scope
+		/// </summary>
+		public string CurrentScope => ScopeStack.Peek();
+
+		/// <summary>
+		/// Constructor
+		/// </summary>
+		public ConfigContext(JsonSerializerOptions jsonOptions, IReadOnlyDictionary<string, IConfigSource> sources)
+		{
+			JsonOptions = jsonOptions;
+			Sources = sources;
+			ScopeStack.Push("$");
+		}
+
+		/// <summary>
+		/// Marks a property as defined in the current file
+		/// </summary>
+		/// <param name="name"></param>
+		public void AddProperty(string name)
+		{
+			string propertyPath = $"{CurrentScope}.{name}";
+
+			Uri currentFile = CurrentFile;
+			if (!PropertyPathToFile.TryAdd(propertyPath, currentFile))
+			{
+				Uri otherFile = PropertyPathToFile[propertyPath];
+				throw new ConfigException(this, $"Property {propertyPath} was already defined in {otherFile}.");
+			}
+		}
+
+		/// <summary>
+		/// Pushes a scope to the property stack
+		/// </summary>
+		/// <param name="name"></param>
+		public void EnterScope(string name)
+		{
+			ScopeStack.Push($"{CurrentScope}.{name}");
+		}
+
+		/// <summary>
+		/// Pops a scope from the property stack
+		/// </summary>
+		public void LeaveScope()
+		{
+			ScopeStack.Pop();
+		}
+	}
+}
diff --git a/Engine/Source/Programs/Horde/Horde.Build/Configuration/ConfigData.cs b/Engine/Source/Programs/Horde/Horde.Build/Configuration/ConfigData.cs
new file mode 100644
--- /dev/null
+++ b/Engine/Source/Programs/Horde/Horde.Build/Configuration/ConfigData.cs
@@ -0,0 +1,270 @@
+// Copyright Epic Games, Inc. All Rights Reserved.
+
+using System;
+using System.Collections.Concurrent;
+using System.Collections.Generic;
+using System.IO;
+using System.Linq;
+using System.Threading;
+using System.Threading.Tasks;
+using EpicGames.Core;
+using EpicGames.Perforce;
+using Horde.Build.Perforce;
+
+namespace Horde.Build.Configuration
+{
+	/// <summary>
+	/// Accessor for a specific revision of a config file. Provides metadata about the current revision, and allows reading its data.
+	/// </summary>
+	public interface IConfigData
+	{
+		/// <summary>
+		/// URI of the config file
+		/// </summary>
+		Uri Uri { get; }
+
+		/// <summary>
+		/// String used to identify a specific version of the config data. Used for ordinal comparisons, otherwise opaque.
+		/// </summary>
+		string Revision { get; }
+
+		/// <summary>
+		/// Reads data for the config file
+		/// </summary>
+		/// <param name="cancellationToken">Cancellation token for the operation</param>
+		/// <returns>UTF-8 encoded data for the config file</returns>
+		ValueTask<ReadOnlyMemory<byte>> ReadAsync(CancellationToken cancellationToken);
+	}
+
+	/// <summary>
+	/// Source for reading config files
+	/// </summary>
+	public interface IConfigSource
+	{
+		/// <summary>
+		/// URI scheme for this config source
+		/// </summary>
+		string Scheme { get; }
+
+		/// <summary>
+		/// Reads a config file from this source
+		/// </summary>
+		/// <param name="uri">Location of the config file</param>
+		/// <param name="cancellationToken">Cancellation token for the operation</param>
+		/// <returns>Config file data</returns>
+		Task<IConfigData> GetAsync(Uri uri, CancellationToken cancellationToken);
+	}
+
+	/// <summary>
+	/// In-memory config file source
+	/// </summary>
+	public sealed class InMemoryConfigSource : IConfigSource
+	{
+		class ConfigFileRevisionImpl : IConfigData
+		{
+			public Uri Uri { get; }
+			public string Revision { get; }
+			public ReadOnlyMemory<byte> Data { get; }
+
+			public ConfigFileRevisionImpl(Uri uri, string version, ReadOnlyMemory<byte> data)
+			{
+				Uri = uri;
+				Revision = version;
+				Data = data;
+			}
+
+			public ValueTask<ReadOnlyMemory<byte>> ReadAsync(CancellationToken cancellationToken) => new ValueTask<ReadOnlyMemory<byte>>(Data);
+		}
+
+		readonly Dictionary<Uri, ConfigFileRevisionImpl> _files = new Dictionary<Uri, ConfigFileRevisionImpl>();
+
+		/// <summary>
+		/// Name of the scheme for this source
+		/// </summary>
+		public const string Scheme = "memory";
+
+		/// <inheritdoc/>
+		string IConfigSource.Scheme => Scheme;
+
+		/// <summary>
+		/// Manually adds a new config file
+		/// </summary>
+		/// <param name="path">Path to the config file</param>
+		/// <param name="data">Config file data</param>
+		public void Add(Uri path, ReadOnlyMemory<byte> data)
+		{
+			_files.Add(path, new ConfigFileRevisionImpl(path, "v1", data));
+		}
+
+		/// <inheritdoc/>
+		public Task<IConfigData> GetAsync(Uri uri, CancellationToken cancellationToken)
+		{
+			ConfigFileRevisionImpl? configFile;
+			if (!_files.TryGetValue(uri, out configFile))
+			{
+				throw new FileNotFoundException($"Config file {uri} not found.");
+			}
+			return Task.FromResult<IConfigData>(configFile);
+		}
+	}
+
+	/// <summary>
+	/// Config file source which reads from the filesystem
+	/// </summary>
+	public sealed class FileConfigSource : IConfigSource
+	{
+		class ConfigFileImpl : IConfigData
+		{
+			public Uri Uri { get; }
+			public string Revision { get; }
+			public DateTime LastWriteTimeUtc { get; }
+			public ReadOnlyMemory<byte> Data { get; }
+
+			public ConfigFileImpl(Uri uri, DateTime lastWriteTimeUtc, ReadOnlyMemory<byte> data)
+			{
+				Uri = uri;
+				Revision = $"timestamp={lastWriteTimeUtc.Ticks}";
+				LastWriteTimeUtc = lastWriteTimeUtc;
+				Data = data;
+			}
+
+			public ValueTask<ReadOnlyMemory<byte>> ReadAsync(CancellationToken cancellationToken) => new ValueTask<ReadOnlyMemory<byte>>(Data);
+		}
+
+		/// <summary>
+		/// Name of the scheme for this source
+		/// </summary>
+		public const string Scheme = "file";
+
+		/// <inheritdoc/>
+		string IConfigSource.Scheme => Scheme;
+
+		readonly DirectoryReference _baseDir;
+		readonly ConcurrentDictionary<FileReference, ConfigFileImpl> _files = new ConcurrentDictionary<FileReference, ConfigFileImpl>();
+
+		/// <summary>
+		/// Constructor
+		/// </summary>
+		public FileConfigSource()
+			: this(DirectoryReference.GetCurrentDirectory())
+		{
+		}
+
+		/// <summary>
+		/// Constructor
+		/// </summary>
+		/// <param name="baseDir">Base directory for resolving relative paths</param>
+		public FileConfigSource(DirectoryReference baseDir)
+		{
+			_baseDir = baseDir;
+		}
+
+		/// <inheritdoc/>
+		public async Task<IConfigData> GetAsync(Uri uri, CancellationToken cancellationToken)
+		{
+			FileReference localPath = FileReference.Combine(_baseDir, uri.LocalPath);
+
+			ConfigFileImpl? file;
+			for(; ;)
+			{
+				if (_files.TryGetValue(localPath, out file))
+				{
+					if (FileReference.GetLastWriteTimeUtc(localPath) == file.LastWriteTimeUtc)
+					{
+						break;
+					}
+					else
+					{
+						_files.TryRemove(new KeyValuePair<FileReference, ConfigFileImpl>(localPath, file));
+					}
+				}
+
+				using (FileStream stream = FileReference.Open(localPath, FileMode.Open, FileAccess.Read, FileShare.Read))
+				{
+					using MemoryStream memoryStream = new MemoryStream();
+					await stream.CopyToAsync(memoryStream, cancellationToken);
+					DateTime lastWriteTime = FileReference.GetLastWriteTimeUtc(localPath);
+					file = new ConfigFileImpl(uri, lastWriteTime, memoryStream.ToArray());
+				}
+
+				if (_files.TryAdd(localPath, file))
+				{
+					break;
+				}
+			}
+
+			return file;
+		}
+	}
+
+	/// <summary>
+	/// Perforce cluster config file source
+	/// </summary>
+	public sealed class PerforceConfigSource : IConfigSource
+	{
+		class ConfigFileImpl : IConfigData
+		{
+			public Uri Uri { get; }
+			public int Change { get; }
+			public string Revision { get; }
+
+			readonly PerforceConfigSource _owner;
+
+			public ConfigFileImpl(Uri uri, int change, PerforceConfigSource owner)
+			{
+				Uri = uri;
+				Change = change;
+				Revision = $"{change}";
+				_owner = owner;
+			}
+
+			public ValueTask<ReadOnlyMemory<byte>> ReadAsync(CancellationToken cancellationToken) => _owner.ReadAsync(Uri, Change, cancellationToken);
+		}
+
+		/// <summary>
+		/// Name of the scheme for this source
+		/// </summary>
+		public const string Scheme = "p4-cluster";
+
+		/// <inheritdoc/>
+		string IConfigSource.Scheme => Scheme;
+
+		readonly IPerforceService _perforceService;
+
+		/// <summary>
+		/// Constructor
+		/// </summary>
+		/// <param name="perforceService">Perforce service instance</param>
+		public PerforceConfigSource(IPerforceService perforceService)
+		{
+			_perforceService = perforceService;
+		}
+
+		/// <inheritdoc/>
+		public async Task<IConfigData> GetAsync(Uri uri, CancellationToken cancellationToken)
+		{
+			using (IPerforceConnection perforce = await _perforceService.ConnectAsync(uri.Host, null, cancellationToken))
+			{
+				List<FStatRecord> files = await perforce.FStatAsync(FStatOptions.ShortenOutput, uri.AbsolutePath, cancellationToken).ToListAsync(cancellationToken);
+				files.RemoveAll(x => x.HeadAction == FileAction.Delete || x.HeadAction == FileAction.MoveDelete);
+
+				if (files.Count == 0)
+				{
+					throw new FileNotFoundException($"Unable to read {uri}. No matching files found.");
+				}
+
+				return new ConfigFileImpl(uri, files[0].HeadChange, this);
+			}
+		}
+
+		async ValueTask<ReadOnlyMemory<byte>> ReadAsync(Uri uri, int change, CancellationToken cancellationToken)
+		{
+			using (IPerforceConnection perforce = await _perforceService.ConnectAsync(uri.Host, null, cancellationToken))
+			{
+				PerforceResponse<PrintRecord<byte[]>> response = await perforce.TryPrintDataAsync($"{uri.AbsolutePath}@{change}", cancellationToken);
+				response.EnsureSuccess();
+				return response.Data.Contents!;
+			}
+		}
+	}
+}
diff --git a/Engine/Source/Programs/Horde/Horde.Build/Configuration/ConfigService.cs b/Engine/Source/Programs/Horde/Horde.Build/Configuration/ConfigService.cs
new file mode 100644
--- /dev/null
+++ b/Engine/Source/Programs/Horde/Horde.Build/Configuration/ConfigService.cs
@@ -0,0 +1,419 @@
+// Copyright Epic Games, Inc. All Rights Reserved.
+
+using System;
+using System.Collections.Generic;
+using System.IO;
+using System.IO.Compression;
+using System.Linq;
+using System.Text.Json;
+using System.Threading;
+using System.Threading.Tasks;
+using EpicGames.Core;
+using EpicGames.Redis;
+using Horde.Build.Projects;
+using Horde.Build.Server;
+using Horde.Build.Streams;
+using Horde.Build.Utilities;
+using HordeCommon;
+using Microsoft.Extensions.Hosting;
+using Microsoft.Extensions.Logging;
+using Microsoft.Extensions.Options;
+using Microsoft.Extensions.Primitives;
+using ProtoBuf;
+using StackExchange.Redis;
+
+namespace Horde.Build.Configuration
+{
+	using ProjectId = StringId<IProject>;
+	using StreamId = StringId<IStream>;
+
+	/// <summary>
+	/// Service which processes runtime configuration data.
+	/// </summary>
+	public sealed class ConfigService : IOptionsFactory<GlobalConfig>, IOptionsChangeTokenSource<GlobalConfig>, IHostedService, IDisposable
+	{
+		// Index to all current config files 
+		[ProtoContract]
+		class ConfigSnapshot
+		{
+			[ProtoMember(1)]
+			public byte[] Global = null!;
+
+			[ProtoMember(2)]
+			public Dictionary<ProjectId, byte[]> Projects { get; set; } = new Dictionary<ProjectId, byte[]>();
+
+			[ProtoMember(3)]
+			public Dictionary<StreamId, byte[]> Streams { get; set; } = new Dictionary<StreamId, byte[]>();
+
+			[ProtoMember(4)]
+			public Dictionary<Uri, string> Dependencies { get; set; } = new Dictionary<Uri, string>();
+		}
+
+		// Implements a token for config change notifications
+		class ChangeToken : IChangeToken
+		{
+			class Registration : IDisposable
+			{
+				public Action<object>? _callback;
+				public object? _state;
+				public Registration? _next;
+
+				public void Dispose() => _callback = null;
+			}
+
+			static readonly Registration s_sentinel = new Registration();
+
+			Registration? _firstRegistration = null;
+
+			public bool ActiveChangeCallbacks => true;
+			public bool HasChanged => _firstRegistration != s_sentinel;
+
+			public void TriggerChange()
+			{
+				for (Registration? firstRegistration = _firstRegistration; firstRegistration != s_sentinel; firstRegistration = _firstRegistration)
+				{
+					if (Interlocked.CompareExchange(ref _firstRegistration, s_sentinel, firstRegistration) == firstRegistration)
+					{
+						for (; firstRegistration != null; firstRegistration = firstRegistration._next)
+						{
+							Action<object>? callback = firstRegistration._callback;
+							if (callback != null)
+							{
+								callback(firstRegistration._state!);
+							}
+						}
+						break;
+					}
+				}
+			}
+
+			public IDisposable RegisterChangeCallback(Action<object> callback, object state)
+			{
+				Registration registration = new Registration { _callback = callback, _state = state };
+				for (; ; )
+				{
+					Registration? nextRegistration = _firstRegistration;
+					registration._next = nextRegistration;
+
+					if (nextRegistration == s_sentinel)
+					{
+						callback(state);
+						break;
+					}
+					if (Interlocked.CompareExchange(ref _firstRegistration, registration, nextRegistration) == nextRegistration)
+					{
+						break;
+					}
+				}
+				return registration;
+			}
+		}
+
+		// The current config
+		record class ConfigState(IoHash Hash, GlobalConfig GlobalConfig);
+
+		readonly RedisService _redisService;
+		readonly ServerSettings _settings;
+		readonly Dictionary<string, IConfigSource> _sources;
+		readonly JsonSerializerOptions _jsonOptions;
+		readonly RedisKey _snapshotKey = "config";
+
+		readonly ITicker _ticker;
+		readonly TimeSpan _tickInterval = TimeSpan.FromSeconds(2.0);
+
+		readonly RedisChannel _updateChannel = "config-update";
+		readonly BackgroundTask _updateTask;
+
+		Task<ConfigState> _stateTask;
+		ChangeToken? _currentChangeToken;
+
+		/// <inheritdoc/>
+		string IOptionsChangeTokenSource<GlobalConfig>.Name => String.Empty;
+
+		/// <summary>
+		/// Constructor
+		/// </summary>
+		public ConfigService(RedisService redisService, IOptions<ServerSettings> settings, IEnumerable<IConfigSource> sources, IClock clock, ILogger<ConfigService> logger)
+		{
+			_redisService = redisService;
+			_settings = settings.Value;
+			_sources = sources.ToDictionary(x => x.Scheme, x => x, StringComparer.OrdinalIgnoreCase);
+
+			_jsonOptions = new JsonSerializerOptions();
+			Startup.ConfigureJsonSerializer(_jsonOptions);
+
+			_ticker = clock.AddSharedTicker<ConfigService>(_tickInterval, TickSharedAsync, logger);
+
+			_updateTask = new BackgroundTask(WaitForUpdatesAsync);
+
+			_stateTask = Task.Run(() => GetStartupStateAsync());
+		}
+
+		/// <inheritdoc/>
+		public void Dispose()
+		{
+			_updateTask.Dispose();
+			_ticker.Dispose();
+		}
+
+		/// <inheritdoc/>
+		public async Task StartAsync(CancellationToken cancellationToken)
+		{
+			await _ticker.StartAsync();
+			_updateTask.Start();
+		}
+
+		/// <inheritdoc/>
+		public async Task StopAsync(CancellationToken cancellationToken)
+		{
+			await _updateTask.StopAsync();
+			await _ticker.StopAsync();
+		}
+
+		/// <inheritdoc/>
+		public GlobalConfig Create(string name) => _stateTask.Result.GlobalConfig;
+
+		/// <inheritdoc/>
+		public IChangeToken GetChangeToken()
+		{
+			ChangeToken? newToken = null;
+			for (; ; )
+			{
+				ChangeToken? token = _currentChangeToken;
+				if (token != null)
+				{
+					return token;
+				}
+
+				newToken ??= new ChangeToken();
+
+				if (Interlocked.CompareExchange(ref _currentChangeToken, newToken, null) == null)
+				{
+					return newToken;
+				}
+			}
+		}
+
+		async Task<ConfigState> GetStartupStateAsync()
+		{
+			ReadOnlyMemory<byte> data = ReadOnlyMemory<byte>.Empty;
+			if (!_settings.ForceConfigUpdateOnStartup)
+			{
+				data = await ReadSnapshotDataAsync();
+			}
+			if (data.Length == 0)
+			{
+				data = SerializeSnapshot(await CreateSnapshotAsync(CancellationToken.None));
+			}
+			return new ConfigState(IoHash.Compute(data.Span), CreateGlobalConfig(data));
+		}
+
+		async ValueTask TickSharedAsync(CancellationToken cancellationToken)
+		{
+			// Read a new copy of the current snapshot in the context of being elected to perform updates. This
+			// avoids any latency with receiving notifications on the pub/sub channel about changes.
+			ReadOnlyMemory<byte> initialData = await ReadSnapshotDataAsync();
+
+			// Update the snapshot until we're asked to stop
+			ConfigSnapshot? snapshot = (initialData.Length > 0) ? DeserializeSnapshot(initialData) : null;
+			while (!cancellationToken.IsCancellationRequested)
+			{
+				if (snapshot == null || await IsOutOfDateAsync(snapshot, cancellationToken))
+				{
+					snapshot = await CreateSnapshotAsync(cancellationToken);
+					await WriteSnapshotAsync(snapshot);
+				}
+				await Task.Delay(_tickInterval, cancellationToken);
+			}
+		}
+
+		async Task WaitForUpdatesAsync(CancellationToken cancellationToken)
+		{
+			AsyncEvent updateEvent = new AsyncEvent();
+			await using (RedisSubscription subscription = await _redisService.SubscribeAsync(_updateChannel, _ => updateEvent.Pulse()))
+			{
+				Task cancellationTask = Task.Delay(-1, cancellationToken);
+				while (!cancellationToken.IsCancellationRequested)
+				{
+					Task updateTask = updateEvent.Task;
+
+					RedisValue value = await _redisService.GetDatabase().StringGetAsync(_snapshotKey);
+					if (!value.IsNullOrEmpty)
+					{
+						await UpdateConfigObjectAsync(value);
+					}
+
+					await await Task.WhenAny(updateTask, cancellationTask);
+				}
+			}
+		}
+
+		/// <summary>
+		/// Create a new snapshot object
+		/// </summary>
+		/// <param name="cancellationToken">Cancellation token for the operation</param>
+		/// <returns>New config snapshot</returns>
+		async Task<ConfigSnapshot> CreateSnapshotAsync(CancellationToken cancellationToken)
+		{
+			// Get the path to the root config file
+			Uri globalConfigUri;
+			if (Path.IsPathRooted(_settings.ConfigPath) && !_settings.ConfigPath.StartsWith("//", StringComparison.Ordinal))
+			{
+				// absolute path to config
+				globalConfigUri = new Uri(_settings.ConfigPath);
+			}
+			else
+			{
+				// relative (development) or perforce path
+				globalConfigUri = ConfigType.CombinePaths(new Uri(FileReference.Combine(Program.AppDir, "_").FullName), _settings.ConfigPath);
+			}
+
+			// Read the config files
+			ConfigContext context = new ConfigContext(_jsonOptions, _sources);
+			try
+			{
+				ConfigSnapshot snapshot = new ConfigSnapshot();
+
+				GlobalConfig globalConfig = await ConfigType.ReadAsync<GlobalConfig>(globalConfigUri, context, cancellationToken);
+				snapshot.Global = JsonSerializer.SerializeToUtf8Bytes(globalConfig, context.JsonOptions);
+
+				foreach (ProjectConfigRef projectConfigRef in globalConfig.Projects)
+				{
+					context.PropertyPathToFile.Clear();
+
+					ProjectConfig projectConfig = await ConfigType.ReadAsync<ProjectConfig>(new Uri(projectConfigRef.Path), context, cancellationToken);
+					snapshot.Projects.Add(projectConfigRef.Id, JsonSerializer.SerializeToUtf8Bytes(projectConfig, context.JsonOptions));
+
+					foreach (StreamConfigRef streamConfigRef in projectConfig.Streams)
+					{
+						if (streamConfigRef.Id.ToString() == "plastic-private-plastic")
+						{
+							continue;
+						}
+
+						context.PropertyPathToFile.Clear();
+
+						StreamConfig streamConfig = await ConfigType.ReadAsync<StreamConfig>(new Uri(streamConfigRef.Path), context, cancellationToken);
+						snapshot.Streams.Add(streamConfigRef.Id, JsonSerializer.SerializeToUtf8Bytes(streamConfig, context.JsonOptions));
+					}
+				}
+
+				// Save all the dependencies
+				foreach ((Uri depUri, IConfigData depData) in context.Files)
+				{
+					snapshot.Dependencies.Add(depUri, depData.Revision);
+				}
+
+				return snapshot;
+			}
+			catch (Exception ex) when (ex is not ConfigException)
+			{
+				throw new ConfigException(context, "Internal exception while parsing config files", ex);
+			}
+		}
+
+		/// <summary>
+		/// Checks if the given snapshot is out of date
+		/// </summary>
+		/// <param name="snapshot">The current config snapshot</param>
+		/// <param name="cancellationToken">Cancellation token for the operation</param>
+		/// <returns>True if the snapshot is out of date</returns>
+		async Task<bool> IsOutOfDateAsync(ConfigSnapshot snapshot, CancellationToken cancellationToken)
+		{
+			foreach ((Uri uri, string version) in snapshot.Dependencies)
+			{
+				IConfigSource? source = _sources[uri.Scheme];
+				if (source == null)
+				{
+					return true;
+				}
+
+				IConfigData data = await source.GetAsync(uri, cancellationToken);
+				if (!data.Revision.Equals(version, StringComparison.Ordinal))
+				{
+					return true;
+				}
+			}
+			return false;
+		}
+
+		async Task WriteSnapshotAsync(ConfigSnapshot index)
+		{
+			ReadOnlyMemory<byte> data = SerializeSnapshot(index);
+			await _redisService.GetDatabase().StringSetAsync(_snapshotKey, data);
+			await _redisService.PublishAsync(_updateChannel, RedisValue.EmptyString);
+		}
+
+		async Task<ReadOnlyMemory<byte>> ReadSnapshotDataAsync()
+		{
+			RedisValue value = await _redisService.GetDatabase().StringGetAsync(_snapshotKey);
+			return value.IsNullOrEmpty ? ReadOnlyMemory<byte>.Empty : (ReadOnlyMemory<byte>)value;
+		}
+
+		static ReadOnlyMemory<byte> SerializeSnapshot(ConfigSnapshot snapshot)
+		{
+			using (MemoryStream outputStream = new MemoryStream())
+			{
+				using GZipStream zipStream = new GZipStream(outputStream, CompressionMode.Compress, false);
+				ProtoBuf.Serializer.Serialize(zipStream, snapshot);
+				zipStream.Flush();
+				return outputStream.ToArray();
+			}
+		}
+
+		static ConfigSnapshot DeserializeSnapshot(ReadOnlyMemory<byte> data)
+		{
+			using (ReadOnlyMemoryStream outputStream = new ReadOnlyMemoryStream(data))
+			{
+				using GZipStream zipStream = new GZipStream(outputStream, CompressionMode.Decompress, false);
+				return ProtoBuf.Serializer.Deserialize<ConfigSnapshot>(zipStream);
+			}
+		}
+
+		async Task UpdateConfigObjectAsync(ReadOnlyMemory<byte> data)
+		{
+			ConfigState state = await _stateTask;
+
+			// Hash the data and only update if it changes; this prevents any double-updates due to time between initialization and the hosted service starting.
+			IoHash hash = IoHash.Compute(data.Span);
+			if (hash != state.Hash)
+			{
+				// Update the config object
+				GlobalConfig globalConfig = CreateGlobalConfig(data);
+				_stateTask = Task.FromResult(new ConfigState(hash, globalConfig));
+
+				// Notify any watchers of the new config object
+				for (; ; )
+				{
+					ChangeToken? token = _currentChangeToken;
+					if (Interlocked.CompareExchange(ref _currentChangeToken, null, token) == token)
+					{
+						token?.TriggerChange();
+						break;
+					}
+				}
+			}
+		}
+
+		GlobalConfig CreateGlobalConfig(ReadOnlyMemory<byte> data)
+		{
+			// Decompress the snapshot object
+			ConfigSnapshot snapshot = DeserializeSnapshot(data);
+
+			// Build the new config and store the project and stream configs inside it
+			GlobalConfig globalConfig = JsonSerializer.Deserialize<GlobalConfig>(snapshot.Global, _jsonOptions)!;
+			foreach (ProjectConfigRef projectConfigRef in globalConfig.Projects)
+			{
+				ProjectConfig projectConfig = JsonSerializer.Deserialize<ProjectConfig>(snapshot.Projects[projectConfigRef.Id], _jsonOptions)!;
+				projectConfigRef.Target = projectConfig;
+
+				foreach (StreamConfigRef streamConfigRef in projectConfig.Streams)
+				{
+					StreamConfig streamConfig = JsonSerializer.Deserialize<StreamConfig>(snapshot.Streams[streamConfigRef.Id], _jsonOptions)!;
+					streamConfigRef.Target = streamConfig;
+				}
+			}
+			return globalConfig;
+		}
+	}
+}
diff --git a/Engine/Source/Programs/Horde/Horde.Build/Configuration/ConfigType.cs b/Engine/Source/Programs/Horde/Horde.Build/Configuration/ConfigType.cs
new file mode 100644
--- /dev/null
+++ b/Engine/Source/Programs/Horde/Horde.Build/Configuration/ConfigType.cs
@@ -0,0 +1,484 @@
+// Copyright Epic Games, Inc. All Rights Reserved.
+
+using System;
+using System.Collections;
+using System.Collections.Concurrent;
+using System.Collections.Generic;
+using System.ComponentModel.DataAnnotations;
+using System.Reflection;
+using System.Text.Json;
+using System.Text.Json.Nodes;
+using System.Text.Json.Serialization;
+using System.Threading;
+using System.Threading.Tasks;
+using EpicGames.Core;
+using Horde.Build.Server;
+using JsonObject = System.Text.Json.Nodes.JsonObject;
+
+namespace Horde.Build.Configuration
+{
+	/// <summary>
+	/// Attribute used to mark <see cref="Uri"/> properties that are relative to their containing file
+	/// </summary>
+	[AttributeUsage(AttributeTargets.Property)]
+	public sealed class ConfigPathAttribute : Attribute
+	{
+	}
+
+	/// <summary>
+	/// Directive to merge config data from another source
+	/// </summary>
+	public class ConfigInclude
+	{
+		/// <summary>
+		/// Path to the config data to be included. May be relative to the including file's location.
+		/// </summary>
+		[Required]
+		public string Path { get; set; } = null!;
+	}
+
+	/// <summary>
+	/// Exception thrown when reading config files
+	/// </summary>
+	public sealed class ConfigException : Exception
+	{
+		/// <summary>
+		/// Path to the config file triggering the exception
+		/// </summary>
+		public Uri Path { get; }
+
+		/// <summary>
+		/// Constructor
+		/// </summary>
+		/// <param name="context">Current parse context for the error</param>
+		/// <param name="message">Description of the error</param>
+		/// <param name="innerException">Inner exception details</param>
+		internal ConfigException(ConfigContext context, string message, Exception? innerException = null)
+			: base(message, innerException)
+		{
+			Path = context.CurrentFile;
+		}
+	}
+
+	/// <summary>
+	/// Base class for types that can be read from config files
+	/// </summary>
+	abstract class ConfigType
+	{
+		public abstract ValueTask<object?> ReadAsync(JsonNode? node, ConfigContext context, CancellationToken cancellationToken);
+
+		static readonly ConcurrentDictionary<Type, ConfigType> s_typeToValueType = new ConcurrentDictionary<Type, ConfigType>();
+
+		public static ConfigType FindOrAddValueType(Type type)
+		{
+			ConfigType? value;
+			if (!s_typeToValueType.TryGetValue(type, out value))
+			{
+				if (!type.IsClass || type == typeof(string))
+				{
+					value = new ScalarConfigType(type);
+				}
+				else
+				{
+					value = ObjectConfigType.FindOrAdd(type);
+				}
+
+				lock (s_typeToValueType)
+				{
+					value = s_typeToValueType.GetOrAdd(type, value);
+				}
+			}
+			return value;
+		}
+
+		/// <summary>
+		/// Reads an object from a particular URL
+		/// </summary>
+		/// <typeparam name="T">Type of object to read</typeparam>
+		/// <param name="uri">Location of the file to read</param>
+		/// <param name="context">Context for reading</param>
+		/// <param name="cancellationToken">Cancellation token for the operation</param>
+		/// <returns></returns>
+		public static async Task<T> ReadAsync<T>(Uri uri, ConfigContext context, CancellationToken cancellationToken) where T : class, new()
+		{
+			ObjectConfigType type = ObjectConfigType.FindOrAdd(typeof(T));
+
+			T target = new T();
+			await type.MergeObjectAsync(target, uri, context, cancellationToken);
+			return target;
+		}
+
+		/// <summary>
+		/// Combine a relative path with a base URI to produce a new URI
+		/// </summary>
+		/// <param name="baseUri">Base uri to rebase relative to</param>
+		/// <param name="path">Relative path</param>
+		/// <returns>Absolute URI</returns>
+		public static Uri CombinePaths(Uri baseUri, string path)
+		{
+			if (path.StartsWith("//", StringComparison.Ordinal))
+			{
+				if (baseUri.Scheme == PerforceConfigSource.Scheme)
+				{
+					return new Uri($"{PerforceConfigSource.Scheme}://{baseUri.Host}{path}");
+				}
+				else
+				{
+					return new Uri($"{PerforceConfigSource.Scheme}://{PerforceCluster.DefaultName}{path}");
+				}
+			}
+			return new Uri(baseUri, path);
+		}
+	}
+
+	/// <summary>
+	/// Implementation of <see cref="ConfigType"/> for scalar values
+	/// </summary>
+	class ScalarConfigType : ConfigType
+	{
+		readonly Type _type;
+
+		public ScalarConfigType(Type type)
+		{
+			if (type.IsGenericType && type.GetGenericTypeDefinition() == typeof(Nullable<>))
+			{
+				type = type.GetGenericArguments()[0];
+			}
+			_type = type;
+		}
+
+		public override ValueTask<object?> ReadAsync(JsonNode? node, ConfigContext context, CancellationToken cancellationToken)
+		{
+			if (node == null)
+			{
+				return new ValueTask<object?>((object?)null);
+			}
+			else
+			{
+				return new ValueTask<object?>(node.Deserialize(_type, context.JsonOptions));
+			}
+		}
+	}
+
+	/// <summary>
+	/// Implementation of <see cref="ConfigType"/> to handle object variables
+	/// </summary>
+	class ObjectConfigType : ConfigType
+	{
+		abstract class Property
+		{
+			public string Name { get; }
+			public PropertyInfo PropertyInfo { get; }
+
+			public Property(string name, PropertyInfo propertyInfo)
+			{
+				Name = name;
+				PropertyInfo = propertyInfo;
+			}
+
+			public abstract Task MergeAsync(object target, JsonNode? node, ConfigContext context, CancellationToken cancellationToken);
+		}
+
+		class ScalarProperty : Property
+		{
+			public ScalarProperty(string name, PropertyInfo propertyInfo)
+				: base(name, propertyInfo)
+			{
+			}
+
+			public override Task MergeAsync(object target, JsonNode? node, ConfigContext context, CancellationToken cancellationToken)
+			{
+				context.AddProperty(Name);
+
+				object? value;
+				if (PropertyInfo.GetCustomAttribute<ConfigPathAttribute>() != null)
+				{
+					value = CombinePaths(context.CurrentFile, JsonSerializer.Deserialize<string>(node, context.JsonOptions) ?? String.Empty).AbsoluteUri;
+				}
+				else
+				{
+					value = JsonSerializer.Deserialize(node, PropertyInfo.PropertyType, context.JsonOptions);
+				}
+				PropertyInfo.SetValue(target, value);
+
+				return Task.CompletedTask;
+			}
+		}
+
+		class ListProperty : Property
+		{
+			readonly ConfigType _elementType;
+
+			public ListProperty(string name, PropertyInfo propertyInfo, ConfigType elementType)
+				: base(name, propertyInfo)
+			{
+				_elementType = elementType;
+			}
+
+			public override async Task MergeAsync(object target, JsonNode? node, ConfigContext context, CancellationToken cancellationToken)
+			{
+				IList? list = (IList?)PropertyInfo.GetValue(target);
+				if(list == null)
+				{
+					object value = Activator.CreateInstance(PropertyInfo.PropertyType)!;
+					PropertyInfo.SetValue(target, value);
+					list = (IList)value;
+				}
+
+				JsonArray array = (JsonArray)node!;
+				foreach (JsonNode? element in array)
+				{
+					context.EnterScope($"{Name}[{list.Count}]");
+
+					object? elementValue = await _elementType.ReadAsync(element, context, cancellationToken);
+					list.Add(elementValue);
+
+					context.LeaveScope();
+				}
+			}
+		}
+
+		class ObjectProperty : Property
+		{
+			readonly ObjectConfigType _classConfigType;
+
+			public ObjectProperty(string name, PropertyInfo propertyInfo, ObjectConfigType classConfigType)
+				: base(name, propertyInfo)
+			{
+				_classConfigType = classConfigType;
+			}
+
+			public override async Task MergeAsync(object target, JsonNode? node, ConfigContext context, CancellationToken cancellationToken)
+			{
+				if (node is JsonObject obj)
+				{
+					Uri? otherFile;
+					if (!context.PropertyPathToFile.TryGetValue(Name, out otherFile))
+					{
+						context.AddProperty(Name);
+					}
+
+					context.EnterScope(Name);
+
+					object? childTarget = PropertyInfo.GetValue(target);
+					if (childTarget == null)
+					{
+						if (otherFile != null)
+						{
+							throw new ConfigException(context, $"Property {context.CurrentScope}{Name} conflicts with value in {otherFile}.");
+						}
+
+						childTarget = await _classConfigType.ReadAsync(node, context, cancellationToken);
+						PropertyInfo.SetValue(target, childTarget);
+					}
+					else
+					{
+						await _classConfigType.MergeObjectAsync(childTarget, obj, context, cancellationToken);
+					}
+
+					context.LeaveScope();
+				}
+				else
+				{
+					context.AddProperty(Name);
+					PropertyInfo.SetValue(target, JsonSerializer.Deserialize(node, PropertyInfo.PropertyType, context.JsonOptions));
+				}
+			}
+		}
+
+		readonly Type _type;
+		readonly Dictionary<string, PropertyInfo> _nameToIncludeProperty = new Dictionary<string, PropertyInfo>(StringComparer.OrdinalIgnoreCase);
+		readonly Dictionary<string, Property> _nameToProperty = new Dictionary<string, Property>(StringComparer.OrdinalIgnoreCase);
+		readonly Dictionary<string, ObjectConfigType>? _knownTypes;
+
+		static readonly ConcurrentDictionary<Type, ObjectConfigType> s_typeToObjectValueType = new ConcurrentDictionary<Type, ObjectConfigType>();
+
+		public ObjectConfigType(Type type)
+		{
+			s_typeToObjectValueType.TryAdd(type, this);
+
+			_type = type;
+
+			foreach (PropertyInfo propertyInfo in type.GetProperties(BindingFlags.Instance | BindingFlags.Public | BindingFlags.GetProperty))
+			{
+				JsonPropertyNameAttribute? attribute = propertyInfo.GetCustomAttribute<JsonPropertyNameAttribute>();
+				string name = attribute?.Name ?? propertyInfo.Name;
+
+				if (propertyInfo.PropertyType == typeof(List<ConfigInclude>))
+				{
+					_nameToIncludeProperty.Add(name, propertyInfo);
+				}
+				else
+				{
+					_nameToProperty.Add(name, CreateProperty(name, propertyInfo));
+				}
+			}
+
+			JsonKnownTypesAttribute? knownTypes = _type.GetCustomAttribute<JsonKnownTypesAttribute>();
+			if (knownTypes != null)
+			{
+				_knownTypes = new Dictionary<string, ObjectConfigType>(StringComparer.Ordinal);
+				foreach (Type knownType in knownTypes.Types)
+				{
+					ObjectConfigType knownConfigType = FindOrAdd(knownType);
+					foreach (JsonDiscriminatorAttribute discriminatorAttribute in knownType.GetCustomAttributes(typeof(JsonDiscriminatorAttribute), true))
+					{
+						_knownTypes.Add(discriminatorAttribute.Name, knownConfigType);
+					}
+				}
+			}
+		}
+
+		public static ObjectConfigType FindOrAdd(Type type)
+		{
+			ObjectConfigType? value;
+			if (!s_typeToObjectValueType.TryGetValue(type, out value))
+			{
+				lock (s_typeToObjectValueType)
+				{
+					if (!s_typeToObjectValueType.TryGetValue(type, out value))
+					{
+						value = new ObjectConfigType(type);
+					}
+				}
+			}
+			return value;
+		}
+
+		static Property CreateProperty(string name, PropertyInfo propertyInfo)
+		{
+			Type propertyType = propertyInfo.PropertyType;
+			if (!propertyType.IsClass || propertyType == typeof(string))
+			{
+				return new ScalarProperty(name, propertyInfo);
+			}
+			else
+			{
+				if (propertyType.IsGenericType && propertyType.GetGenericTypeDefinition() == typeof(List<>))
+//				if (propertyType.IsAssignableTo(typeof(IList)))
+				{
+	//				foreach (Type type in propertyType.GetInterfaces())
+	//				{
+	//					if (type.IsGenericType && type.GetGenericTypeDefinition() == typeof(IList<>))
+	//					{
+//							Type elementType = type.GetGenericArguments()[0];
+					Type elementType = propertyType.GetGenericArguments()[0];
+					return new ListProperty(name, propertyInfo, FindOrAddValueType(elementType));
+//						}
+//					}
+				}
+				return new ObjectProperty(name, propertyInfo, FindOrAdd(propertyType));
+			}
+		}
+
+		public override async ValueTask<object?> ReadAsync(JsonNode? node, ConfigContext context, CancellationToken cancellationToken)
+		{
+			if (node == null)
+			{
+				throw new ConfigException(context, "Unable to deserialize object from null value");
+			}
+			else if (node is JsonObject obj)
+			{
+				return await ReadAsync(obj, context, cancellationToken);
+			}
+			else
+			{
+				return JsonSerializer.Deserialize(node, _type, context.JsonOptions);
+			}
+		}
+
+		public ValueTask<object> ReadAsync(JsonObject obj, ConfigContext context, CancellationToken cancellationToken)
+		{
+			ObjectConfigType targetType = this;
+			if (_knownTypes != null && obj.TryGetPropertyValue("Type", out JsonNode? knownTypeNode) && knownTypeNode != null)
+			{
+				targetType = _knownTypes[knownTypeNode.ToString()];
+			}
+			return targetType.ReadConcreteTypeAsync(obj, context, cancellationToken);
+		}
+
+		async ValueTask<object> ReadConcreteTypeAsync(JsonObject obj, ConfigContext context, CancellationToken cancellationToken)
+		{
+			object result = Activator.CreateInstance(_type)!;
+			await MergeObjectAsync(result, obj, context, cancellationToken);
+			return result;
+		}
+
+		static async ValueTask<JsonObject> ReadFileAsync(Uri uri, ConfigContext context, CancellationToken cancellationToken)
+		{
+			IConfigSource? source = context.Sources[uri.Scheme];
+			if (source == null)
+			{
+				throw new ConfigException(context, $"Invalid/unknown scheme for config file {uri}");
+			}
+
+			IConfigData? file;
+			if (!context.Files.TryGetValue(uri, out file))
+			{
+				file = await source.GetAsync(uri, cancellationToken);
+				context.Files.Add(uri, file);
+			}
+
+			ReadOnlyMemory<byte> data = await file.ReadAsync(cancellationToken);
+
+			JsonObject? obj = JsonSerializer.Deserialize<JsonObject>(data.Span, context.JsonOptions);
+			if (obj == null)
+			{
+				throw new ConfigException(context, $"Config file {uri} contains a null object.");
+			}
+
+			return obj;
+		}
+
+		public async Task MergeObjectAsync(object target, Uri uri, ConfigContext context, CancellationToken cancellationToken)
+		{
+			if (context.IncludeStack.Contains(uri))
+			{
+				throw new ConfigException(context, $"Recursive include of file {uri}");
+			}
+
+			context.IncludeStack.Push(uri);
+
+			JsonObject obj = await ReadFileAsync(uri, context, cancellationToken);
+			await MergeObjectAsync(target, obj, context, cancellationToken);
+
+			context.IncludeStack.Pop();
+		}
+
+		async Task MergeObjectAsync(object target, JsonObject obj, ConfigContext context, CancellationToken cancellationToken)
+		{
+			// If it's a BaseConfig object, read the parent config files recursively first
+			foreach ((string name, JsonNode? node) in obj)
+			{
+				if (_nameToIncludeProperty.TryGetValue(name, out PropertyInfo? propertyInfo))
+				{
+					List<ConfigInclude>? includes = node.Deserialize<List<ConfigInclude>>(context.JsonOptions);
+					if (includes != null)
+					{
+						List<ConfigInclude> targetIncludes = (List<ConfigInclude>)propertyInfo.GetValue(target)!;
+						foreach (ConfigInclude include in includes)
+						{
+							Uri uri = ConfigType.CombinePaths(context.CurrentFile, include.Path);
+							context.IncludeStack.Push(uri);
+
+							JsonObject nextObj = await ReadFileAsync(uri, context, cancellationToken);
+							await MergeObjectAsync(target, nextObj, context, cancellationToken);
+
+							context.IncludeStack.Pop();
+							targetIncludes.Add(include);
+						}
+					}
+				}
+			}
+
+			// Parse all the properties into this object
+			foreach ((string name, JsonNode? node) in obj)
+			{
+				if (_nameToProperty.TryGetValue(name, out Property? property))
+				{
+					await property.MergeAsync(target, node, context, cancellationToken);
+				}
+			}
+		}
+	}
+}
diff --git a/Engine/Source/Programs/Horde/Horde.Build/Projects/ProjectConfig.cs b/Engine/Source/Programs/Horde/Horde.Build/Projects/ProjectConfig.cs
--- a/Engine/Source/Programs/Horde/Horde.Build/Projects/ProjectConfig.cs
+++ b/Engine/Source/Programs/Horde/Horde.Build/Projects/ProjectConfig.cs
@@ -3,10 +3,12 @@
 using System;
 using System.Collections.Generic;
 using System.ComponentModel.DataAnnotations;
-using System.Text.RegularExpressions;
+using System.Diagnostics;
+using System.Text.Json.Serialization;
 using EpicGames.Core;
 using Horde.Build.Acls;
 using Horde.Build.Agents.Pools;
+using Horde.Build.Configuration;
 using Horde.Build.Streams;
 using Horde.Build.Utilities;
 
@@ -56,6 +58,7 @@ namespace Horde.Build.Projects
 	/// <summary>
 	/// Reference to configuration for a stream
 	/// </summary>
+	[DebuggerDisplay("{Id}")]
 	public class StreamConfigRef
 	{
 		/// <summary>
@@ -67,8 +70,14 @@ namespace Horde.Build.Projects
 		/// <summary>
 		/// Path to the configuration file
 		/// </summary>
-		[Required]
+		[Required, ConfigPath]
 		public string Path { get; set; } = null!;
+
+		/// <summary>
+		/// Target configuration object
+		/// </summary>
+		[JsonIgnore]
+		public StreamConfig Target { get; set; } = null!;
 	}
 
 	/// <summary>
diff --git a/Engine/Source/Programs/Horde/Horde.Build/Server/GlobalConfig.cs b/Engine/Source/Programs/Horde/Horde.Build/Server/GlobalConfig.cs
--- a/Engine/Source/Programs/Horde/Horde.Build/Server/GlobalConfig.cs
+++ b/Engine/Source/Programs/Horde/Horde.Build/Server/GlobalConfig.cs
@@ -3,7 +3,9 @@
 using System;
 using System.Collections.Generic;
 using System.ComponentModel.DataAnnotations;
+using System.Diagnostics;
 using System.Linq;
+using System.Text.Json.Serialization;
 using EpicGames.Core;
 using EpicGames.Horde.Common;
 using EpicGames.Perforce;
@@ -11,6 +13,7 @@ using EpicGames.Serialization;
 using Horde.Build.Acls;
 using Horde.Build.Agents.Pools;
 using Horde.Build.Agents.Software;
+using Horde.Build.Configuration;
 using Horde.Build.Projects;
 using Horde.Build.Storage;
 using Horde.Build.Tools;
@@ -28,6 +31,11 @@ namespace Horde.Build.Server
 	[JsonSchemaCatalog("Horde Globals", "Horde global configuration file", "globals.json")]
 	public class GlobalConfig
 	{
+		/// <summary>
+		/// Other paths to include
+		/// </summary>
+		public List<ConfigInclude> Include { get; set; } = new List<ConfigInclude>();
+
 		/// <summary>
 		/// List of projects
 		/// </summary>
@@ -137,6 +145,7 @@ namespace Horde.Build.Server
 	/// <summary>
 	/// Profile for executing compute requests
 	/// </summary>
+	[DebuggerDisplay("{Id}")]
 	public class ComputeClusterConfig
 	{
 		/// <summary>
@@ -168,6 +177,7 @@ namespace Horde.Build.Server
 	/// <summary>
 	/// References a project configuration
 	/// </summary>
+	[DebuggerDisplay("{Id}")]
 	public class ProjectConfigRef
 	{
 		/// <summary>
@@ -179,13 +189,20 @@ namespace Horde.Build.Server
 		/// <summary>
 		/// Config path for the project
 		/// </summary>
-		[Required]
+		[Required, ConfigPath]
 		public string Path { get; set; } = null!;
+
+		/// <summary>
+		/// Configuration for this project
+		/// </summary>
+		[JsonIgnore]
+		public ProjectConfig Target { get; set; } = null!;
 	}
 
 	/// <summary>
 	/// Configuration for a device platform 
 	/// </summary>
+	[DebuggerDisplay("{Id}")]
 	public class DevicePlatformConfig
 	{
 		/// <summary>
@@ -218,6 +235,7 @@ namespace Horde.Build.Server
 	/// <summary>
 	/// Selects different agent software versions by evaluating a condition
 	/// </summary>
+	[DebuggerDisplay("{Channel}")]
 	public class AgentSoftwareConfig
 	{
 		/// <summary>
@@ -443,6 +461,7 @@ namespace Horde.Build.Server
 	/// <summary>
 	/// Information about a cluster of Perforce servers. 
 	/// </summary>
+	[DebuggerDisplay("{Name}")]
 	public class PerforceCluster
 	{
 		/// <summary>
diff --git a/Engine/Source/Programs/Horde/Horde.Build/Server/RedisService.cs b/Engine/Source/Programs/Horde/Horde.Build/Server/RedisService.cs
--- a/Engine/Source/Programs/Horde/Horde.Build/Server/RedisService.cs
+++ b/Engine/Source/Programs/Horde/Horde.Build/Server/RedisService.cs
@@ -222,6 +222,17 @@ namespace Horde.Build.Server
 			redisLogger.LogInformation("Exit code {ExitCode}", _redisProcess.ExitCode);
 		}
 
+		/// <summary>
+		/// Publish a message to a channel
+		/// </summary>
+		/// <param name="channel">Channel to post to</param>
+		/// <param name="message">Message to post to the channel</param>
+		/// <param name="flags">Flags for the request</param>
+		public Task PublishAsync(RedisChannel channel, RedisValue message, CommandFlags flags = CommandFlags.None)
+		{
+			return ConnectionPool.GetDatabase().PublishAsync(channel, message, flags);
+		}
+
 		/// <summary>
 		/// Publish a message to a channel
 		/// </summary>
@@ -234,9 +245,24 @@ namespace Horde.Build.Server
 			return ConnectionPool.GetDatabase().PublishAsync(channel, message, flags);
 		}
 
+		/// <inheritdoc cref="SubscribeAsync{T}(RedisChannel{T}, Action{RedisChannel{T}, T})"/>
+		public Task<RedisSubscription> SubscribeAsync(RedisChannel channel, Action<RedisValue> callback) => SubscribeAsync(channel, (ch, x) => callback(x));
+
 		/// <inheritdoc cref="SubscribeAsync{T}(RedisChannel{T}, Action{RedisChannel{T}, T})"/>
 		public Task<RedisSubscription> SubscribeAsync<T>(RedisChannel<T> channel, Action<T> callback) => SubscribeAsync(channel, (ch, x) => callback(x));
 
+		/// <summary>
+		/// Subscribe to notifications on a channel
+		/// </summary>
+		/// <param name="channel">Channel to monitor</param>
+		/// <param name="callback">Callback for new events</param>
+		/// <returns>Subscription object</returns>
+		public async Task<RedisSubscription> SubscribeAsync(RedisChannel channel, Action<RedisChannel, RedisValue> callback)
+		{
+			IConnectionMultiplexer connection = ConnectionPool.GetConnection();
+			return await connection.SubscribeAsync(channel, callback);
+		}
+
 		/// <summary>
 		/// Subscribe to notifications on a channel
 		/// </summary>
diff --git a/Engine/Source/Programs/Horde/Horde.Build/Startup.cs b/Engine/Source/Programs/Horde/Horde.Build/Startup.cs
--- a/Engine/Source/Programs/Horde/Horde.Build/Startup.cs
+++ b/Engine/Source/Programs/Horde/Horde.Build/Startup.cs
@@ -91,11 +91,9 @@ using Horde.Build.Projects;
 using Horde.Build.Streams;
 using Horde.Build.Ugs;
 using Horde.Build.Auditing;
-using Horde.Build.Agents.Fleet.Providers;
 using Horde.Build.Server.Notices;
 using Horde.Build.Notifications.Sinks;
 using StatusCode = Grpc.Core.StatusCode;
-using Microsoft.Extensions.Caching.Memory;
 
 namespace Horde.Build
 {
@@ -362,6 +360,14 @@ namespace Horde.Build
 			services.AddSingleton<ConfigCollection>();
 			services.AddSingleton<ToolCollection>();
 
+			services.AddSingleton<IConfigSource, InMemoryConfigSource>();
+			services.AddSingleton<IConfigSource, FileConfigSource>();
+			services.AddSingleton<IConfigSource, PerforceConfigSource>();
+
+			services.AddSingleton<ConfigService>();
+			services.AddSingleton<IOptionsFactory<GlobalConfig>>(sp => sp.GetRequiredService<ConfigService>());
+			services.AddSingleton<IOptionsChangeTokenSource<GlobalConfig>>(sp => sp.GetRequiredService<ConfigService>());
+
 			// Auditing
 			services.AddSingleton<IAuditLog<AgentId>>(sp => sp.GetRequiredService<IAuditLogFactory<AgentId>>().Create("Agents.Log", "AgentId"));
 
@@ -641,6 +647,7 @@ namespace Horde.Build
 				{
 					services.AddHostedService(provider => provider.GetRequiredService<SlackNotificationSink>());
 				}
+				services.AddHostedService(provider => provider.GetRequiredService<ConfigService>());
 				services.AddHostedService(provider => provider.GetRequiredService<ConfigUpdateService>());
 				services.AddHostedService<TelemetryService>();
 				services.AddHostedService(provider => provider.GetRequiredService<DeviceService>());
@@ -683,6 +690,9 @@ namespace Horde.Build
 
 			services.AddMvc().AddJsonOptions(options => ConfigureJsonSerializer(options.JsonSerializerOptions));
 
+			ProtoBuf.Meta.RuntimeTypeModel.Default[typeof(StringId<IProject>)].SetSurrogate(typeof(StringIdProto<IProject>));
+			ProtoBuf.Meta.RuntimeTypeModel.Default[typeof(StringId<IStream>)].SetSurrogate(typeof(StringIdProto<IStream>));
+
 			services.AddControllers(options =>
 			{
 				options.InputFormatters.Add(new CbInputFormatter());
@@ -741,6 +751,7 @@ namespace Horde.Build
 
 		public static void ConfigureJsonSerializer(JsonSerializerOptions options)
 		{
+			options.AllowTrailingCommas = true;
 			options.ReadCommentHandling = JsonCommentHandling.Skip;
 			options.DefaultIgnoreCondition = JsonIgnoreCondition.WhenWritingNull;
 			options.PropertyNamingPolicy = JsonNamingPolicy.CamelCase;
diff --git a/Engine/Source/Programs/Horde/Horde.Build/Storage/StorageConfig.cs b/Engine/Source/Programs/Horde/Horde.Build/Storage/StorageConfig.cs
--- a/Engine/Source/Programs/Horde/Horde.Build/Storage/StorageConfig.cs
+++ b/Engine/Source/Programs/Horde/Horde.Build/Storage/StorageConfig.cs
@@ -3,6 +3,7 @@
 using System;
 using System.Collections.Generic;
 using System.ComponentModel.DataAnnotations;
+using System.Diagnostics;
 using EpicGames.Horde.Storage;
 using Horde.Build.Acls;
 using Horde.Build.Storage.Backends;
@@ -57,6 +58,7 @@ namespace Horde.Build.Storage
 	/// <summary>
 	/// Common settings object for different providers
 	/// </summary>
+	[DebuggerDisplay("{Id}")]
 	public class BackendConfig : IStorageBackendOptions
 	{
 		/// <summary>
@@ -104,6 +106,7 @@ namespace Horde.Build.Storage
 	/// <summary>
 	/// Configuration of a particular namespace
 	/// </summary>
+	[DebuggerDisplay("{Id}")]
 	public class NamespaceConfig
 	{
 		/// <summary>
diff --git a/Engine/Source/Programs/Horde/Horde.Build/Tools/Tool.cs b/Engine/Source/Programs/Horde/Horde.Build/Tools/Tool.cs
--- a/Engine/Source/Programs/Horde/Horde.Build/Tools/Tool.cs
+++ b/Engine/Source/Programs/Horde/Horde.Build/Tools/Tool.cs
@@ -6,6 +6,8 @@ using Horde.Build.Utilities;
 using MongoDB.Bson.Serialization.Attributes;
 using System;
 using System.Collections.Generic;
+using System.ComponentModel.DataAnnotations;
+using System.Diagnostics;
 using System.Security.Claims;
 using System.Text.Json.Serialization;
 using System.Threading.Tasks;
@@ -18,12 +20,15 @@ namespace Horde.Build.Tools
 	/// <summary>
 	/// Options for configuring a tool
 	/// </summary>
+	[DebuggerDisplay("{Id}")]
 	public class ToolOptions
 	{
 		/// <inheritdoc cref="VersionedDocument{ToolId, Tool}.Id"/>
+		[Required]
 		public ToolId Id { get; set; }
 
 		/// <inheritdoc cref="Tool.Name"/>
+		[Required]
 		public string Name { get; set; }
 
 		/// <inheritdoc cref="Tool.Description"/>
@@ -35,6 +40,15 @@ namespace Horde.Build.Tools
 		/// <inheritdoc cref="Tool.Acl"/>
 		public AclV2? Acl { get; set; }
 
+		/// <summary>
+		/// Default constructor for serialization
+		/// </summary>
+		public ToolOptions()
+		{
+			Name = String.Empty;
+			Description = String.Empty;
+		}
+
 		/// <summary>
 		/// Constructor
 		/// </summary>
diff --git a/Engine/Source/Programs/Horde/Horde.Build/Utilities/StringId.cs b/Engine/Source/Programs/Horde/Horde.Build/Utilities/StringId.cs
--- a/Engine/Source/Programs/Horde/Horde.Build/Utilities/StringId.cs
+++ b/Engine/Source/Programs/Horde/Horde.Build/Utilities/StringId.cs
@@ -11,6 +11,7 @@ using EpicGames.Core;
 using EpicGames.Serialization;
 using MongoDB.Bson.Serialization;
 using MongoDB.Bson.Serialization.Serializers;
+using ProtoBuf;
 
 namespace Horde.Build.Utilities
 {
@@ -390,4 +391,18 @@ namespace Horde.Build.Utilities
 			}
 		}
 	}
+
+	/// <summary>
+	/// Surrogate type for serializing StringId types to ProtoBuf
+	/// </summary>
+	/// <typeparam name="T"></typeparam>
+	[ProtoContract]
+	struct StringIdProto<T>
+	{
+		[ProtoMember(1)]
+		public string? Id { get; set; }
+
+		public static implicit operator StringId<T>(StringIdProto<T> source) => new StringId<T>(source.Id!);
+		public static implicit operator StringIdProto<T>(StringId<T> source) => new StringIdProto<T> { Id = source.ToString() };
+	}
 }
