Horde: Add some safeguards to prevent P4 connection errors killing the background replication task.

diff --git a/Engine/Source/Programs/Horde/Horde.Build/Perforce/PerforceServiceCache.cs b/Engine/Source/Programs/Horde/Horde.Build/Perforce/PerforceServiceCache.cs
--- a/Engine/Source/Programs/Horde/Horde.Build/Perforce/PerforceServiceCache.cs
+++ b/Engine/Source/Programs/Horde/Horde.Build/Perforce/PerforceServiceCache.cs
@@ -254,7 +254,14 @@ namespace Horde.Build.Perforce
 				// Update the background task for refreshing the list of clusters
 				if(clusterTask != null && clusterTask.IsCompleted)
 				{
-					clusters = await clusterTask;
+					try
+					{
+						clusters = await clusterTask;
+					}
+					catch (Exception ex)
+					{
+						_logger.LogError(ex, "Exception while updating cluster information: {Message}", ex.Message);
+					}
 					clusterTask = null;
 				}
 				if (clusterTask == null && clusterTimer.Elapsed > TimeSpan.FromSeconds(30.0))
@@ -313,7 +320,7 @@ namespace Horde.Build.Perforce
 							clusterState = new ClusterState();
 						}
 
-						ticker.Task = Task.Run(() => UpdateClusterAsync(ticker.ClusterName, streams, clusterState, cancellationToken));
+						ticker.Task = Task.Run(() => UpdateClusterGuardedAsync(ticker.ClusterName, streams, clusterState, cancellationToken));
 					}
 				}
 
@@ -360,6 +367,20 @@ namespace Horde.Build.Perforce
 			}
 		}
 
+		async Task<ClusterState?> UpdateClusterGuardedAsync(string clusterName, List<StreamInfo> streamInfos, ClusterState state, CancellationToken cancellationToken)
+		{
+			try
+			{
+				ClusterState? next = await UpdateClusterAsync(clusterName, streamInfos, state, cancellationToken);
+				return next;
+			}
+			catch (Exception ex)
+			{
+				_logger.LogError(ex, "Exception while updating cluster state: {Message}", ex.Message);
+				return null;
+			}
+		}
+
 		async Task<ClusterState?> UpdateClusterAsync(string clusterName, List<StreamInfo> streamInfos, ClusterState state, CancellationToken cancellationToken)
 		{
 			const int MaxChanges = 250;
