Horde: Add some basic tests for new config system.

diff --git a/Engine/Source/Programs/Horde/Horde.Build.Tests/ConfigTests.cs b/Engine/Source/Programs/Horde/Horde.Build.Tests/ConfigTests.cs
new file mode 100644
--- /dev/null
+++ b/Engine/Source/Programs/Horde/Horde.Build.Tests/ConfigTests.cs
@@ -0,0 +1,124 @@
+// Copyright Epic Games, Inc. All Rights Reserved.
+
+using System;
+using System.Collections.Generic;
+using System.Text.Json;
+using System.Threading;
+using Horde.Build.Configuration;
+using Microsoft.VisualStudio.TestTools.UnitTesting;
+using System.Threading.Tasks;
+using System.Linq;
+using EpicGames.Core;
+using Moq;
+using Horde.Build.Perforce;
+using EpicGames.Perforce;
+using Microsoft.Extensions.Logging.Abstractions;
+using System.IO;
+
+namespace Horde.Build.Tests
+{
+	[TestClass]
+	public class ConfigTests
+	{
+		class SubObject
+		{
+			public string? ValueA { get; set; }
+			public int ValueB { get; set; }
+			public SubObject? ValueC { get; set; }
+		}
+
+		class ConfigObject
+		{
+			public List<ConfigInclude> Include { get; set; } = new List<ConfigInclude>();
+			public string? TestString { get; set; }
+			public List<string> TestList { get; set; } = new List<string>();
+			public SubObject? TestObject { get; set; }
+		}
+
+		JsonSerializerOptions _jsonOptions = new JsonSerializerOptions { DefaultIgnoreCondition = System.Text.Json.Serialization.JsonIgnoreCondition.WhenWritingDefault };
+
+		[TestMethod]
+		public async Task IncludeTestAsync()
+		{
+			CancellationToken cancellationToken = CancellationToken.None;
+			InMemoryConfigSource source = new InMemoryConfigSource();
+
+			// memory:///bar
+			Uri barUri = new Uri("memory:///bar");
+			{
+				ConfigObject obj = new ConfigObject();
+				obj.TestList.Add("secondobj");
+				obj.TestObject = new SubObject { ValueB = 123, ValueC = new SubObject { ValueB = 456 } };
+
+				byte[] data2 = JsonSerializer.SerializeToUtf8Bytes(obj, _jsonOptions);
+				source.Add(barUri, data2);
+			}
+
+			// memory:///foo
+			Uri fooUri = new Uri("memory:///foo");
+			{
+				ConfigObject obj = new ConfigObject();
+				obj.Include.Add(new ConfigInclude { Path = barUri.ToString() });
+				obj.TestString = "hello";
+				obj.TestList.Add("there");
+				obj.TestList.Add("world");
+				obj.TestObject = new SubObject { ValueA = "hi", ValueC = new SubObject { ValueA = "yo" } };
+
+				byte[] json1 = JsonSerializer.SerializeToUtf8Bytes(obj, _jsonOptions);
+				source.Add(fooUri, json1);
+			}
+
+			Dictionary<string, IConfigSource> sources = new Dictionary<string, IConfigSource>();
+			sources["memory"] = source;
+
+			ConfigContext options = new ConfigContext(_jsonOptions, sources);
+
+			ConfigObject result = await ConfigType.ReadAsync<ConfigObject>(fooUri, options, cancellationToken);
+			Assert.AreEqual(result.TestString, "hello");
+			Assert.IsTrue(result.TestList.SequenceEqual(new[] { "secondobj", "there", "world" }));
+			Assert.AreEqual(result.TestObject!.ValueA, "hi");
+			Assert.AreEqual(result.TestObject!.ValueB, 123);
+			Assert.AreEqual(result.TestObject!.ValueC!.ValueA, "yo");
+		}
+
+		[TestMethod]
+		public async Task FileTestAsync()
+		{
+			CancellationToken cancellationToken = CancellationToken.None;
+
+			DirectoryReference baseDir = new DirectoryReference("test");
+			DirectoryReference.CreateDirectory(baseDir);
+
+			FileConfigSource source = new FileConfigSource(baseDir);
+
+			// file:test/foo
+			Uri fooUri = new Uri($"file:///{baseDir}\\test.json");
+
+			byte[] data;
+			{
+				ConfigObject obj = new ConfigObject();
+				obj.TestString = "hello";
+				data = JsonSerializer.SerializeToUtf8Bytes(obj, new JsonSerializerOptions { DefaultIgnoreCondition = System.Text.Json.Serialization.JsonIgnoreCondition.WhenWritingDefault });
+			}
+			await FileReference.WriteAllBytesAsync(new FileReference(fooUri.LocalPath), data, cancellationToken);
+
+			Dictionary<string, IConfigSource> sources = new Dictionary<string, IConfigSource>();
+			sources["file"] = source;
+
+			ConfigContext context = new ConfigContext(_jsonOptions, sources);
+
+			ConfigObject result = await ConfigType.ReadAsync<ConfigObject>(fooUri, context, cancellationToken);
+			Assert.AreEqual(result.TestString, "hello");
+
+			// Check it returns the same object if the timestamp hasn't changed
+			IConfigData file1 = await source.GetAsync(fooUri, cancellationToken);
+			IConfigData file2 = await source.GetAsync(fooUri, cancellationToken);
+			Assert.IsTrue(ReferenceEquals(file1, file2));
+
+			// Check it returns a new object if the timestamp HAS changed
+			await FileReference.WriteAllBytesAsync(new FileReference(fooUri.LocalPath), data, cancellationToken);
+			IConfigData file3 = await source.GetAsync(fooUri, cancellationToken);
+			Assert.IsTrue(!ReferenceEquals(file1, file3));
+		}
+	}
+}
diff --git a/Engine/Source/Programs/Horde/Horde.Build/Configuration/ConfigContext.cs b/Engine/Source/Programs/Horde/Horde.Build/Configuration/ConfigContext.cs
--- a/Engine/Source/Programs/Horde/Horde.Build/Configuration/ConfigContext.cs
+++ b/Engine/Source/Programs/Horde/Horde.Build/Configuration/ConfigContext.cs
@@ -3,6 +3,7 @@
 using System;
 using System.Collections.Generic;
 using System.Diagnostics;
+using System.Diagnostics.CodeAnalysis;
 using System.Text.Json;
 
 namespace Horde.Build.Configuration
@@ -68,14 +69,32 @@ namespace Horde.Build.Configuration
 		/// </summary>
 		/// <param name="name"></param>
 		public void AddProperty(string name)
+		{
+			if (!TryAddProperty(name, out Uri? otherFile))
+			{
+				throw new ConfigException(this, $"Property {CurrentScope}.{name} was already defined in {otherFile}.");
+			}
+		}
+
+		/// <summary>
+		/// Marks a property as defined in the current file
+		/// </summary>
+		/// <param name="name">Name of the property within the current scope</param>
+		/// <param name="otherFile">If the property is not added, the file that previously defined it</param>
+		public bool TryAddProperty(string name, [NotNullWhen(false)] out Uri? otherFile)
 		{
 			string propertyPath = $"{CurrentScope}.{name}";
 
 			Uri currentFile = CurrentFile;
-			if (!PropertyPathToFile.TryAdd(propertyPath, currentFile))
+			if (PropertyPathToFile.TryAdd(propertyPath, currentFile))
+			{
+				otherFile = null;
+				return true;
+			}
+			else
 			{
-				Uri otherFile = PropertyPathToFile[propertyPath];
-				throw new ConfigException(this, $"Property {propertyPath} was already defined in {otherFile}.");
+				otherFile = PropertyPathToFile[propertyPath];
+				return false;
 			}
 		}
 
diff --git a/Engine/Source/Programs/Horde/Horde.Build/Configuration/ConfigType.cs b/Engine/Source/Programs/Horde/Horde.Build/Configuration/ConfigType.cs
--- a/Engine/Source/Programs/Horde/Horde.Build/Configuration/ConfigType.cs
+++ b/Engine/Source/Programs/Horde/Horde.Build/Configuration/ConfigType.cs
@@ -253,10 +253,7 @@ namespace Horde.Build.Configuration
 				if (node is JsonObject obj)
 				{
 					Uri? otherFile;
-					if (!context.PropertyPathToFile.TryGetValue(Name, out otherFile))
-					{
-						context.AddProperty(Name);
-					}
+					context.TryAddProperty(Name, out otherFile);
 
 					context.EnterScope(Name);
 
