Horde: Always force a job's batches to be updated when transitioning to an error state.

diff --git a/Engine/Source/Programs/Horde/Horde.Build/Jobs/JobCollection.cs b/Engine/Source/Programs/Horde/Horde.Build/Jobs/JobCollection.cs
--- a/Engine/Source/Programs/Horde/Horde.Build/Jobs/JobCollection.cs
+++ b/Engine/Source/Programs/Horde/Horde.Build/Jobs/JobCollection.cs
@@ -703,56 +703,48 @@ namespace Horde.Build.Jobs
 					updates.Add(updateBuilder.Set(x => x.Batches[batchIdx].FinishTime, batch.FinishTime));
 				}
 			}
-			if (newError != null)
+			if (newError != null && newError.Value != batch.Error)
 			{
 				batch.Error = newError.Value;
-				updates.Add(updateBuilder.Set(x => x.Batches[batchIdx].Error, batch.Error));
-			}
-
-			// If the batch is being marked as incomplete, see if we can reschedule any of the work
-			if (newError == JobStepBatchError.Incomplete)
-			{
-				// Get the new list of retried nodes
-				List<NodeRef> retriedNodes = jobDocument.RetriedNodes ?? new List<NodeRef>();
 
 				// Check if there are any steps that need to be run again
-				bool updateState = false;
-				foreach (JobStepDocument step in batch.Steps)
+				List<NodeRef> retriedNodes = jobDocument.RetriedNodes ?? new List<NodeRef>();
+				if (newError == JobStepBatchError.Incomplete)
 				{
-					if (step.State == JobStepState.Running)
+					foreach (JobStepDocument step in batch.Steps)
 					{
-						step.State = JobStepState.Completed;
-						step.Outcome = JobStepOutcome.Failure;
-						step.Error = JobStepError.Incomplete;
-
-						if (CanRetryNode(jobDocument, batch.GroupIdx, step.NodeIdx))
+						if (step.State == JobStepState.Running)
 						{
-							step.Retry = true;
-							retriedNodes.Add(new NodeRef(batch.GroupIdx, step.NodeIdx));
-						}
+							step.State = JobStepState.Completed;
+							step.Outcome = JobStepOutcome.Failure;
+							step.Error = JobStepError.Incomplete;
 
-						updateState = true;
-					}
-					else if (step.State == JobStepState.Ready || step.State == JobStepState.Waiting)
-					{
-						if (CanRetryNode(jobDocument, batch.GroupIdx, step.NodeIdx))
-						{
-							retriedNodes.Add(new NodeRef(batch.GroupIdx, step.NodeIdx));
+							if (CanRetryNode(jobDocument, batch.GroupIdx, step.NodeIdx))
+							{
+								step.Retry = true;
+								retriedNodes.Add(new NodeRef(batch.GroupIdx, step.NodeIdx));
+							}
 						}
-						else
+						else if (step.State == JobStepState.Ready || step.State == JobStepState.Waiting)
 						{
-							step.State = JobStepState.Skipped;
+							if (CanRetryNode(jobDocument, batch.GroupIdx, step.NodeIdx))
+							{
+								retriedNodes.Add(new NodeRef(batch.GroupIdx, step.NodeIdx));
+							}
+							else
+							{
+								step.State = JobStepState.Skipped;
+							}
 						}
-						updateState = true;
 					}
 				}
 
-				// Update the steps
-				if (updateState)
-				{
-					updates.Clear();
-					UpdateBatches(jobDocument, graph, updates, _logger);
+				// Force an update of all batches in the job. This will fail or reschedule any nodes that can no longer be executed in this batch.
+				updates.Clear();
+				UpdateBatches(jobDocument, graph, updates, _logger);
 
+				if(retriedNodes.Count > 0)
+				{
 					jobDocument.RetriedNodes = retriedNodes;
 					updates.Add(updateBuilder.Set(x => x.RetriedNodes, jobDocument.RetriedNodes));
 				}
@@ -786,7 +778,6 @@ namespace Horde.Build.Jobs
 						JobStepDocument step = batch.Steps[stepIdx];
 						if (step.Id == stepId)
 						{
-
 							// Update the request abort status
 							if (newAbortRequested != null && step.AbortRequested == false)
 							{
