Horde: Fix config file test execution on platforms with slashes for path separators.

diff --git a/Engine/Source/Programs/Horde/Horde.Build.Tests/ConfigTests.cs b/Engine/Source/Programs/Horde/Horde.Build.Tests/ConfigTests.cs
--- a/Engine/Source/Programs/Horde/Horde.Build.Tests/ConfigTests.cs
+++ b/Engine/Source/Programs/Horde/Horde.Build.Tests/ConfigTests.cs
@@ -92,7 +92,7 @@ namespace Horde.Build.Tests
 			FileConfigSource source = new FileConfigSource(baseDir);
 
 			// file:test/foo
-			Uri fooUri = new Uri($"file:///{baseDir}\\test.json");
+			Uri fooUri = new Uri($"file:///{FileReference.Combine(baseDir, "test.json")}");
 
 			byte[] data;
 			{
