Horde: Fix static analysis warnings due to fields not being disposed.

diff --git a/Engine/Source/Programs/Horde/Horde.Agent.Tests/AwsLambdaListenerCommandTest.cs b/Engine/Source/Programs/Horde/Horde.Agent.Tests/AwsLambdaListenerCommandTest.cs
--- a/Engine/Source/Programs/Horde/Horde.Agent.Tests/AwsLambdaListenerCommandTest.cs
+++ b/Engine/Source/Programs/Horde/Horde.Agent.Tests/AwsLambdaListenerCommandTest.cs
@@ -16,7 +16,7 @@ using Microsoft.VisualStudio.TestTools.UnitTesting;
 namespace Horde.Agent.Tests
 {
 	[TestClass]
-	public class AwsLambdaListenerCommandTest
+	public sealed class AwsLambdaListenerCommandTest : IDisposable
 	{
 		class TestTaskHandler : LeaseHandler<TestTask>
 		{
@@ -38,6 +38,11 @@ namespace Horde.Agent.Tests
 			_serviceCollection.AddSingleton<ISessionFactoryService>(sp => new FakeServerSessionFactory(_fakeServer));
 		}
 
+		public void Dispose()
+		{
+			_fakeServer.DisposeAsync().AsTask().Wait();
+		}
+
 		[TestMethod]
 		public async Task ShutdownIfNoLeasesArriveWithinMaxWaitTime()
 		{
diff --git a/Engine/Source/Programs/Horde/Horde.Agent.Tests/Execution/WorkspaceExecutorTest.cs b/Engine/Source/Programs/Horde/Horde.Agent.Tests/Execution/WorkspaceExecutorTest.cs
--- a/Engine/Source/Programs/Horde/Horde.Agent.Tests/Execution/WorkspaceExecutorTest.cs
+++ b/Engine/Source/Programs/Horde/Horde.Agent.Tests/Execution/WorkspaceExecutorTest.cs
@@ -1,5 +1,6 @@
 // Copyright Epic Games, Inc. All Rights Reserved.
 
+using System;
 using System.Collections.Generic;
 using System.IO;
 using System.Threading;
@@ -13,7 +14,7 @@ using Microsoft.VisualStudio.TestTools.UnitTesting;
 namespace Horde.Agent.Tests.Execution;
 
 [TestClass]
-public class WorkspaceExecutorTest
+public sealed class WorkspaceExecutorTest : IDisposable
 {
 	private const string StreamId = "foo-main";
 	private const string JobId = "job1";
@@ -47,6 +48,12 @@ public class WorkspaceExecutorTest
 		
 		_executor = new (_session, JobId, "batch1", AgentType, null, _workspace, null!);
 	}
+
+	public void Dispose()
+	{
+		_server.DisposeAsync().AsTask().Wait();
+		_session.DisposeAsync().AsTask().Wait();
+	}
 	
 	[TestMethod]
 	public async Task RegularWorkspace()
