Horde: Add a set of perforce connection settings in the server configuration file used for reading configuration data.

Previous implementation relied on cluster info defined in the config file itself, preventing it from being able to bootstrap itself.

diff --git a/Engine/Source/Programs/Horde/Horde.Build/Configuration/ConfigData.cs b/Engine/Source/Programs/Horde/Horde.Build/Configuration/ConfigData.cs
--- a/Engine/Source/Programs/Horde/Horde.Build/Configuration/ConfigData.cs
+++ b/Engine/Source/Programs/Horde/Horde.Build/Configuration/ConfigData.cs
@@ -9,10 +9,14 @@ using System.Threading;
 using System.Threading.Tasks;
 using EpicGames.Core;
 using EpicGames.Perforce;
-using Horde.Build.Perforce;
+using Horde.Build.Utilities;
+using Microsoft.Extensions.Logging;
+using Microsoft.Extensions.Options;
 
 namespace Horde.Build.Configuration
 {
+	using PerforceConnectionId = StringId<PerforceConnectionSettings>;
+
 	/// <summary>
 	/// Accessor for a specific revision of a config file. Provides metadata about the current revision, and allows reading its data.
 	/// </summary>
@@ -224,26 +228,29 @@ namespace Horde.Build.Configuration
 		/// <summary>
 		/// Name of the scheme for this source
 		/// </summary>
-		public const string Scheme = "p4-cluster";
+		public const string Scheme = "perforce";
 
 		/// <inheritdoc/>
 		string IConfigSource.Scheme => Scheme;
 
-		readonly IPerforceService _perforceService;
+		readonly IOptionsMonitor<ServerSettings> _settings;
+		readonly ILogger _logger;
 
 		/// <summary>
 		/// Constructor
 		/// </summary>
-		/// <param name="perforceService">Perforce service instance</param>
-		public PerforceConfigSource(IPerforceService perforceService)
+		/// <param name="settings">Global settings</param>
+		/// <param name="logger">Logger instance</param>
+		public PerforceConfigSource(IOptionsMonitor<ServerSettings> settings, ILogger<PerforceConfigSource> logger)
 		{
-			_perforceService = perforceService;
+			_settings = settings;
+			_logger = logger;
 		}
 
 		/// <inheritdoc/>
 		public async Task<IConfigData> GetAsync(Uri uri, CancellationToken cancellationToken)
 		{
-			using (IPerforceConnection perforce = await _perforceService.ConnectAsync(uri.Host, null, cancellationToken))
+			using (IPerforceConnection perforce = await ConnectAsync(uri.Host))
 			{
 				List<FStatRecord> files = await perforce.FStatAsync(FStatOptions.ShortenOutput, uri.AbsolutePath, cancellationToken).ToListAsync(cancellationToken);
 				files.RemoveAll(x => x.HeadAction == FileAction.Delete || x.HeadAction == FileAction.MoveDelete);
@@ -259,12 +266,38 @@ namespace Horde.Build.Configuration
 
 		async ValueTask<ReadOnlyMemory<byte>> ReadAsync(Uri uri, int change, CancellationToken cancellationToken)
 		{
-			using (IPerforceConnection perforce = await _perforceService.ConnectAsync(uri.Host, null, cancellationToken))
+			using (IPerforceConnection perforce = await ConnectAsync(uri.Host))
 			{
 				PerforceResponse<PrintRecord<byte[]>> response = await perforce.TryPrintDataAsync($"{uri.AbsolutePath}@{change}", cancellationToken);
 				response.EnsureSuccess();
 				return response.Data.Contents!;
 			}
 		}
+
+		async Task<IPerforceConnection> ConnectAsync(string host)
+		{
+			ServerSettings settings = _settings.CurrentValue;
+
+			PerforceConnectionId connectionId = new PerforceConnectionId();
+			if (!String.IsNullOrEmpty(host))
+			{
+				connectionId = new PerforceConnectionId(host);
+			}
+
+			PerforceConnectionSettings? connectionSettings = settings.Perforce.FirstOrDefault(x => x.Id == connectionId);
+			if (connectionSettings == null)
+			{
+				if (connectionId == PerforceConnectionSettings.Default)
+				{
+					connectionSettings = new PerforceConnectionSettings();
+				}
+				else
+				{
+					throw new InvalidOperationException($"No Perforce connection settings defined for '{connectionId}'.");
+				}
+			}
+
+			return await PerforceConnection.CreateAsync(connectionSettings.ToPerforceSettings(), _logger);
+		}
 	}
 }
diff --git a/Engine/Source/Programs/Horde/Horde.Build/Configuration/ConfigType.cs b/Engine/Source/Programs/Horde/Horde.Build/Configuration/ConfigType.cs
--- a/Engine/Source/Programs/Horde/Horde.Build/Configuration/ConfigType.cs
+++ b/Engine/Source/Programs/Horde/Horde.Build/Configuration/ConfigType.cs
@@ -124,7 +124,7 @@ namespace Horde.Build.Configuration
 				}
 				else
 				{
-					return new Uri($"{PerforceConfigSource.Scheme}://{PerforceCluster.DefaultName}{path}");
+					return new Uri($"{PerforceConfigSource.Scheme}://{PerforceConnectionSettings.Default}{path}");
 				}
 			}
 			return new Uri(baseUri, path);
diff --git a/Engine/Source/Programs/Horde/Horde.Build/ServerSettings.cs b/Engine/Source/Programs/Horde/Horde.Build/ServerSettings.cs
--- a/Engine/Source/Programs/Horde/Horde.Build/ServerSettings.cs
+++ b/Engine/Source/Programs/Horde/Horde.Build/ServerSettings.cs
@@ -1,20 +1,22 @@
 // Copyright Epic Games, Inc. All Rights Reserved.
 
 using System;
-using System.ComponentModel.DataAnnotations;
+using System.Collections.Generic;
 using System.Linq;
-using EpicGames.Horde.Common;
 using EpicGames.Horde.Storage;
 using EpicGames.Horde.Storage.Nodes;
+using EpicGames.Perforce;
 using Horde.Build.Agents.Fleet;
-using Horde.Build.Agents.Software;
+using Horde.Build.Server;
 using Horde.Build.Storage.Backends;
 using Horde.Build.Utilities;
+using MongoDB.Driver.Core.Configuration;
 using Serilog.Events;
-using TimeZoneConverter;
 
 namespace Horde.Build
 {
+	using PerforceConnectionId = StringId<PerforceConnectionSettings>;
+
 	/// <summary>
 	/// Types of storage backend to use
 	/// </summary>
@@ -280,17 +282,17 @@ namespace Horde.Build
 		/// Issuer for tokens from the auth provider
 		/// </summary>
 		public AuthMethod AuthMethod { get; set; } = AuthMethod.Anonymous;
-		
+
 		/// <summary>
 		/// Issuer for tokens from the auth provider
 		/// </summary>
 		public string? OidcAuthority { get; set; }
-		
+
 		/// <summary>
 		/// Client id for the OIDC authority
 		/// </summary>
 		public string? OidcClientId { get; set; }
-		
+
 		/// <summary>
 		/// Client secret for the OIDC authority
 		/// </summary>
@@ -305,22 +307,22 @@ namespace Horde.Build
 		/// OpenID Connect scopes to request when signing in
 		/// </summary>
 		public string[] OidcRequestedScopes { get; set; } = { "profile", "email", "openid" };
-		
+
 		/// <summary>
 		/// List of fields in /userinfo endpoint to try map to the standard name claim (see System.Security.Claims.ClaimTypes.Name)
 		/// </summary>
 		public string[] OidcClaimNameMapping { get; set; } = { "preferred_username", "email" };
-		
+
 		/// <summary>
 		/// List of fields in /userinfo endpoint to try map to the standard email claim (see System.Security.Claims.ClaimTypes.Email)
 		/// </summary>
 		public string[] OidcClaimEmailMapping { get; set; } = { "email" };
-		
+
 		/// <summary>
 		/// List of fields in /userinfo endpoint to try map to the Horde user claim (see HordeClaimTypes.User)
 		/// </summary>
 		public string[] OidcClaimHordeUserMapping { get; set; } = { "preferred_username", "email" };
-		
+
 		/// <summary>
 		/// List of fields in /userinfo endpoint to try map to the Horde Perforce user claim (see HordeClaimTypes.PerforceUser)
 		/// </summary>
@@ -387,7 +389,7 @@ namespace Horde.Build
 		/// See format at https://stackexchange.github.io/StackExchange.Redis/Configuration.html
 		/// </summary>
 		public string? RedisConnectionConfig { get; set; }
-		
+
 		/// <summary>
 		/// Type of write cache to use in log service
 		/// Currently Supported: "InMemory" or "Redis"
@@ -423,7 +425,7 @@ namespace Horde.Build
 		/// Default fleet manager to use (when not specified by pool)
 		/// </summary>
 		public FleetManagerType FleetManagerV2 { get; set; } = FleetManagerType.NoOp;
-		
+
 		/// <summary>
 		/// Config for the fleet manager (serialized JSON)
 		/// </summary>
@@ -473,7 +475,7 @@ namespace Horde.Build
 		/// Channel to send stream notification update failures to
 		/// </summary>
 		public string? UpdateStreamsNotificationChannel { get; set; }
-		
+
 		/// <summary>
 		/// Slack channel to send job related notifications to. Multiple channels can be specified, separated by ;
 		/// </summary>
@@ -563,7 +565,7 @@ namespace Horde.Build
 		/// Scale-out cooldown for auto-scaling agent pools (in seconds). Can be overridden by per-pool settings.
 		/// </summary>
 		public int AgentPoolScaleOutCooldownSeconds { get; set; } = 60; // 1 min
-		
+
 		/// <summary>
 		/// Scale-in cooldown for auto-scaling agent pools (in seconds). Can be overridden by per-pool settings.
 		/// </summary>
@@ -592,6 +594,11 @@ namespace Horde.Build
 		/// </summary>
 		public string ConfigPath { get; set; } = "Defaults/globals.json";
 
+		/// <summary>
+		/// Perforce connections for use by the Horde server (not agents)
+		/// </summary>
+		public List<PerforceConnectionSettings> Perforce { get; set; } = new List<PerforceConnectionSettings>();
+
 		/// <summary>
 		/// Settings for the storage service
 		/// </summary>
@@ -661,4 +668,57 @@ namespace Horde.Build
 			}
 		}
 	}
+
+	/// <summary>
+	/// Perforce connection information for use by the Horde server (for reading config files, etc...)
+	/// </summary>
+	public class PerforceConnectionSettings
+	{
+		/// <summary>
+		/// Identifier for the default perforce connection profile
+		/// </summary>
+		public static PerforceConnectionId Default { get; } = new PerforceConnectionId("default");
+
+		/// <summary>
+		/// Identifier for this server
+		/// </summary>
+		public PerforceConnectionId Id { get; set; } = Default;
+
+		/// <summary>
+		/// Server and port
+		/// </summary>
+		public string? ServerAndPort { get; set; }
+
+		/// <summary>
+		/// Credentials for the server
+		/// </summary>
+		public PerforceCredentials? Credentials { get; set; }
+
+		/// <summary>
+		/// Create a <see cref="PerforceSettings"/> object with these settings as overrides
+		/// </summary>
+		/// <returns>New perforce settings object</returns>
+		public PerforceSettings ToPerforceSettings()
+		{
+			PerforceSettings settings = new PerforceSettings(PerforceSettings.Default);
+			settings.PreferNativeClient = true;
+
+			if (!String.IsNullOrEmpty(ServerAndPort))
+			{
+				settings.ServerAndPort = ServerAndPort;
+			}
+			if (Credentials != null)
+			{
+				if (!String.IsNullOrEmpty(Credentials.UserName))
+				{
+					settings.UserName = Credentials.UserName;
+				}
+				if (!String.IsNullOrEmpty(Credentials.Password))
+				{
+					settings.Password = Credentials.Password;
+				}
+			}
+			return settings;
+		}
+	}
 }
