Horde: Add tests for config file schema generation, and fix exception generating schema with namespace ids.

diff --git a/Engine/Source/Programs/Horde/Horde.Build.Tests/SchemaControllerTest.cs b/Engine/Source/Programs/Horde/Horde.Build.Tests/SchemaControllerTest.cs
new file mode 100644
--- /dev/null
+++ b/Engine/Source/Programs/Horde/Horde.Build.Tests/SchemaControllerTest.cs
@@ -0,0 +1,51 @@
+// Copyright Epic Games, Inc. All Rights Reserved.
+
+using System;
+using System.Collections.Generic;
+using System.Net;
+using System.Net.Http;
+using System.Text.Json;
+using System.Text.Json.Serialization;
+using System.Threading.Tasks;
+using Horde.Build.Jobs.Artifacts;
+using Microsoft.VisualStudio.TestTools.UnitTesting;
+
+namespace Horde.Build.Tests
+{
+	[TestClass]
+	public class SchemaControllerTest : ControllerIntegrationTest
+	{
+		class SchemaCatalog
+		{
+			public List<SchemaCatalogItem> Schemas { get; set; } = new List<SchemaCatalogItem>();
+		}
+
+		class SchemaCatalogItem
+		{
+			public string? Name { get; set; }
+			public Uri Url { get; set; } = null!;
+		}
+
+		[TestMethod]
+		public async Task TestCatalogAsync()
+		{
+			await GetFixture();
+
+			// Fetch the catalog
+			SchemaCatalog catalog;
+			using (HttpResponseMessage response = await Client.GetAsync(new Uri("http://localhost/api/v1/schema/catalog.json")))
+			{
+				response.EnsureSuccessStatusCode();
+				byte[] data = await response.Content.ReadAsByteArrayAsync();
+				catalog = JsonSerializer.Deserialize<SchemaCatalog>(data, new JsonSerializerOptions { PropertyNameCaseInsensitive = true })!;
+			}
+
+			// Fetch all the individual files
+			foreach (SchemaCatalogItem schema in catalog.Schemas)
+			{
+				using HttpResponseMessage response = await Client.GetAsync(schema.Url);
+				response.EnsureSuccessStatusCode();
+			}
+		}
+	}
+}
diff --git a/Engine/Source/Programs/Shared/EpicGames.Horde/Storage/NamespaceId.cs b/Engine/Source/Programs/Shared/EpicGames.Horde/Storage/NamespaceId.cs
--- a/Engine/Source/Programs/Shared/EpicGames.Horde/Storage/NamespaceId.cs
+++ b/Engine/Source/Programs/Shared/EpicGames.Horde/Storage/NamespaceId.cs
@@ -12,6 +12,7 @@ namespace EpicGames.Horde.Storage
 	/// <summary>
 	/// Identifier for a storage namespace
 	/// </summary>
+	[JsonSchemaString]
 	[JsonConverter(typeof(NamespaceIdJsonConverter))]
 	[TypeConverter(typeof(NamespaceIdTypeConverter))]
 	public struct NamespaceId : IEquatable<NamespaceId>
