Horde: Add tests for storage service.

diff --git a/Engine/Source/Programs/Horde/Horde.Build.Tests/StorageServiceTests.cs b/Engine/Source/Programs/Horde/Horde.Build.Tests/StorageServiceTests.cs
new file mode 100644
--- /dev/null
+++ b/Engine/Source/Programs/Horde/Horde.Build.Tests/StorageServiceTests.cs
@@ -0,0 +1,62 @@
+// Copyright Epic Games, Inc. All Rights Reserved.
+
+using System;
+using System.Collections.Generic;
+using System.Linq;
+using System.Text;
+using System.Threading.Tasks;
+using EpicGames.Core;
+using EpicGames.Horde.Storage;
+using EpicGames.Serialization;
+using Horde.Build.Configuration;
+using Horde.Build.Storage;
+using Microsoft.Extensions.DependencyInjection;
+using Microsoft.VisualStudio.TestTools.UnitTesting;
+
+namespace Horde.Build.Tests
+{
+	[TestClass]
+	public class StorageServiceTests : TestSetup
+	{
+		[TestMethod]
+		public async Task BlobCollectionTest()
+		{
+			ConfigUpdateService configUpdateService = ServiceProvider.GetRequiredService<ConfigUpdateService>();
+			await configUpdateService.StartAsync(default);
+
+			StorageService storageService = ServiceProvider.GetRequiredService<StorageService>();
+			IStorageClient client = await storageService.GetClientAsync(new NamespaceId("memory"), default);
+
+			List<BundleType> types = new List<BundleType>();
+			types.Add(new BundleType(Guid.Parse("{11C2D886-3349-4164-946F-E9D10BD12E3D}"), 0));
+			types.Add(new BundleType(Guid.Parse("{6CB3A005-26BA-4787-86D2-793ED13771CB}"), 0));
+
+			byte[] data1 = new byte[] { 1, 2, 3 };
+			IoHash hash1 = IoHash.Compute(data1);
+
+			byte[] data2 = new byte[] { 4, 5, 6 };
+			IoHash hash2 = IoHash.Compute(data2);
+
+			List<BundleExport> exports = new List<BundleExport>();
+
+			exports.Add(new BundleExport(0, hash1, data1.Length, Array.Empty<int>(), "foo"));
+			exports.Add(new BundleExport(0, hash1, data1.Length, Array.Empty<int>(), "foo"));
+			exports.Add(new BundleExport(0, hash2, data2.Length, Array.Empty<int>(), "bar"));
+
+			BundleHeader header = new BundleHeader(BundleCompressionFormat.None, types, Array.Empty<BundleImport>(), exports, Array.Empty<BundlePacket>());
+			Bundle bundle = new Bundle(header, Array.Empty<ReadOnlyMemory<byte>>());
+			BlobLocator locator = await client.WriteBundleAsync(bundle);
+
+			List<NodeHandle> handles;
+			
+			handles = await client.FindNodesAsync("foo").ToListAsync();
+			Assert.AreEqual(2, handles.Count);
+			Assert.AreEqual(new NodeLocator(locator, 0), handles[0].Locator);
+			Assert.AreEqual(new NodeLocator(locator, 1), handles[1].Locator);
+
+			handles = await client.FindNodesAsync("bar").ToListAsync();
+			Assert.AreEqual(1, handles.Count);
+			Assert.AreEqual(new NodeLocator(locator, 2), handles[0].Locator);
+		}
+	}
+}
