Horde: Prevent steps being retried when shelved change is missing.

diff --git a/Engine/Source/Programs/Horde/Horde.Build/Jobs/JobCollection.cs b/Engine/Source/Programs/Horde/Horde.Build/Jobs/JobCollection.cs
--- a/Engine/Source/Programs/Horde/Horde.Build/Jobs/JobCollection.cs
+++ b/Engine/Source/Programs/Horde/Horde.Build/Jobs/JobCollection.cs
@@ -707,34 +707,33 @@ namespace Horde.Build.Jobs
 			{
 				batch.Error = newError.Value;
 
+				bool allowRetrying = (newError == JobStepBatchError.Incomplete);
+
 				// Check if there are any steps that need to be run again
 				List<NodeRef> retriedNodes = jobDocument.RetriedNodes ?? new List<NodeRef>();
-				if (newError == JobStepBatchError.Incomplete)
+				foreach (JobStepDocument step in batch.Steps)
 				{
-					foreach (JobStepDocument step in batch.Steps)
+					if (step.State == JobStepState.Running)
 					{
-						if (step.State == JobStepState.Running)
-						{
-							step.State = JobStepState.Completed;
-							step.Outcome = JobStepOutcome.Failure;
-							step.Error = JobStepError.Incomplete;
+						step.State = JobStepState.Completed;
+						step.Outcome = JobStepOutcome.Failure;
+						step.Error = JobStepError.Incomplete;
 
-							if (CanRetryNode(jobDocument, batch.GroupIdx, step.NodeIdx))
-							{
-								step.Retry = true;
-								retriedNodes.Add(new NodeRef(batch.GroupIdx, step.NodeIdx));
-							}
+						if (allowRetrying && CanRetryNode(jobDocument, batch.GroupIdx, step.NodeIdx))
+						{
+							step.Retry = true;
+							retriedNodes.Add(new NodeRef(batch.GroupIdx, step.NodeIdx));
 						}
-						else if (step.State == JobStepState.Ready || step.State == JobStepState.Waiting)
+					}
+					else if (step.State == JobStepState.Ready || step.State == JobStepState.Waiting)
+					{
+						if (allowRetrying && CanRetryNode(jobDocument, batch.GroupIdx, step.NodeIdx))
 						{
-							if (CanRetryNode(jobDocument, batch.GroupIdx, step.NodeIdx))
-							{
-								retriedNodes.Add(new NodeRef(batch.GroupIdx, step.NodeIdx));
-							}
-							else
-							{
-								step.State = JobStepState.Skipped;
-							}
+							retriedNodes.Add(new NodeRef(batch.GroupIdx, step.NodeIdx));
+						}
+						else
+						{
+							step.State = JobStepState.Skipped;
 						}
 					}
 				}
