Horde: Fix test failures due to repeated registration of surrogate types for protobuf serialization.

diff --git a/Engine/Source/Programs/Horde/Horde.Build/Startup.cs b/Engine/Source/Programs/Horde/Horde.Build/Startup.cs
--- a/Engine/Source/Programs/Horde/Horde.Build/Startup.cs
+++ b/Engine/Source/Programs/Horde/Horde.Build/Startup.cs
@@ -104,6 +104,12 @@ namespace Horde.Build
 
 	class Startup
 	{
+		static Startup()
+		{
+			ProtoBuf.Meta.RuntimeTypeModel.Default[typeof(StringId<IProject>)].SetSurrogate(typeof(StringIdProto<IProject>));
+			ProtoBuf.Meta.RuntimeTypeModel.Default[typeof(StringId<IStream>)].SetSurrogate(typeof(StringIdProto<IStream>));
+		}
+
 		class GrpcExceptionInterceptor : Interceptor
 		{
 			readonly ILogger<GrpcExceptionInterceptor> _logger;
@@ -690,9 +696,6 @@ namespace Horde.Build
 
 			services.AddMvc().AddJsonOptions(options => ConfigureJsonSerializer(options.JsonSerializerOptions));
 
-			ProtoBuf.Meta.RuntimeTypeModel.Default[typeof(StringId<IProject>)].SetSurrogate(typeof(StringIdProto<IProject>));
-			ProtoBuf.Meta.RuntimeTypeModel.Default[typeof(StringId<IStream>)].SetSurrogate(typeof(StringIdProto<IStream>));
-
 			services.AddControllers(options =>
 			{
 				options.InputFormatters.Add(new CbInputFormatter());
