Horde: Fix exception when querying a config data source that no longer exists.

diff --git a/Engine/Source/Programs/Horde/Horde.Build/Configuration/ConfigService.cs b/Engine/Source/Programs/Horde/Horde.Build/Configuration/ConfigService.cs
--- a/Engine/Source/Programs/Horde/Horde.Build/Configuration/ConfigService.cs
+++ b/Engine/Source/Programs/Horde/Horde.Build/Configuration/ConfigService.cs
@@ -286,11 +286,6 @@ namespace Horde.Build.Configuration
 
 					foreach (StreamConfigRef streamConfigRef in projectConfig.Streams)
 					{
-						if (streamConfigRef.Id.ToString() == "plastic-private-plastic")
-						{
-							continue;
-						}
-
 						context.PropertyPathToFile.Clear();
 
 						StreamConfig streamConfig = await ConfigType.ReadAsync<StreamConfig>(new Uri(streamConfigRef.Path), context, cancellationToken);
@@ -322,8 +317,8 @@ namespace Horde.Build.Configuration
 		{
 			foreach ((Uri uri, string version) in snapshot.Dependencies)
 			{
-				IConfigSource? source = _sources[uri.Scheme];
-				if (source == null)
+				IConfigSource? source;
+				if (!_sources.TryGetValue(uri.Scheme, out source))
 				{
 					return true;
 				}
