Horde: Prevent retrying nodes in a batch that cannot be re-run.

diff --git a/Engine/Source/Programs/Horde/Horde.Build.Tests/JobCollectionTests.cs b/Engine/Source/Programs/Horde/Horde.Build.Tests/JobCollectionTests.cs
--- a/Engine/Source/Programs/Horde/Horde.Build.Tests/JobCollectionTests.cs
+++ b/Engine/Source/Programs/Horde/Horde.Build.Tests/JobCollectionTests.cs
@@ -327,5 +327,26 @@ namespace Horde.Build.Tests
 			Assert.AreEqual(JobStepState.Completed, job.Batches[2].Steps[0].State);
 			Assert.AreEqual(JobStepError.Incomplete, job.Batches[2].Steps[0].Error);
 		}
+
+		[TestMethod]
+		public async Task UnknownShelfAsync()
+		{
+			Mock<ITemplate> templateMock = new Mock<ITemplate>(MockBehavior.Strict);
+			templateMock.SetupGet(x => x.InitialAgentType).Returns((string?)null);
+
+			IGraph baseGraph = await GraphCollection.AddAsync(templateMock.Object, null);
+
+			CreateJobOptions options = new CreateJobOptions();
+			options.Arguments.Add("-Target=Step 1");
+			options.Arguments.Add("-Target=Step 3");
+
+			IJob job = await JobCollection.AddAsync(JobId.GenerateNewId(), new StreamId("ue4-main"), new TemplateId("test-build"), ContentHash.SHA1("hello"), baseGraph, "Test job", 123, 123, options);
+			Assert.AreEqual(1, job.Batches.Count);
+
+			job = await StartBatch(job, baseGraph, 0);
+			job = Deref(await JobCollection.TryUpdateBatchAsync(job, baseGraph, job.Batches[0].Id, null, JobStepBatchState.Complete, JobStepBatchError.UnknownShelf));
+
+			Assert.AreEqual(1, job.Batches.Count);
+		}
 	}
 }
diff --git a/Engine/Source/Programs/Horde/Horde.Build/Jobs/JobCollection.cs b/Engine/Source/Programs/Horde/Horde.Build/Jobs/JobCollection.cs
--- a/Engine/Source/Programs/Horde/Horde.Build/Jobs/JobCollection.cs
+++ b/Engine/Source/Programs/Horde/Horde.Build/Jobs/JobCollection.cs
@@ -6,7 +6,6 @@ using System.Diagnostics.CodeAnalysis;
 using System.IO;
 using System.Linq;
 using System.Threading.Tasks;
-using System.Xml.Linq;
 using EpicGames.Core;
 using Horde.Build.Acls;
 using Horde.Build.Agents;
@@ -707,9 +706,10 @@ namespace Horde.Build.Jobs
 			{
 				batch.Error = newError.Value;
 
-				bool allowRetrying = (newError == JobStepBatchError.Incomplete);
+				// Only allow retrying nodes that may succeed if run again
+				bool allowRetrying = !IsFatalBatchError(newError.Value);
 
-				// Check if there are any steps that need to be run again
+				// Update the state of the nodes
 				List<NodeRef> retriedNodes = jobDocument.RetriedNodes ?? new List<NodeRef>();
 				foreach (JobStepDocument step in batch.Steps)
 				{
@@ -1324,6 +1324,10 @@ namespace Horde.Build.Jobs
 					{
 						failedNodes.Remove(node);
 					}
+					else if (batch.State == JobStepBatchState.Complete && IsFatalBatchError(batch.Error))
+					{
+						failedNodes.Add(node);
+					}
 					else if (step.State == JobStepState.Skipped && (node.InputDependencies.Any(x => failedNodes.Contains(graph.GetNode(x))) || !CanRetryNode(job, batch.GroupIdx, step.NodeIdx)))
 					{
 						failedNodes.Add(node);
@@ -1568,6 +1572,14 @@ namespace Horde.Build.Jobs
 			}
 		}
 
+		/// <summary>
+		/// Test whether the nodes in a batch can still be run
+		/// </summary>
+		static bool IsFatalBatchError(JobStepBatchError error)
+		{
+			return error != JobStepBatchError.None && error != JobStepBatchError.Incomplete;
+		}
+
 		/// <summary>
 		/// Tests whether a node can be retried again
 		/// </summary>
