Horde: Add delays to confile file refresh tests, to ensure OS granularity on timestamps is large enough.

diff --git a/Engine/Source/Programs/Horde/Horde.Build.Tests/ConfigTests.cs b/Engine/Source/Programs/Horde/Horde.Build.Tests/ConfigTests.cs
--- a/Engine/Source/Programs/Horde/Horde.Build.Tests/ConfigTests.cs
+++ b/Engine/Source/Programs/Horde/Horde.Build.Tests/ConfigTests.cs
@@ -112,10 +112,12 @@ namespace Horde.Build.Tests
 
 			// Check it returns the same object if the timestamp hasn't changed
 			IConfigData file1 = await source.GetAsync(fooUri, cancellationToken);
+			await Task.Delay(TimeSpan.FromSeconds(1));
 			IConfigData file2 = await source.GetAsync(fooUri, cancellationToken);
 			Assert.IsTrue(ReferenceEquals(file1, file2));
 
 			// Check it returns a new object if the timestamp HAS changed
+			await Task.Delay(TimeSpan.FromSeconds(1));
 			await FileReference.WriteAllBytesAsync(new FileReference(fooUri.LocalPath), data, cancellationToken);
 			IConfigData file3 = await source.GetAsync(fooUri, cancellationToken);
 			Assert.IsTrue(!ReferenceEquals(file1, file3));
