Horde: Batch queries to Perforce when looking for config file changes.

diff --git a/Engine/Source/Programs/Horde/Horde.Build/Configuration/ConfigData.cs b/Engine/Source/Programs/Horde/Horde.Build/Configuration/ConfigData.cs
--- a/Engine/Source/Programs/Horde/Horde.Build/Configuration/ConfigData.cs
+++ b/Engine/Source/Programs/Horde/Horde.Build/Configuration/ConfigData.cs
@@ -53,10 +53,29 @@ namespace Horde.Build.Configuration
 		/// <summary>
 		/// Reads a config file from this source
 		/// </summary>
-		/// <param name="uri">Location of the config file</param>
+		/// <param name="uris">Locations of the config files to query</param>
 		/// <param name="cancellationToken">Cancellation token for the operation</param>
 		/// <returns>Config file data</returns>
-		Task<IConfigData> GetAsync(Uri uri, CancellationToken cancellationToken);
+		Task<IConfigData[]> GetAsync(Uri[] uris, CancellationToken cancellationToken);
+	}
+
+	/// <summary>
+	/// Extension methods for <see cref="IConfigSource"/>
+	/// </summary>
+	public static class ConfigSource
+	{
+		/// <summary>
+		/// Gets a single config file from a source
+		/// </summary>
+		/// <param name="source">Source to query</param>
+		/// <param name="uri">Location of the config file to query</param>
+		/// <param name="cancellationToken">Cancellation token for the operation</param>
+		/// <returns>Config file data</returns>
+		public static async Task<IConfigData> GetAsync(this IConfigSource source, Uri uri, CancellationToken cancellationToken)
+		{
+			IConfigData[] result = await source.GetAsync(new[] { uri }, cancellationToken);
+			return result[0];
+		}
 	}
 
 	/// <summary>
@@ -101,14 +120,19 @@ namespace Horde.Build.Configuration
 		}
 
 		/// <inheritdoc/>
-		public Task<IConfigData> GetAsync(Uri uri, CancellationToken cancellationToken)
+		public Task<IConfigData[]> GetAsync(Uri[] uris, CancellationToken cancellationToken)
 		{
-			ConfigFileRevisionImpl? configFile;
-			if (!_files.TryGetValue(uri, out configFile))
+			IConfigData[] result = new IConfigData[uris.Length];
+			for (int idx = 0; idx < uris.Length; idx++)
 			{
-				throw new FileNotFoundException($"Config file {uri} not found.");
+				ConfigFileRevisionImpl? configFile;
+				if (!_files.TryGetValue(uris[idx], out configFile))
+				{
+					throw new FileNotFoundException($"Config file {uris[idx]} not found.");
+				}
+				result[idx] = configFile;
 			}
-			return Task.FromResult<IConfigData>(configFile);
+			return Task.FromResult(result);
 		}
 	}
 
@@ -164,40 +188,46 @@ namespace Horde.Build.Configuration
 		}
 
 		/// <inheritdoc/>
-		public async Task<IConfigData> GetAsync(Uri uri, CancellationToken cancellationToken)
+		public async Task<IConfigData[]> GetAsync(Uri[] uris, CancellationToken cancellationToken)
 		{
-			FileReference localPath = FileReference.Combine(_baseDir, uri.LocalPath);
-
-			ConfigFileImpl? file;
-			for(; ;)
+			IConfigData[] files = new IConfigData[uris.Length];
+			for (int idx = 0; idx < uris.Length; idx++)
 			{
-				if (_files.TryGetValue(localPath, out file))
+				Uri uri = uris[idx];
+				FileReference localPath = FileReference.Combine(_baseDir, uri.LocalPath);
+
+				ConfigFileImpl? file;
+				for (; ; )
 				{
-					if (FileReference.GetLastWriteTimeUtc(localPath) == file.LastWriteTimeUtc)
+					if (_files.TryGetValue(localPath, out file))
 					{
-						break;
+						if (FileReference.GetLastWriteTimeUtc(localPath) == file.LastWriteTimeUtc)
+						{
+							break;
+						}
+						else
+						{
+							_files.TryRemove(new KeyValuePair<FileReference, ConfigFileImpl>(localPath, file));
+						}
 					}
-					else
+
+					using (FileStream stream = FileReference.Open(localPath, FileMode.Open, FileAccess.Read, FileShare.Read))
 					{
-						_files.TryRemove(new KeyValuePair<FileReference, ConfigFileImpl>(localPath, file));
+						using MemoryStream memoryStream = new MemoryStream();
+						await stream.CopyToAsync(memoryStream, cancellationToken);
+						DateTime lastWriteTime = FileReference.GetLastWriteTimeUtc(localPath);
+						file = new ConfigFileImpl(uri, lastWriteTime, memoryStream.ToArray());
 					}
-				}
 
-				using (FileStream stream = FileReference.Open(localPath, FileMode.Open, FileAccess.Read, FileShare.Read))
-				{
-					using MemoryStream memoryStream = new MemoryStream();
-					await stream.CopyToAsync(memoryStream, cancellationToken);
-					DateTime lastWriteTime = FileReference.GetLastWriteTimeUtc(localPath);
-					file = new ConfigFileImpl(uri, lastWriteTime, memoryStream.ToArray());
+					if (_files.TryAdd(localPath, file))
+					{
+						break;
+					}
 				}
 
-				if (_files.TryAdd(localPath, file))
-				{
-					break;
-				}
+				files[idx] = file;
 			}
-
-			return file;
+			return files;
 		}
 	}
 
@@ -248,20 +278,31 @@ namespace Horde.Build.Configuration
 		}
 
 		/// <inheritdoc/>
-		public async Task<IConfigData> GetAsync(Uri uri, CancellationToken cancellationToken)
+		public async Task<IConfigData[]> GetAsync(Uri[] uris, CancellationToken cancellationToken)
 		{
-			using (IPerforceConnection perforce = await ConnectAsync(uri.Host))
+			Dictionary<Uri, IConfigData> results = new Dictionary<Uri, IConfigData>();
+			foreach (IGrouping<string, Uri> group in uris.GroupBy(x => x.Host))
 			{
-				List<FStatRecord> files = await perforce.FStatAsync(FStatOptions.ShortenOutput, uri.AbsolutePath, cancellationToken).ToListAsync(cancellationToken);
-				files.RemoveAll(x => x.HeadAction == FileAction.Delete || x.HeadAction == FileAction.MoveDelete);
-
-				if (files.Count == 0)
+				using (IPerforceConnection perforce = await ConnectAsync(group.Key))
 				{
-					throw new FileNotFoundException($"Unable to read {uri}. No matching files found.");
-				}
+					FileSpecList fileSpec = group.Select(x => x.AbsolutePath).Distinct(StringComparer.OrdinalIgnoreCase).ToList();
 
-				return new ConfigFileImpl(uri, files[0].HeadChange, this);
+					List<FStatRecord> records = await perforce.FStatAsync(FStatOptions.ShortenOutput, fileSpec, cancellationToken).ToListAsync(cancellationToken);
+					records.RemoveAll(x => x.HeadAction == FileAction.Delete || x.HeadAction == FileAction.MoveDelete);
+
+					Dictionary<string, FStatRecord> absolutePathToRecord = records.ToDictionary(x => x.DepotFile ?? String.Empty, x => x, StringComparer.OrdinalIgnoreCase);
+					foreach (Uri uri in group)
+					{
+						FStatRecord? record;
+						if (!absolutePathToRecord.TryGetValue(uri.AbsolutePath, out record))
+						{
+							throw new FileNotFoundException($"Unable to read {uri}. No matching files found.");
+						}
+						results[uri] = new ConfigFileImpl(uri, record.HeadChange, this);
+					}
+				}
 			}
+			return uris.ConvertAll(x => results[x]);
 		}
 
 		async ValueTask<ReadOnlyMemory<byte>> ReadAsync(Uri uri, int change, CancellationToken cancellationToken)
diff --git a/Engine/Source/Programs/Horde/Horde.Build/Configuration/ConfigService.cs b/Engine/Source/Programs/Horde/Horde.Build/Configuration/ConfigService.cs
--- a/Engine/Source/Programs/Horde/Horde.Build/Configuration/ConfigService.cs
+++ b/Engine/Source/Programs/Horde/Horde.Build/Configuration/ConfigService.cs
@@ -8,6 +8,7 @@ using System.Linq;
 using System.Text.Json;
 using System.Threading;
 using System.Threading.Tasks;
+using Amazon.Runtime.Internal.Endpoints.StandardLibrary;
 using EpicGames.Core;
 using EpicGames.Redis;
 using Horde.Build.Projects;
@@ -119,7 +120,7 @@ namespace Horde.Build.Configuration
 		readonly RedisKey _snapshotKey = "config";
 
 		readonly ITicker _ticker;
-		readonly TimeSpan _tickInterval = TimeSpan.FromSeconds(2.0);
+		readonly TimeSpan _tickInterval = TimeSpan.FromSeconds(10.0);
 
 		readonly RedisChannel _updateChannel = "config-update";
 		readonly BackgroundTask _updateTask;
@@ -315,18 +316,24 @@ namespace Horde.Build.Configuration
 		/// <returns>True if the snapshot is out of date</returns>
 		async Task<bool> IsOutOfDateAsync(ConfigSnapshot snapshot, CancellationToken cancellationToken)
 		{
-			foreach ((Uri uri, string version) in snapshot.Dependencies)
+			// Group the dependencies by scheme in order to allow the source to batch-query them
+			foreach(IGrouping<string, KeyValuePair<Uri, string>> group in snapshot.Dependencies.GroupBy(x => x.Key.Scheme))
 			{
+				KeyValuePair<Uri, string>[] pairs = group.ToArray();
+
 				IConfigSource? source;
-				if (!_sources.TryGetValue(uri.Scheme, out source))
+				if (!_sources.TryGetValue(group.Key, out source))
 				{
 					return true;
 				}
 
-				IConfigData data = await source.GetAsync(uri, cancellationToken);
-				if (!data.Revision.Equals(version, StringComparison.Ordinal))
+				IConfigData[] data = await source.GetAsync(pairs.ConvertAll(x => x.Key), cancellationToken);
+				for (int idx = 0; idx < pairs.Length; idx++)
 				{
-					return true;
+					if (!data[idx].Revision.Equals(pairs[idx].Value, StringComparison.Ordinal))
+					{
+						return true;
+					}
 				}
 			}
 			return false;
diff --git a/Engine/Source/Programs/Horde/Horde.Build/Configuration/ConfigType.cs b/Engine/Source/Programs/Horde/Horde.Build/Configuration/ConfigType.cs
--- a/Engine/Source/Programs/Horde/Horde.Build/Configuration/ConfigType.cs
+++ b/Engine/Source/Programs/Horde/Horde.Build/Configuration/ConfigType.cs
@@ -349,21 +349,13 @@ namespace Horde.Build.Configuration
 			{
 				return new ScalarProperty(name, propertyInfo);
 			}
+			else if (propertyType.IsGenericType && propertyType.GetGenericTypeDefinition() == typeof(List<>))
+			{
+				Type elementType = propertyType.GetGenericArguments()[0];
+				return new ListProperty(name, propertyInfo, FindOrAddValueType(elementType));
+			}
 			else
 			{
-				if (propertyType.IsGenericType && propertyType.GetGenericTypeDefinition() == typeof(List<>))
-//				if (propertyType.IsAssignableTo(typeof(IList)))
-				{
-	//				foreach (Type type in propertyType.GetInterfaces())
-	//				{
-	//					if (type.IsGenericType && type.GetGenericTypeDefinition() == typeof(IList<>))
-	//					{
-//							Type elementType = type.GetGenericArguments()[0];
-					Type elementType = propertyType.GetGenericArguments()[0];
-					return new ListProperty(name, propertyInfo, FindOrAddValueType(elementType));
-//						}
-//					}
-				}
 				return new ObjectProperty(name, propertyInfo, FindOrAdd(propertyType));
 			}
 		}
diff --git a/Engine/Source/Programs/Shared/EpicGames.Perforce/IPerforceOutput.cs b/Engine/Source/Programs/Shared/EpicGames.Perforce/IPerforceOutput.cs
--- a/Engine/Source/Programs/Shared/EpicGames.Perforce/IPerforceOutput.cs
+++ b/Engine/Source/Programs/Shared/EpicGames.Perforce/IPerforceOutput.cs
@@ -231,7 +231,7 @@ namespace EpicGames.Perforce
 				ReadOnlyMemory<byte> data = perforce.Data;
 				if (data.Length > 0 && responses.Count == 0 && data.Span[0] != '{')
 				{
-					throw new PerforceException("Unexpected response from server (expected '{'):{0}", FormatDataAsString(data.Span));
+					throw new PerforceException("Unexpected response from server (expected '{{'):{0}", FormatDataAsString(data.Span));
 				}
 
 				// Parse the responses from the current buffer
