Horde: Use a memory cache for files read by Perforce config data provider.

diff --git a/Engine/Source/Programs/Horde/Horde.Build/Configuration/ConfigData.cs b/Engine/Source/Programs/Horde/Horde.Build/Configuration/ConfigData.cs
--- a/Engine/Source/Programs/Horde/Horde.Build/Configuration/ConfigData.cs
+++ b/Engine/Source/Programs/Horde/Horde.Build/Configuration/ConfigData.cs
@@ -10,6 +10,7 @@ using System.Threading.Tasks;
 using EpicGames.Core;
 using EpicGames.Perforce;
 using Horde.Build.Utilities;
+using Microsoft.Extensions.Caching.Memory;
 using Microsoft.Extensions.Logging;
 using Microsoft.Extensions.Options;
 
@@ -264,16 +265,16 @@ namespace Horde.Build.Configuration
 		string IConfigSource.Scheme => Scheme;
 
 		readonly IOptionsMonitor<ServerSettings> _settings;
+		readonly IMemoryCache _cache;
 		readonly ILogger _logger;
 
 		/// <summary>
 		/// Constructor
 		/// </summary>
-		/// <param name="settings">Global settings</param>
-		/// <param name="logger">Logger instance</param>
-		public PerforceConfigSource(IOptionsMonitor<ServerSettings> settings, ILogger<PerforceConfigSource> logger)
+		public PerforceConfigSource(IOptionsMonitor<ServerSettings> settings, IMemoryCache cache, ILogger<PerforceConfigSource> logger)
 		{
 			_settings = settings;
+			_cache = cache;
 			_logger = logger;
 		}
 
@@ -307,12 +308,23 @@ namespace Horde.Build.Configuration
 
 		async ValueTask<ReadOnlyMemory<byte>> ReadAsync(Uri uri, int change, CancellationToken cancellationToken)
 		{
-			using (IPerforceConnection perforce = await ConnectAsync(uri.Host))
+			string cacheKey = $"{nameof(PerforceConfigSource)}:{uri}@{change}";
+			if (!_cache.TryGetValue(cacheKey, out ReadOnlyMemory<byte> data))
 			{
-				PerforceResponse<PrintRecord<byte[]>> response = await perforce.TryPrintDataAsync($"{uri.AbsolutePath}@{change}", cancellationToken);
-				response.EnsureSuccess();
-				return response.Data.Contents!;
+				using (IPerforceConnection perforce = await ConnectAsync(uri.Host))
+				{
+					PerforceResponse<PrintRecord<byte[]>> response = await perforce.TryPrintDataAsync($"{uri.AbsolutePath}@{change}", cancellationToken);
+					response.EnsureSuccess();
+					data = response.Data.Contents!;
+				}
+				using (ICacheEntry entry = _cache.CreateEntry(cacheKey))
+				{
+					entry.SetSlidingExpiration(TimeSpan.FromHours(1.0));
+					entry.SetSize(data.Length);
+					entry.SetValue(data);
+				}
 			}
+			return data;
 		}
 
 		async Task<IPerforceConnection> ConnectAsync(string host)
