Horde: Unknown error => Batch error in messages
#rnx
#skipci

diff --git a/Engine/Source/Programs/Horde/HordeDashboard/src/backend/Api.ts b/Engine/Source/Programs/Horde/HordeDashboard/src/backend/Api.ts
--- a/Engine/Source/Programs/Horde/HordeDashboard/src/backend/Api.ts
+++ b/Engine/Source/Programs/Horde/HordeDashboard/src/backend/Api.ts
@@ -138,8 +138,10 @@ export enum JobStepBatchError {
 	Incomplete = "Incomplete",
 
 	/** An error ocurred while executing the lease. Cannot be retried. */
-	ExecutionError = "ExecutionError"
+	ExecutionError = "ExecutionError",
 
+	/** The change that the job is running against is invalid. */
+	UnknownShelf  = "UnknownShelf"
 
 }
 
diff --git a/Engine/Source/Programs/Horde/HordeDashboard/src/components/JobDetailCommon.tsx b/Engine/Source/Programs/Horde/HordeDashboard/src/components/JobDetailCommon.tsx
--- a/Engine/Source/Programs/Horde/HordeDashboard/src/components/JobDetailCommon.tsx
+++ b/Engine/Source/Programs/Horde/HordeDashboard/src/components/JobDetailCommon.tsx
@@ -145,7 +145,7 @@ export const getBatchText = (item: StepItem): string | undefined => {
 
 
       if (!statusText && batch.error !== undefined && batch.error !== JobStepBatchError.None) {
-         statusText = `Unknown error: ${batch.error}`;
+         statusText = `Batch error: ${batch.error}`;
       }
 
       if (!statusText) {
