Horde: Add Slack notifications for errors in new config system.

diff --git a/Engine/Source/Programs/Horde/Horde.Build/Configuration/ConfigContext.cs b/Engine/Source/Programs/Horde/Horde.Build/Configuration/ConfigContext.cs
--- a/Engine/Source/Programs/Horde/Horde.Build/Configuration/ConfigContext.cs
+++ b/Engine/Source/Programs/Horde/Horde.Build/Configuration/ConfigContext.cs
@@ -12,7 +12,7 @@ namespace Horde.Build.Configuration
 	/// Context for reading a tree of config files
 	/// </summary>
 	[DebuggerDisplay("{CurrentFile}")]
-	class ConfigContext
+	public class ConfigContext
 	{
 		/// <summary>
 		/// Options for serializing config files
@@ -22,7 +22,7 @@ namespace Horde.Build.Configuration
 		/// <summary>
 		/// Stack of included files
 		/// </summary>
-		public Stack<Uri> IncludeStack { get; } = new Stack<Uri>();
+		public Stack<IConfigFile> IncludeStack { get; } = new Stack<IConfigFile>();
 
 		/// <summary>
 		/// Stack of properties
@@ -32,7 +32,7 @@ namespace Horde.Build.Configuration
 		/// <summary>
 		/// Map of property path to the file declaring a value for it
 		/// </summary>
-		public Dictionary<string, Uri> PropertyPathToFile = new Dictionary<string, Uri>(StringComparer.OrdinalIgnoreCase);
+		public Dictionary<string, Uri> PropertyPathToFile { get; } = new Dictionary<string, Uri>(StringComparer.OrdinalIgnoreCase);
 
 		/// <summary>
 		/// Sources to read config files from
@@ -42,12 +42,12 @@ namespace Horde.Build.Configuration
 		/// <summary>
 		/// Tracks files read as part of the configuration
 		/// </summary>
-		public Dictionary<Uri, IConfigData> Files { get; } = new Dictionary<Uri, IConfigData>();
+		public Dictionary<Uri, IConfigFile> Files { get; } = new Dictionary<Uri, IConfigFile>();
 
 		/// <summary>
 		/// Uri of the current file
 		/// </summary>
-		public Uri CurrentFile => (IncludeStack.Count > 0) ? IncludeStack.Peek() : null!;
+		public Uri CurrentFile => (IncludeStack.Count > 0) ? IncludeStack.Peek().Uri : null!;
 
 		/// <summary>
 		/// Current property scope
diff --git a/Engine/Source/Programs/Horde/Horde.Build/Configuration/ConfigData.cs b/Engine/Source/Programs/Horde/Horde.Build/Configuration/ConfigData.cs
deleted file mode 100644
diff --git a/Engine/Source/Programs/Horde/Horde.Build/Configuration/ConfigFile.cs b/Engine/Source/Programs/Horde/Horde.Build/Configuration/ConfigFile.cs
new file mode 100644
--- /dev/null
+++ b/Engine/Source/Programs/Horde/Horde.Build/Configuration/ConfigFile.cs
@@ -0,0 +1,37 @@
+// Copyright Epic Games, Inc. All Rights Reserved.
+
+using System.Threading.Tasks;
+using System.Threading;
+using System;
+using Horde.Build.Users;
+
+namespace Horde.Build.Configuration
+{
+	/// <summary>
+	/// Accessor for a specific revision of a config file. Provides metadata about the current revision, and allows reading its data.
+	/// </summary>
+	public interface IConfigFile
+	{
+		/// <summary>
+		/// URI of the config file
+		/// </summary>
+		Uri Uri { get; }
+
+		/// <summary>
+		/// String used to identify a specific version of the config data. Used for ordinal comparisons, otherwise opaque.
+		/// </summary>
+		string Revision { get; }
+
+		/// <summary>
+		/// Author of this revision, if known
+		/// </summary>
+		IUser? Author { get; }
+
+		/// <summary>
+		/// Reads data for the config file
+		/// </summary>
+		/// <param name="cancellationToken">Cancellation token for the operation</param>
+		/// <returns>UTF-8 encoded data for the config file</returns>
+		ValueTask<ReadOnlyMemory<byte>> ReadAsync(CancellationToken cancellationToken);
+	}
+}
diff --git a/Engine/Source/Programs/Horde/Horde.Build/Configuration/ConfigService.cs b/Engine/Source/Programs/Horde/Horde.Build/Configuration/ConfigService.cs
--- a/Engine/Source/Programs/Horde/Horde.Build/Configuration/ConfigService.cs
+++ b/Engine/Source/Programs/Horde/Horde.Build/Configuration/ConfigService.cs
@@ -8,9 +8,9 @@ using System.Linq;
 using System.Text.Json;
 using System.Threading;
 using System.Threading.Tasks;
-using Amazon.Runtime.Internal.Endpoints.StandardLibrary;
 using EpicGames.Core;
 using EpicGames.Redis;
+using Horde.Build.Notifications;
 using Horde.Build.Projects;
 using Horde.Build.Server;
 using Horde.Build.Streams;
@@ -114,10 +114,12 @@ namespace Horde.Build.Configuration
 		record class ConfigState(IoHash Hash, GlobalConfig GlobalConfig);
 
 		readonly RedisService _redisService;
+		readonly INotificationService _notificationService;
 		readonly ServerSettings _settings;
 		readonly Dictionary<string, IConfigSource> _sources;
 		readonly JsonSerializerOptions _jsonOptions;
 		readonly RedisKey _snapshotKey = "config";
+		readonly ILogger _logger;
 
 		readonly ITicker _ticker;
 		readonly TimeSpan _tickInterval = TimeSpan.FromSeconds(10.0);
@@ -134,11 +136,13 @@ namespace Horde.Build.Configuration
 		/// <summary>
 		/// Constructor
 		/// </summary>
-		public ConfigService(RedisService redisService, IOptions<ServerSettings> settings, IEnumerable<IConfigSource> sources, IClock clock, ILogger<ConfigService> logger)
+		public ConfigService(RedisService redisService, INotificationService notificationService, IOptions<ServerSettings> settings, IEnumerable<IConfigSource> sources, IClock clock, ILogger<ConfigService> logger)
 		{
 			_redisService = redisService;
+			_notificationService = notificationService;
 			_settings = settings.Value;
 			_sources = sources.ToDictionary(x => x.Scheme, x => x, StringComparer.OrdinalIgnoreCase);
+			_logger = logger;
 
 			_jsonOptions = new JsonSerializerOptions();
 			Startup.ConfigureJsonSerializer(_jsonOptions);
@@ -221,8 +225,11 @@ namespace Horde.Build.Configuration
 			{
 				if (snapshot == null || await IsOutOfDateAsync(snapshot, cancellationToken))
 				{
-					snapshot = await CreateSnapshotAsync(cancellationToken);
-					await WriteSnapshotAsync(snapshot);
+					snapshot = await CreateSnapshotGuardedAsync(cancellationToken);
+					if (snapshot != null)
+					{
+						await WriteSnapshotAsync(snapshot);
+					}
 				}
 				await Task.Delay(_tickInterval, cancellationToken);
 			}
@@ -249,6 +256,27 @@ namespace Horde.Build.Configuration
 			}
 		}
 
+		/// <summary>
+		/// Create a new snapshot object, catching any exceptions and sending notifications
+		/// </summary>
+		/// <param name="cancellationToken">Cancellation token for the operation</param>
+		/// <returns>New config snapshot</returns>
+		async Task<ConfigSnapshot?> CreateSnapshotGuardedAsync(CancellationToken cancellationToken)
+		{
+			try
+			{
+				ConfigSnapshot snapshot = await CreateSnapshotAsync(cancellationToken);
+				_notificationService.NotifyConfigUpdate(null);
+				return snapshot;
+			}
+			catch (Exception ex)
+			{
+				_logger.LogError(ex, "Exception while updating config: {Message}", ex.Message);
+				_notificationService.NotifyConfigUpdate(ex);
+				return null;
+			}
+		}
+
 		/// <summary>
 		/// Create a new snapshot object
 		/// </summary>
@@ -295,16 +323,16 @@ namespace Horde.Build.Configuration
 				}
 
 				// Save all the dependencies
-				foreach ((Uri depUri, IConfigData depData) in context.Files)
+				foreach ((Uri depUri, IConfigFile depFile) in context.Files)
 				{
-					snapshot.Dependencies.Add(depUri, depData.Revision);
+					snapshot.Dependencies.Add(depUri, depFile.Revision);
 				}
 
 				return snapshot;
 			}
 			catch (Exception ex) when (ex is not ConfigException)
 			{
-				throw new ConfigException(context, "Internal exception while parsing config files", ex);
+				throw new ConfigException(context, ex.Message, ex);
 			}
 		}
 
@@ -327,10 +355,10 @@ namespace Horde.Build.Configuration
 					return true;
 				}
 
-				IConfigData[] data = await source.GetAsync(pairs.ConvertAll(x => x.Key), cancellationToken);
+				IConfigFile[] files = await source.GetAsync(pairs.ConvertAll(x => x.Key), cancellationToken);
 				for (int idx = 0; idx < pairs.Length; idx++)
 				{
-					if (!data[idx].Revision.Equals(pairs[idx].Value, StringComparison.Ordinal))
+					if (!files[idx].Revision.Equals(pairs[idx].Value, StringComparison.Ordinal))
 					{
 						return true;
 					}
diff --git a/Engine/Source/Programs/Horde/Horde.Build/Configuration/ConfigSource.cs b/Engine/Source/Programs/Horde/Horde.Build/Configuration/ConfigSource.cs
new file mode 100644
--- /dev/null
+++ b/Engine/Source/Programs/Horde/Horde.Build/Configuration/ConfigSource.cs
@@ -0,0 +1,359 @@
+// Copyright Epic Games, Inc. All Rights Reserved.
+
+using System;
+using System.Collections.Concurrent;
+using System.Collections.Generic;
+using System.IO;
+using System.Linq;
+using System.Threading;
+using System.Threading.Tasks;
+using Amazon.Runtime.Internal.Endpoints.StandardLibrary;
+using EpicGames.Core;
+using EpicGames.Perforce;
+using Horde.Build.Users;
+using Horde.Build.Utilities;
+using Microsoft.Extensions.Caching.Memory;
+using Microsoft.Extensions.Logging;
+using Microsoft.Extensions.Options;
+
+namespace Horde.Build.Configuration
+{
+	using PerforceConnectionId = StringId<PerforceConnectionSettings>;
+
+	/// <summary>
+	/// Source for reading config files
+	/// </summary>
+	public interface IConfigSource
+	{
+		/// <summary>
+		/// URI scheme for this config source
+		/// </summary>
+		string Scheme { get; }
+
+		/// <summary>
+		/// Reads a config file from this source
+		/// </summary>
+		/// <param name="uris">Locations of the config files to query</param>
+		/// <param name="cancellationToken">Cancellation token for the operation</param>
+		/// <returns>Config file data</returns>
+		Task<IConfigFile[]> GetAsync(Uri[] uris, CancellationToken cancellationToken);
+	}
+
+	/// <summary>
+	/// Extension methods for <see cref="IConfigSource"/>
+	/// </summary>
+	public static class ConfigSource
+	{
+		/// <summary>
+		/// Gets a single config file from a source
+		/// </summary>
+		/// <param name="source">Source to query</param>
+		/// <param name="uri">Location of the config file to query</param>
+		/// <param name="cancellationToken">Cancellation token for the operation</param>
+		/// <returns>Config file data</returns>
+		public static async Task<IConfigFile> GetAsync(this IConfigSource source, Uri uri, CancellationToken cancellationToken)
+		{
+			IConfigFile[] result = await source.GetAsync(new[] { uri }, cancellationToken);
+			return result[0];
+		}
+	}
+
+	/// <summary>
+	/// In-memory config file source
+	/// </summary>
+	public sealed class InMemoryConfigSource : IConfigSource
+	{
+		class ConfigFileRevisionImpl : IConfigFile
+		{
+			public Uri Uri { get; }
+			public string Revision { get; }
+			public ReadOnlyMemory<byte> Data { get; }
+			public IUser? Author => null;
+
+			public ConfigFileRevisionImpl(Uri uri, string version, ReadOnlyMemory<byte> data)
+			{
+				Uri = uri;
+				Revision = version;
+				Data = data;
+			}
+
+			public ValueTask<ReadOnlyMemory<byte>> ReadAsync(CancellationToken cancellationToken) => new ValueTask<ReadOnlyMemory<byte>>(Data);
+		}
+
+		readonly Dictionary<Uri, ConfigFileRevisionImpl> _files = new Dictionary<Uri, ConfigFileRevisionImpl>();
+
+		/// <summary>
+		/// Name of the scheme for this source
+		/// </summary>
+		public const string Scheme = "memory";
+
+		/// <inheritdoc/>
+		string IConfigSource.Scheme => Scheme;
+
+		/// <summary>
+		/// Manually adds a new config file
+		/// </summary>
+		/// <param name="path">Path to the config file</param>
+		/// <param name="data">Config file data</param>
+		public void Add(Uri path, ReadOnlyMemory<byte> data)
+		{
+			_files.Add(path, new ConfigFileRevisionImpl(path, "v1", data));
+		}
+
+		/// <inheritdoc/>
+		public Task<IConfigFile[]> GetAsync(Uri[] uris, CancellationToken cancellationToken)
+		{
+			IConfigFile[] result = new IConfigFile[uris.Length];
+			for (int idx = 0; idx < uris.Length; idx++)
+			{
+				ConfigFileRevisionImpl? configFile;
+				if (!_files.TryGetValue(uris[idx], out configFile))
+				{
+					throw new FileNotFoundException($"Config file {uris[idx]} not found.");
+				}
+				result[idx] = configFile;
+			}
+			return Task.FromResult(result);
+		}
+	}
+
+	/// <summary>
+	/// Config file source which reads from the filesystem
+	/// </summary>
+	public sealed class FileConfigSource : IConfigSource
+	{
+		class ConfigFileImpl : IConfigFile
+		{
+			public Uri Uri { get; }
+			public string Revision { get; }
+			public DateTime LastWriteTimeUtc { get; }
+			public ReadOnlyMemory<byte> Data { get; }
+			public IUser? Author => null;
+
+			public ConfigFileImpl(Uri uri, DateTime lastWriteTimeUtc, ReadOnlyMemory<byte> data)
+			{
+				Uri = uri;
+				Revision = $"timestamp={lastWriteTimeUtc.Ticks}";
+				LastWriteTimeUtc = lastWriteTimeUtc;
+				Data = data;
+			}
+
+			public ValueTask<ReadOnlyMemory<byte>> ReadAsync(CancellationToken cancellationToken) => new ValueTask<ReadOnlyMemory<byte>>(Data);
+		}
+
+		/// <summary>
+		/// Name of the scheme for this source
+		/// </summary>
+		public const string Scheme = "file";
+
+		/// <inheritdoc/>
+		string IConfigSource.Scheme => Scheme;
+
+		readonly DirectoryReference _baseDir;
+		readonly ConcurrentDictionary<FileReference, ConfigFileImpl> _files = new ConcurrentDictionary<FileReference, ConfigFileImpl>();
+
+		/// <summary>
+		/// Constructor
+		/// </summary>
+		public FileConfigSource()
+			: this(DirectoryReference.GetCurrentDirectory())
+		{
+		}
+
+		/// <summary>
+		/// Constructor
+		/// </summary>
+		/// <param name="baseDir">Base directory for resolving relative paths</param>
+		public FileConfigSource(DirectoryReference baseDir)
+		{
+			_baseDir = baseDir;
+		}
+
+		/// <inheritdoc/>
+		public async Task<IConfigFile[]> GetAsync(Uri[] uris, CancellationToken cancellationToken)
+		{
+			IConfigFile[] files = new IConfigFile[uris.Length];
+			for (int idx = 0; idx < uris.Length; idx++)
+			{
+				Uri uri = uris[idx];
+				FileReference localPath = FileReference.Combine(_baseDir, uri.LocalPath);
+
+				ConfigFileImpl? file;
+				for (; ; )
+				{
+					if (_files.TryGetValue(localPath, out file))
+					{
+						if (FileReference.GetLastWriteTimeUtc(localPath) == file.LastWriteTimeUtc)
+						{
+							break;
+						}
+						else
+						{
+							_files.TryRemove(new KeyValuePair<FileReference, ConfigFileImpl>(localPath, file));
+						}
+					}
+
+					using (FileStream stream = FileReference.Open(localPath, FileMode.Open, FileAccess.Read, FileShare.Read))
+					{
+						using MemoryStream memoryStream = new MemoryStream();
+						await stream.CopyToAsync(memoryStream, cancellationToken);
+						DateTime lastWriteTime = FileReference.GetLastWriteTimeUtc(localPath);
+						file = new ConfigFileImpl(uri, lastWriteTime, memoryStream.ToArray());
+					}
+
+					if (_files.TryAdd(localPath, file))
+					{
+						break;
+					}
+				}
+
+				files[idx] = file;
+			}
+			return files;
+		}
+	}
+
+	/// <summary>
+	/// Perforce cluster config file source
+	/// </summary>
+	public sealed class PerforceConfigSource : IConfigSource
+	{
+		class ConfigFileImpl : IConfigFile
+		{
+			public Uri Uri { get; }
+			public int Change { get; }
+			public string Revision { get; }
+			public IUser? Author { get; }
+
+			readonly PerforceConfigSource _owner;
+
+			public ConfigFileImpl(Uri uri, int change, IUser? author, PerforceConfigSource owner)
+			{
+				Uri = uri;
+				Change = change;
+				Revision = $"{change}";
+				Author = author;
+				_owner = owner;
+			}
+
+			public ValueTask<ReadOnlyMemory<byte>> ReadAsync(CancellationToken cancellationToken) => _owner.ReadAsync(Uri, Change, cancellationToken);
+		}
+
+		/// <summary>
+		/// Name of the scheme for this source
+		/// </summary>
+		public const string Scheme = "perforce";
+
+		/// <inheritdoc/>
+		string IConfigSource.Scheme => Scheme;
+
+		readonly IOptionsMonitor<ServerSettings> _settings;
+		readonly IUserCollection _userCollection;
+		readonly IMemoryCache _cache;
+		readonly ILogger _logger;
+
+		/// <summary>
+		/// Constructor
+		/// </summary>
+		public PerforceConfigSource(IOptionsMonitor<ServerSettings> settings, IUserCollection userCollection, IMemoryCache cache, ILogger<PerforceConfigSource> logger)
+		{
+			_settings = settings;
+			_userCollection = userCollection;
+			_cache = cache;
+			_logger = logger;
+		}
+
+		/// <inheritdoc/>
+		public async Task<IConfigFile[]> GetAsync(Uri[] uris, CancellationToken cancellationToken)
+		{
+			Dictionary<Uri, IConfigFile> results = new Dictionary<Uri, IConfigFile>();
+			foreach (IGrouping<string, Uri> group in uris.GroupBy(x => x.Host))
+			{
+				using (IPerforceConnection perforce = await ConnectAsync(group.Key))
+				{
+					FileSpecList fileSpec = group.Select(x => x.AbsolutePath).Distinct(StringComparer.OrdinalIgnoreCase).ToList();
+
+					List<FStatRecord> records = await perforce.FStatAsync(FStatOptions.ShortenOutput, fileSpec, cancellationToken).ToListAsync(cancellationToken);
+					records.RemoveAll(x => x.HeadAction == FileAction.Delete || x.HeadAction == FileAction.MoveDelete);
+
+					Dictionary<string, FStatRecord> absolutePathToRecord = records.ToDictionary(x => x.DepotFile ?? String.Empty, x => x, StringComparer.OrdinalIgnoreCase);
+					foreach (Uri uri in group)
+					{
+						FStatRecord? record;
+						if (!absolutePathToRecord.TryGetValue(uri.AbsolutePath, out record))
+						{
+							throw new FileNotFoundException($"Unable to read {uri}. No matching files found.");
+						}
+
+						IUser? author = await GetAuthorAsync(perforce, group.Key, record.HeadChange, cancellationToken);
+						results[uri] = new ConfigFileImpl(uri, record.HeadChange, author, this);
+					}
+				}
+			}
+			return uris.ConvertAll(x => results[x]);
+		}
+
+		async ValueTask<IUser?> GetAuthorAsync(IPerforceConnection perforce, string host, int change, CancellationToken cancellationToken)
+		{
+			string cacheKey = $"{nameof(PerforceConfigSource)}:author:{host}@{change}";
+			if (!_cache.TryGetValue(cacheKey, out string? author))
+			{
+				ChangeRecord record = await perforce.GetChangeAsync(GetChangeOptions.None, change, cancellationToken);
+				using (ICacheEntry entry = _cache.CreateEntry(cacheKey))
+				{
+					entry.SetSlidingExpiration(TimeSpan.FromHours(1.0));
+					entry.SetSize(256);
+					entry.SetValue(record.User);
+				}
+			}
+			return (author != null)? await _userCollection.FindUserByLoginAsync(author) : null;
+		}
+
+		async ValueTask<ReadOnlyMemory<byte>> ReadAsync(Uri uri, int change, CancellationToken cancellationToken)
+		{
+			string cacheKey = $"{nameof(PerforceConfigSource)}:data:{uri}@{change}";
+			if (!_cache.TryGetValue(cacheKey, out ReadOnlyMemory<byte> data))
+			{
+				using (IPerforceConnection perforce = await ConnectAsync(uri.Host))
+				{
+					PerforceResponse<PrintRecord<byte[]>> response = await perforce.TryPrintDataAsync($"{uri.AbsolutePath}@{change}", cancellationToken);
+					response.EnsureSuccess();
+					data = response.Data.Contents!;
+				}
+				using (ICacheEntry entry = _cache.CreateEntry(cacheKey))
+				{
+					entry.SetSlidingExpiration(TimeSpan.FromHours(1.0));
+					entry.SetSize(data.Length);
+					entry.SetValue(data);
+				}
+			}
+			return data;
+		}
+
+		async Task<IPerforceConnection> ConnectAsync(string host)
+		{
+			ServerSettings settings = _settings.CurrentValue;
+
+			PerforceConnectionId connectionId = new PerforceConnectionId();
+			if (!String.IsNullOrEmpty(host))
+			{
+				connectionId = new PerforceConnectionId(host);
+			}
+
+			PerforceConnectionSettings? connectionSettings = settings.Perforce.FirstOrDefault(x => x.Id == connectionId);
+			if (connectionSettings == null)
+			{
+				if (connectionId == PerforceConnectionSettings.Default)
+				{
+					connectionSettings = new PerforceConnectionSettings();
+				}
+				else
+				{
+					throw new InvalidOperationException($"No Perforce connection settings defined for '{connectionId}'.");
+				}
+			}
+
+			return await PerforceConnection.CreateAsync(connectionSettings.ToPerforceSettings(), _logger);
+		}
+	}
+}
diff --git a/Engine/Source/Programs/Horde/Horde.Build/Configuration/ConfigType.cs b/Engine/Source/Programs/Horde/Horde.Build/Configuration/ConfigType.cs
--- a/Engine/Source/Programs/Horde/Horde.Build/Configuration/ConfigType.cs
+++ b/Engine/Source/Programs/Horde/Horde.Build/Configuration/ConfigType.cs
@@ -5,6 +5,7 @@ using System.Collections;
 using System.Collections.Concurrent;
 using System.Collections.Generic;
 using System.ComponentModel.DataAnnotations;
+using System.Linq;
 using System.Reflection;
 using System.Text.Json;
 using System.Text.Json.Nodes;
@@ -12,7 +13,7 @@ using System.Text.Json.Serialization;
 using System.Threading;
 using System.Threading.Tasks;
 using EpicGames.Core;
-using Horde.Build.Server;
+
 using JsonObject = System.Text.Json.Nodes.JsonObject;
 
 namespace Horde.Build.Configuration
@@ -43,9 +44,9 @@ namespace Horde.Build.Configuration
 	public sealed class ConfigException : Exception
 	{
 		/// <summary>
-		/// Path to the config file triggering the exception
+		/// Context for the parser
 		/// </summary>
-		public Uri Path { get; }
+		public ConfigContext Context { get; }
 
 		/// <summary>
 		/// Constructor
@@ -56,7 +57,7 @@ namespace Horde.Build.Configuration
 		internal ConfigException(ConfigContext context, string message, Exception? innerException = null)
 			: base(message, innerException)
 		{
-			Path = context.CurrentFile;
+			Context = context;
 		}
 	}
 
@@ -393,7 +394,7 @@ namespace Horde.Build.Configuration
 			return result;
 		}
 
-		static async ValueTask<JsonObject> ReadFileAsync(Uri uri, ConfigContext context, CancellationToken cancellationToken)
+		static async ValueTask<IConfigFile> ReadFileAsync(Uri uri, ConfigContext context, CancellationToken cancellationToken)
 		{
 			IConfigSource? source = context.Sources[uri.Scheme];
 			if (source == null)
@@ -401,19 +402,24 @@ namespace Horde.Build.Configuration
 				throw new ConfigException(context, $"Invalid/unknown scheme for config file {uri}");
 			}
 
-			IConfigData? file;
+			IConfigFile? file;
 			if (!context.Files.TryGetValue(uri, out file))
 			{
 				file = await source.GetAsync(uri, cancellationToken);
 				context.Files.Add(uri, file);
 			}
 
+			return file;
+		}
+
+		static async ValueTask<JsonObject> ParseFileAsync(IConfigFile file, ConfigContext context, CancellationToken cancellationToken)
+		{
 			ReadOnlyMemory<byte> data = await file.ReadAsync(cancellationToken);
 
 			JsonObject? obj = JsonSerializer.Deserialize<JsonObject>(data.Span, context.JsonOptions);
 			if (obj == null)
 			{
-				throw new ConfigException(context, $"Config file {uri} contains a null object.");
+				throw new ConfigException(context, $"Config file {file.Uri} contains a null object.");
 			}
 
 			return obj;
@@ -421,14 +427,16 @@ namespace Horde.Build.Configuration
 
 		public async Task MergeObjectAsync(object target, Uri uri, ConfigContext context, CancellationToken cancellationToken)
 		{
-			if (context.IncludeStack.Contains(uri))
+			if (context.IncludeStack.Any(x => x.Uri == uri))
 			{
 				throw new ConfigException(context, $"Recursive include of file {uri}");
 			}
 
-			context.IncludeStack.Push(uri);
+			IConfigFile file = await ReadFileAsync(uri, context, cancellationToken);
+
+			context.IncludeStack.Push(file);
 
-			JsonObject obj = await ReadFileAsync(uri, context, cancellationToken);
+			JsonObject obj = await ParseFileAsync(file, context, cancellationToken);
 			await MergeObjectAsync(target, obj, context, cancellationToken);
 
 			context.IncludeStack.Pop();
@@ -448,9 +456,11 @@ namespace Horde.Build.Configuration
 						foreach (ConfigInclude include in includes)
 						{
 							Uri uri = ConfigType.CombinePaths(context.CurrentFile, include.Path);
-							context.IncludeStack.Push(uri);
+							IConfigFile file = await ReadFileAsync(uri, context, cancellationToken);
+
+							context.IncludeStack.Push(file);
 
-							JsonObject nextObj = await ReadFileAsync(uri, context, cancellationToken);
+							JsonObject nextObj = await ParseFileAsync(file, context, cancellationToken);
 							await MergeObjectAsync(target, nextObj, context, cancellationToken);
 
 							context.IncludeStack.Pop();
diff --git a/Engine/Source/Programs/Horde/Horde.Build/Notifications/INotificationService.cs b/Engine/Source/Programs/Horde/Horde.Build/Notifications/INotificationService.cs
--- a/Engine/Source/Programs/Horde/Horde.Build/Notifications/INotificationService.cs
+++ b/Engine/Source/Programs/Horde/Horde.Build/Notifications/INotificationService.cs
@@ -115,6 +115,12 @@ namespace Horde.Build.Notifications
 		/// <returns>Async task</returns>
 		void NotifyLabelUpdate(IJob job, IReadOnlyList<(LabelState, LabelOutcome)> oldLabelStates, IReadOnlyList<(LabelState, LabelOutcome)> newLabelStates);
 
+		/// <summary>
+		/// Notify slack channel about a configuration update
+		/// </summary>
+		/// <param name="ex">Exception thrown during the update. Null if the update completed successfully.</param>
+		void NotifyConfigUpdate(Exception? ex);
+
 		/// <summary>
 		/// Notify slack channel about a stream update failure
 		/// </summary>
diff --git a/Engine/Source/Programs/Horde/Horde.Build/Notifications/INotificationSink.cs b/Engine/Source/Programs/Horde/Horde.Build/Notifications/INotificationSink.cs
--- a/Engine/Source/Programs/Horde/Horde.Build/Notifications/INotificationSink.cs
+++ b/Engine/Source/Programs/Horde/Horde.Build/Notifications/INotificationSink.cs
@@ -3,6 +3,7 @@
 using System;
 using System.Collections.Generic;
 using System.Threading.Tasks;
+using Horde.Build.Configuration;
 using Horde.Build.Devices;
 using Horde.Build.Issues;
 using Horde.Build.Jobs;
@@ -91,6 +92,12 @@ namespace Horde.Build.Notifications
 
 		#endregion
 
+		/// <summary>
+		/// Notification that the configuration state has changed
+		/// </summary>
+		/// <param name="ex">Exception during update. Null if the update completed successfully.</param>
+		Task NotifyConfigUpdateAsync(Exception? ex);
+
 		/// <summary>
 		/// Notification that a stream has failed to update
 		/// </summary>
diff --git a/Engine/Source/Programs/Horde/Horde.Build/Notifications/NotificationService.cs b/Engine/Source/Programs/Horde/Horde.Build/Notifications/NotificationService.cs
--- a/Engine/Source/Programs/Horde/Horde.Build/Notifications/NotificationService.cs
+++ b/Engine/Source/Programs/Horde/Horde.Build/Notifications/NotificationService.cs
@@ -276,6 +276,13 @@ namespace Horde.Build.Notifications
 			EnqueueTasks(sink => sink.NotifyIssueUpdatedAsync(issue));
 		}
 
+		/// <inheritdoc/>
+		public void NotifyConfigUpdate(Exception? ex)
+		{
+			_logger.LogInformation(ex, "Configuration updated ({Result})", (ex == null) ? "success" : "failure");
+			EnqueueTasks(sink => sink.NotifyConfigUpdateAsync(ex));
+		}
+
 		/// <inheritdoc/>
 		public void NotifyConfigUpdateFailure(string errorMessage, string fileName, int? change = null, IUser? author = null, string? description = null)
 		{
diff --git a/Engine/Source/Programs/Horde/Horde.Build/Notifications/Sinks/SlackNotificationSink.cs b/Engine/Source/Programs/Horde/Horde.Build/Notifications/Sinks/SlackNotificationSink.cs
--- a/Engine/Source/Programs/Horde/Horde.Build/Notifications/Sinks/SlackNotificationSink.cs
+++ b/Engine/Source/Programs/Horde/Horde.Build/Notifications/Sinks/SlackNotificationSink.cs
@@ -4,6 +4,7 @@ using System;
 using System.Collections.Generic;
 using System.Diagnostics.CodeAnalysis;
 using System.Globalization;
+using System.IO;
 using System.Linq;
 using System.Net.Http;
 using System.Net.WebSockets;
@@ -20,6 +21,8 @@ using EpicGames.Redis.Utility;
 using EpicGames.Slack;
 using EpicGames.Slack.Blocks;
 using EpicGames.Slack.Elements;
+using Horde.Build.Agents;
+using Horde.Build.Configuration;
 using Horde.Build.Devices;
 using Horde.Build.Issues;
 using Horde.Build.Issues.External;
@@ -355,6 +358,13 @@ namespace Horde.Build.Notifications.Sinks
 			return await _messageStates.Find(filter).FirstOrDefaultAsync();
 		}
 
+		async Task<bool> DeleteMessageStateAsync(string recipient, string eventId)
+		{
+			FilterDefinition<MessageStateDocument> filter = Builders<MessageStateDocument>.Filter.Eq(x => x.Recipient, recipient) & Builders<MessageStateDocument>.Filter.Eq(x => x.EventId, eventId);
+			DeleteResult result = await _messageStates.DeleteOneAsync(filter);
+			return result.DeletedCount > 0;
+		}
+
 		async Task UpdateMessageStateAsync(ObjectId stateId, SlackMessageId id, string? permalink = null)
 		{
 			FilterDefinition<MessageStateDocument> filter = Builders<MessageStateDocument>.Filter.Eq(x => x.Id, stateId);
@@ -1877,6 +1887,110 @@ namespace Horde.Build.Notifications.Sinks
 
 		#region Stream updates
 
+		/// <inheritdoc/>
+		public async Task NotifyConfigUpdateAsync(Exception? ex)
+		{
+			if (String.IsNullOrEmpty(_settings.ConfigNotificationChannel))
+			{
+				return;
+			}
+
+			const string EventId = "config-update";
+
+			if (ex != null)
+			{
+				_logger.LogInformation(ex, "Sending config update failure notification: {Message}", ex.Message);
+
+				List<string> details = new List<string>();
+				if (ex is ConfigException configEx)
+				{
+					ConfigContext context = configEx.Context;
+					if (context.IncludeStack.TryPeek(out IConfigFile? blame))
+					{
+						string line = String.Empty;
+						if (configEx.InnerException is JsonException jsonEx && jsonEx.LineNumber != null)
+						{
+							line = $"({jsonEx.LineNumber})";
+						}
+
+						string file = FormatConfigFileUri(blame.Uri);
+						details.Add($"Error parsing `{file}{line}`:");
+						details.Add($"```{ex.Message}```");
+						details.Add("Include stack:\n```" + String.Join("\n", context.IncludeStack.Select(x => FormatConfigFileRevision(x))) + "```");
+
+						if (blame.Uri.Scheme == PerforceConfigSource.Scheme)
+						{
+							string blameMessage = $"Possibly due to CL {blame.Revision}";
+							if (blame.Author != null)
+							{
+								string userId = await FormatMentionAsync(blame.Author.Id, true);
+								blameMessage += $" ({FormatUserOrGroupMention(userId)})";
+							}
+							details.Add(blameMessage.ToString());
+						}
+					}
+				}
+
+				if (details.Count == 0)
+				{
+					details.Add(QuoteText(ex.Message));
+				}
+
+				string message = String.Join("\n", details);
+				string digest = GetMessageDigest(message);
+
+				MessageStateDocument? state = await GetMessageStateAsync(_settings.ConfigNotificationChannel, EventId);
+				if (state == null || state.Digest != digest)
+				{
+					SlackMessage header = new SlackMessage();
+					header.AddHeader($"Config Update Error");
+					await SendMessageAsync(_settings.ConfigNotificationChannel, header);
+
+					SlackMessageId? messageId = await SendMessageAsync(_settings.ConfigNotificationChannel, message);
+					await AddOrUpdateMessageStateAsync(_settings.ConfigNotificationChannel, EventId, null, digest, messageId);
+				}
+			}
+			else
+			{
+				if (await DeleteMessageStateAsync(_settings.ConfigNotificationChannel, EventId))
+				{
+					SlackMessage message = new SlackMessage();
+					message.AddSection($"*Config Update Succeeded*");
+					await SendMessageAsync(_settings.ConfigNotificationChannel, message);
+					await DeleteMessageStateAsync(_settings.ConfigNotificationChannel, EventId);
+				}
+			}
+		}
+
+		static string FormatConfigFileUri(Uri uri)
+		{
+			if (uri.Scheme == FileConfigSource.Scheme)
+			{
+				return Uri.UnescapeDataString(uri.AbsolutePath).Replace('/', Path.DirectorySeparatorChar);
+			}
+			else if (uri.Scheme == PerforceConfigSource.Scheme)
+			{
+				return $"/{uri.AbsolutePath}";
+			}
+			else
+			{
+				return uri.ToString();
+			}
+		}
+
+		static string FormatConfigFileRevision(IConfigFile file)
+		{
+			Uri uri = file.Uri;
+			if (uri.Scheme == PerforceConfigSource.Scheme)
+			{
+				return $"{FormatConfigFileUri(uri)}@{file.Revision}";
+			}
+			else
+			{
+				return FormatConfigFileUri(uri);
+			}
+		}
+
 		/// <inheritdoc/>
 		public async Task NotifyConfigUpdateFailureAsync(string errorMessage, string fileName, int? change = null, IUser? author = null, string? description = null)
 		{
@@ -2137,9 +2251,14 @@ namespace Horde.Build.Notifications.Sinks
 			return await SendOrUpdateMessageToThreadAsync(recipient, eventId, userId, null, message);
 		}
 
+		private static string GetMessageDigest(SlackMessage message)
+		{
+			return ContentHash.MD5(JsonSerializer.Serialize(message, new JsonSerializerOptions { DefaultIgnoreCondition = JsonIgnoreCondition.WhenWritingNull })).ToString();
+		}
+
 		private async Task<(MessageStateDocument, bool)> SendOrUpdateMessageToThreadAsync(string recipient, string eventId, UserId? userId, SlackMessageId? threadId, SlackMessage message)
 		{
-			string requestDigest = ContentHash.MD5(JsonSerializer.Serialize(message, new JsonSerializerOptions { DefaultIgnoreCondition = JsonIgnoreCondition.WhenWritingNull })).ToString();
+			string requestDigest = GetMessageDigest(message);
 
 			MessageStateDocument? prevState = await GetMessageStateAsync(recipient, eventId);
 			if (prevState != null && prevState.Digest == requestDigest)
diff --git a/Engine/Source/Programs/Horde/Horde.Build/ServerSettings.cs b/Engine/Source/Programs/Horde/Horde.Build/ServerSettings.cs
--- a/Engine/Source/Programs/Horde/Horde.Build/ServerSettings.cs
+++ b/Engine/Source/Programs/Horde/Horde.Build/ServerSettings.cs
@@ -471,6 +471,11 @@ namespace Horde.Build
 		/// </summary>
 		public string SlackWarningPrefix { get; set; } = ":horde-warning: ";
 
+		/// <summary>
+		/// Channel for sending messages related to config update failures
+		/// </summary>
+		public string? ConfigNotificationChannel { get; set; }
+
 		/// <summary>
 		/// Channel to send stream notification update failures to
 		/// </summary>
