Horde: Fix tests.

diff --git a/Engine/Source/Programs/Horde/Horde.Build.Tests/ConfigTests.cs b/Engine/Source/Programs/Horde/Horde.Build.Tests/ConfigTests.cs
--- a/Engine/Source/Programs/Horde/Horde.Build.Tests/ConfigTests.cs
+++ b/Engine/Source/Programs/Horde/Horde.Build.Tests/ConfigTests.cs
@@ -111,15 +111,15 @@ namespace Horde.Build.Tests
 			Assert.AreEqual(result.TestString, "hello");
 
 			// Check it returns the same object if the timestamp hasn't changed
-			IConfigData file1 = await source.GetAsync(fooUri, cancellationToken);
+			IConfigFile file1 = await source.GetAsync(fooUri, cancellationToken);
 			await Task.Delay(TimeSpan.FromSeconds(1));
-			IConfigData file2 = await source.GetAsync(fooUri, cancellationToken);
+			IConfigFile file2 = await source.GetAsync(fooUri, cancellationToken);
 			Assert.IsTrue(ReferenceEquals(file1, file2));
 
 			// Check it returns a new object if the timestamp HAS changed
 			await Task.Delay(TimeSpan.FromSeconds(1));
 			await FileReference.WriteAllBytesAsync(new FileReference(fooUri.LocalPath), data, cancellationToken);
-			IConfigData file3 = await source.GetAsync(fooUri, cancellationToken);
+			IConfigFile file3 = await source.GetAsync(fooUri, cancellationToken);
 			Assert.IsTrue(!ReferenceEquals(file1, file3));
 		}
 	}
diff --git a/Engine/Source/Programs/Horde/Horde.Build.Tests/NotificationServiceTests.cs b/Engine/Source/Programs/Horde/Horde.Build.Tests/NotificationServiceTests.cs
--- a/Engine/Source/Programs/Horde/Horde.Build.Tests/NotificationServiceTests.cs
+++ b/Engine/Source/Programs/Horde/Horde.Build.Tests/NotificationServiceTests.cs
@@ -41,6 +41,7 @@ namespace Horde.Build.Tests
 		public Task NotifyJobStepCompleteAsync(IUser user, IStream jobStream, IJob job, IJobStepBatch batch, IJobStep step, INode node, List<ILogEventData> jobStepEventData) { throw new NotImplementedException(); }
 		public Task NotifyLabelCompleteAsync(IUser user, IJob job, IStream stream, ILabel label, int labelIdx, LabelOutcome outcome, List<(string, JobStepOutcome, Uri)> stepData) { throw new NotImplementedException(); }
 		public Task NotifyIssueUpdatedAsync(IIssue issue) { throw new NotImplementedException(); }
+		public Task NotifyConfigUpdateAsync(Exception? ex) => Task.CompletedTask;
 		public Task NotifyConfigUpdateFailureAsync(string errorMessage, string fileName, int? change = null, IUser? author = null, string? description = null) { throw new NotImplementedException(); }
 		public Task NotifyDeviceServiceAsync(string message, IDevice? device = null, IDevicePool? pool = null, IStream? stream = null, IJob? job = null, IJobStep? step = null, INode? node = null, IUser? user = null) { throw new NotImplementedException(); }
 		public Task SendIssueReportAsync(IssueReportGroup report) => throw new NotImplementedException();
