Horde: Refactor Tools collection to have an interface more similar to other collections, with an interface for each document and config classes in a separate file.

diff --git a/Engine/Source/Programs/Horde/Horde.Build.Tests/ToolTests.cs b/Engine/Source/Programs/Horde/Horde.Build.Tests/ToolTests.cs
--- a/Engine/Source/Programs/Horde/Horde.Build.Tests/ToolTests.cs
+++ b/Engine/Source/Programs/Horde/Horde.Build.Tests/ToolTests.cs
@@ -13,8 +13,8 @@ using Microsoft.VisualStudio.TestTools.UnitTesting;
 
 namespace Horde.Build.Tests
 {
-	using ToolId = StringId<Tool>;
-	using ToolDeploymentId = ObjectId<ToolDeployment>;
+	using ToolId = StringId<ITool>;
+	using ToolDeploymentId = ObjectId<IToolDeployment>;
 
 	[TestClass]
     public class ToolTests : TestSetup
@@ -25,14 +25,14 @@ namespace Horde.Build.Tests
 			ToolCollection collection = ServiceProvider.GetRequiredService<ToolCollection>();
 
 			ToolId toolId = new ToolId("ugs");
-			await collection.ConfigureAsync(new List<ToolOptions> { new ToolOptions(toolId) { Name = "UnrealGameSync", Description = "Tool for syncing content from source control" } });
+			await collection.ConfigureAsync(new List<ToolConfig> { new ToolConfig(toolId) { Name = "UnrealGameSync", Description = "Tool for syncing content from source control" } });
 
-			Tool tool = Deref(await collection.GetAsync(toolId));
+			ITool tool = Deref(await collection.GetAsync(toolId));
 			Assert.AreEqual(tool.Id, new ToolId("ugs"));
 			Assert.AreEqual(tool.Name, "UnrealGameSync");
 			Assert.AreEqual(tool.Description, "Tool for syncing content from source control");
 
-			Tool tool2 = Deref(await collection.GetAsync(tool.Id));
+			ITool tool2 = Deref(await collection.GetAsync(tool.Id));
 			Assert.AreEqual(tool.Id, tool2.Id);
 			Assert.AreEqual(tool.LastUpdateTime, tool2.LastUpdateTime);
 			Assert.AreEqual(tool.Name, tool2.Name);
@@ -45,9 +45,9 @@ namespace Horde.Build.Tests
 			ToolCollection collection = ServiceProvider.GetRequiredService<ToolCollection>();
 
 			ToolId toolId = new ToolId("ugs");
-			await collection.ConfigureAsync(new List<ToolOptions> { new ToolOptions(toolId) { Name = "UnrealGameSync", Description = "Tool for syncing content from source control" } });
+			await collection.ConfigureAsync(new List<ToolConfig> { new ToolConfig(toolId) { Name = "UnrealGameSync", Description = "Tool for syncing content from source control" } });
 
-			Tool tool = Deref(await collection.GetAsync(toolId));
+			ITool tool = Deref(await collection.GetAsync(toolId));
 			Assert.AreEqual(new ToolId("ugs"), tool.Id);
 			Assert.AreEqual("UnrealGameSync", tool.Name);
 			Assert.AreEqual("Tool for syncing content from source control", tool.Description);
@@ -58,7 +58,7 @@ namespace Horde.Build.Tests
 			ToolDeploymentId deploymentId;
 			using (MemoryStream stream = new MemoryStream(deploymentData))
 			{
-				tool = Deref(await collection.CreateDeploymentAsync(tool, new ToolDeploymentOptions { Version = "1.0", Duration = TimeSpan.FromMinutes(5.0), CreatePaused = true }, stream));
+				tool = Deref(await collection.CreateDeploymentAsync(tool, new ToolDeploymentConfig { Version = "1.0", Duration = TimeSpan.FromMinutes(5.0), CreatePaused = true }, stream));
 				Assert.AreEqual(1, tool.Deployments.Count);
 				Assert.IsNull(tool.Deployments[0].StartedAt);
 				deploymentId = tool.Deployments[^1].Id;
diff --git a/Engine/Source/Programs/Horde/Horde.Build/Server/GlobalConfig.cs b/Engine/Source/Programs/Horde/Horde.Build/Server/GlobalConfig.cs
--- a/Engine/Source/Programs/Horde/Horde.Build/Server/GlobalConfig.cs
+++ b/Engine/Source/Programs/Horde/Horde.Build/Server/GlobalConfig.cs
@@ -79,7 +79,7 @@ namespace Horde.Build.Server
 		/// <summary>
 		/// List of tools hosted by the server
 		/// </summary>
-		public List<ToolOptions> Tools { get; set; } = new List<ToolOptions>();
+		public List<ToolConfig> Tools { get; set; } = new List<ToolConfig>();
 
 		/// <summary>
 		/// Maximum number of conforms to run at once
diff --git a/Engine/Source/Programs/Horde/Horde.Build/Tools/ITool.cs b/Engine/Source/Programs/Horde/Horde.Build/Tools/ITool.cs
new file mode 100644
--- /dev/null
+++ b/Engine/Source/Programs/Horde/Horde.Build/Tools/ITool.cs
@@ -0,0 +1,173 @@
+// Copyright Epic Games, Inc. All Rights Reserved.
+
+using EpicGames.Horde.Storage;
+using Horde.Build.Acls;
+using Horde.Build.Utilities;
+using System;
+using System.Collections.Generic;
+using System.Security.Claims;
+using System.Threading.Tasks;
+
+namespace Horde.Build.Tools
+{
+	using ToolId = StringId<ITool>;
+	using ToolDeploymentId = ObjectId<IToolDeployment>;
+
+	/// <summary>
+	/// Describes a standalone, external tool hosted and deployed by Horde. Provides basic functionality for performing
+	/// gradual roll-out, versioning, etc...
+	/// </summary>
+	public interface ITool
+	{
+		/// <summary>
+		/// Identifier for the tool
+		/// </summary>
+		ToolId Id { get; }
+
+		/// <summary>
+		/// Name of the tool
+		/// </summary>
+		string Name { get; }
+
+		/// <summary>
+		/// Description for the tool
+		/// </summary>
+		string Description { get; }
+
+		/// <summary>
+		/// Last time that the tool was updated
+		/// </summary>
+		DateTime LastUpdateTime { get; }
+
+		/// <summary>
+		/// Current deployments of this tool, sorted by time.
+		/// </summary>
+		IReadOnlyList<IToolDeployment> Deployments { get; }
+
+		/// <summary>
+		/// Whether this tool should be exposed for download on a public endpoint without authentication
+		/// </summary>
+		bool Public { get; }
+
+		/// <summary>
+		/// Access list for this tool
+		/// </summary>
+		public AclV2? Acl { get; }
+	}
+
+	/// <summary>
+	/// Current state of a tool's deployment
+	/// </summary>
+	public enum ToolDeploymentState
+	{
+		/// <summary>
+		/// The deployment is ongoing
+		/// </summary>
+		Active,
+
+		/// <summary>
+		/// The deployment should be paused at its current state
+		/// </summary>
+		Paused,
+
+		/// <summary>
+		/// Deployment of this version is complete
+		/// </summary>
+		Complete,
+
+		/// <summary>
+		/// The deployment has been cancelled.
+		/// </summary>
+		Cancelled,
+	}
+
+	/// <summary>
+	/// Deployment of a tool
+	/// </summary>
+	public interface IToolDeployment
+	{
+		/// <summary>
+		/// Identifier for this deployment. A new identifier will be assigned to each created instance, so an identifier corresponds to a unique deployment.
+		/// </summary>
+		ToolDeploymentId Id { get; }
+
+		/// <summary>
+		/// Descriptive version string for this tool revision
+		/// </summary>
+		string Version { get; }
+
+		/// <summary>
+		/// Current state of this deployment
+		/// </summary>
+		ToolDeploymentState State { get; }
+
+		/// <summary>
+		/// Current progress of the deployment
+		/// </summary>
+		double Progress { get; }
+
+		/// <summary>
+		/// Last time at which the progress started. Set to null if the deployment was paused.
+		/// </summary>
+		DateTime? StartedAt { get; }
+
+		/// <summary>
+		/// Length of time over which to make the deployment
+		/// </summary>
+		TimeSpan Duration { get; }
+
+		/// <summary>
+		/// Mime-type for the data
+		/// </summary>
+		string MimeType { get; }
+
+		/// <summary>
+		/// Reference to this tool in Horde Storage.
+		/// </summary>
+		RefId RefId { get; }
+	}
+
+	/// <summary>
+	/// Extension methods for tools
+	/// </summary>
+	public static class ToolExtensions
+	{
+		/// <summary>
+		/// Test whether a user can perform an action on this tool 
+		/// </summary>
+		public static Task<bool> AuthorizeAsync(this ITool tool, AclAction action, ClaimsPrincipal user, AclService aclService, GlobalPermissionsCache? cache)
+		{
+			bool? result = tool.Acl?.Authorize(action, user);
+			if (result == null)
+			{
+				return aclService.AuthorizeAsync(action, user, cache);
+			}
+			else
+			{
+				return Task.FromResult(result.Value);
+			}
+		}
+
+		/// <summary>
+		/// Get the progress fraction for a particular deployment and time
+		/// </summary>
+		/// <param name="deployment"></param>
+		/// <param name="utcNow"></param>
+		/// <returns></returns>
+		public static double GetProgressValue(this IToolDeployment deployment, DateTime utcNow)
+		{
+			if (deployment.StartedAt == null)
+			{
+				return deployment.Progress;
+			}
+			else if (deployment.Duration > TimeSpan.Zero)
+			{
+				return Math.Clamp((utcNow - deployment.StartedAt.Value) / deployment.Duration, 0.0, 1.0);
+			}
+			else
+			{
+				return 1.0;
+			}
+		}
+	}
+}
diff --git a/Engine/Source/Programs/Horde/Horde.Build/Tools/Tool.cs b/Engine/Source/Programs/Horde/Horde.Build/Tools/Tool.cs
deleted file mode 100644
diff --git a/Engine/Source/Programs/Horde/Horde.Build/Tools/ToolCollection.cs b/Engine/Source/Programs/Horde/Horde.Build/Tools/ToolCollection.cs
--- a/Engine/Source/Programs/Horde/Horde.Build/Tools/ToolCollection.cs
+++ b/Engine/Source/Programs/Horde/Horde.Build/Tools/ToolCollection.cs
@@ -12,6 +12,7 @@ using Microsoft.Extensions.Logging;
 using Microsoft.Extensions.Options;
 using MongoDB.Bson;
 using MongoDB.Bson.Serialization;
+using MongoDB.Bson.Serialization.Attributes;
 using MongoDB.Driver;
 using StackExchange.Redis;
 using System;
@@ -22,14 +23,144 @@ using System.Threading.Tasks;
 
 namespace Horde.Build.Tools
 {
-	using ToolId = StringId<Tool>;
-	using ToolDeploymentId = ObjectId<ToolDeployment>;
+	using ToolId = StringId<ITool>;
+	using ToolDeploymentId = ObjectId<IToolDeployment>;
 
 	/// <summary>
 	/// Collection of tool documents
 	/// </summary>
 	public class ToolCollection
 	{
+		class Tool : VersionedDocument<ToolId, Tool>, ITool
+		{
+			/// <summary>
+			/// Name of the tool
+			/// </summary>
+			[BsonElement("nam")]
+			public string Name { get; set; }
+
+			/// <summary>
+			/// Description for the tool
+			/// </summary>
+			[BsonElement("dsc")]
+			public string Description { get; set; }
+
+			[BsonElement("dep")]
+			public List<ToolDeployment> Deployments { get; set; } = new List<ToolDeployment>();
+
+			IReadOnlyList<IToolDeployment> ITool.Deployments => Deployments;
+
+			[BsonElement("pub")]
+			public bool Public { get; set; }
+
+			[BsonElement("acl")]
+			public AclV2? Acl
+			{
+				get => _acl;
+				set => _acl = (value is not null && !value.IsDefault()) ? value : null;
+			}
+
+			private AclV2? _acl;
+
+			[BsonConstructor]
+			public Tool(ToolId id)
+				: base(id)
+			{
+				Id = id;
+				Name = id.ToString();
+				Description = String.Empty;
+			}
+
+			public Tool(ToolConfig config)
+				: base(config.Id)
+			{
+				Name = config.Name;
+				Description = config.Description;
+				Public = config.Public;
+				Acl = config.Acl;
+			}
+
+			/// <inheritdoc/>
+			public override Tool UpgradeToLatest() => this;
+
+			public void UpdateTemporalState(DateTime utcNow)
+			{
+				foreach (ToolDeployment deployment in Deployments)
+				{
+					deployment.UpdateTemporalState(utcNow);
+				}
+			}
+		}
+
+		class ToolDeployment : IToolDeployment
+		{
+			public ToolDeploymentId Id { get; set; }
+
+			[BsonElement("ver")]
+			public string Version { get; set; }
+
+			[BsonIgnore]
+			public ToolDeploymentState State { get; set; }
+
+			[BsonIgnore]
+			public double Progress { get; set; }
+
+			[BsonElement("bpr")]
+			public double BaseProgress { get; set; }
+
+			[BsonElement("stm")]
+			public DateTime? StartedAt { get; set; }
+
+			[BsonElement("dur")]
+			public TimeSpan Duration { get; set; }
+
+			[BsonElement("mtp")]
+			public string MimeType { get; set; }
+
+			[BsonElement("ref")]
+			public RefId RefId { get; set; }
+
+			public ToolDeployment(ToolDeploymentId id)
+			{
+				Id = id;
+				Version = String.Empty;
+				MimeType = ToolDeploymentConfig.DefaultMimeType;
+			}
+
+			public ToolDeployment(ToolDeploymentId id, ToolDeploymentConfig options, RefId refId)
+			{
+				Id = id;
+				Version = options.Version;
+				Duration = options.Duration;
+				MimeType = options.MimeType;
+				RefId = refId;
+			}
+
+			public void UpdateTemporalState(DateTime utcNow)
+			{
+				if (BaseProgress >= 1.0)
+				{
+					State = ToolDeploymentState.Complete;
+					Progress = 1.0;
+				}
+				else if (StartedAt == null)
+				{
+					State = ToolDeploymentState.Paused;
+					Progress = BaseProgress;
+				}
+				else if (Duration > TimeSpan.Zero)
+				{
+					State = ToolDeploymentState.Active;
+					Progress = Math.Clamp((utcNow - StartedAt.Value) / Duration, 0.0, 1.0);
+				}
+				else
+				{
+					State = ToolDeploymentState.Complete;
+					Progress = 1.0;
+				}
+			}
+		}
+
 		private class ToolDeploymentData
 		{
 			[CbField]
@@ -102,7 +233,7 @@ namespace Horde.Build.Tools
 		/// Configures all the available tools
 		/// </summary>
 		/// <param name="tools">The list of configured tools</param>
-		public async Task ConfigureAsync(List<ToolOptions> tools)
+		public async Task ConfigureAsync(List<ToolConfig> tools)
 		{
 			List<ToolId> ids = await FindAllIdsAsync();
 
@@ -113,7 +244,7 @@ namespace Horde.Build.Tools
 				await _tools.DeleteAsync(removeId);
 			}
 
-			foreach (ToolOptions tool in tools)
+			foreach (ToolConfig tool in tools)
 			{
 				if (!ids.Contains(tool.Id))
 				{
@@ -128,7 +259,7 @@ namespace Horde.Build.Tools
 		/// </summary>
 		/// <param name="options">Options for the tool</param>
 		/// <returns>Document describing the tool</returns>
-		private async Task<Tool> AddOrUpdateAsync(ToolOptions options)
+		private async Task<ITool> AddOrUpdateAsync(ToolConfig options)
 		{
 			for (; ; )
 			{
@@ -236,17 +367,17 @@ namespace Horde.Build.Tools
 		/// Finds all tools matching a set of criteria
 		/// </summary>
 		/// <returns>Sequence of tool documents</returns>
-		public async Task<List<Tool>> FindAllAsync()
+		public async Task<List<ITool>> FindAllAsync()
 		{
 			// Fetch all the tools in parallel
 			IEnumerable<ToolId> toolIds =await FindAllIdsAsync();
-			List<Task<Tool?>> tasks = toolIds.Select(x => Task.Run(() => GetAsync(x))).ToList();
+			List<Task<ITool?>> tasks = toolIds.Select(x => Task.Run(() => GetAsync(x))).ToList();
 
 			// Return all the successful results
-			List<Tool> results = new List<Tool>();
-			foreach (Task<Tool?> task in tasks)
+			List<ITool> results = new List<ITool>();
+			foreach (Task<ITool?> task in tasks)
 			{
-				Tool? tool = await task;
+				ITool? tool = await task;
 				if (tool != null)
 				{
 					results.Add(tool);
@@ -260,7 +391,14 @@ namespace Horde.Build.Tools
 		/// </summary>
 		/// <param name="id">The tool identifier</param>
 		/// <returns></returns>
-		public async Task<Tool?> GetAsync(ToolId id)
+		public async Task<ITool?> GetAsync(ToolId id) => await GetInternalAsync(id);
+
+		/// <summary>
+		/// Gets a tool with the given identifier
+		/// </summary>
+		/// <param name="id">The tool identifier</param>
+		/// <returns></returns>
+		async Task<Tool?> GetInternalAsync(ToolId id)
 		{
 			Tool? tool = await _tools.GetAsync(id);
 			if (tool != null)
@@ -277,7 +415,7 @@ namespace Horde.Build.Tools
 		/// <param name="options">Options for the new deployment</param>
 		/// <param name="stream">Stream containing the tool data</param>
 		/// <returns>Updated tool document, or null if it does not exist</returns>
-		public async Task<Tool?> CreateDeploymentAsync(Tool tool, ToolDeploymentOptions options, Stream stream)
+		public async Task<ITool?> CreateDeploymentAsync(ITool tool, ToolDeploymentConfig options, Stream stream)
 		{
 			// Upload the tool data first
 			IoHash hash = await _storage.WriteBlobAsync(_namespaceId, stream);
@@ -302,7 +440,7 @@ namespace Horde.Build.Tools
 			}
 
 			// Create the deployment
-			Tool? newTool = tool;
+			Tool? newTool = (Tool)tool;
 			for (; ; )
 			{
 				newTool = await TryAddDeploymentAsync(newTool, deployment);
@@ -311,7 +449,7 @@ namespace Horde.Build.Tools
 					break;
 				}
 
-				newTool = await GetAsync(tool.Id);
+				newTool = await GetInternalAsync(tool.Id);
 				if (newTool == null)
 				{
 					return null;
@@ -350,7 +488,12 @@ namespace Horde.Build.Tools
 		/// <param name="deploymentId">Identifier for the deployment to modify</param>
 		/// <param name="action">New state of the deployment</param>
 		/// <returns></returns>
-		public async Task<Tool?> UpdateDeploymentAsync(Tool tool, ToolDeploymentId deploymentId, ToolDeploymentState action)
+		public async Task<ITool?> UpdateDeploymentAsync(ITool tool, ToolDeploymentId deploymentId, ToolDeploymentState action)
+		{
+			return await UpdateDeploymentInternalAsync((Tool)tool, deploymentId, action);
+		}
+
+		async Task<Tool?> UpdateDeploymentInternalAsync(Tool tool, ToolDeploymentId deploymentId, ToolDeploymentState action)
 		{
 			int idx = tool.Deployments.FindIndex(x => x.Id == deploymentId);
 			if (idx == -1)
@@ -401,7 +544,7 @@ namespace Horde.Build.Tools
 		/// <param name="tool">Identifier for the tool</param>
 		/// <param name="deployment">The deployment</param>
 		/// <returns>Stream for the data</returns>
-		public async Task<Stream> GetDeploymentPayloadAsync(Tool tool, ToolDeployment deployment)
+		public async Task<Stream> GetDeploymentPayloadAsync(ITool tool, IToolDeployment deployment)
 		{
 			ToolDeploymentData data = await _storage.GetRefAsync<ToolDeploymentData>(_namespaceId, GetBucket(tool.Id), deployment.RefId);
 			return await _storage.ReadBlobAsync(_namespaceId, data.Data);
diff --git a/Engine/Source/Programs/Horde/Horde.Build/Tools/ToolConfig.cs b/Engine/Source/Programs/Horde/Horde.Build/Tools/ToolConfig.cs
new file mode 100644
--- /dev/null
+++ b/Engine/Source/Programs/Horde/Horde.Build/Tools/ToolConfig.cs
@@ -0,0 +1,80 @@
+// Copyright Epic Games, Inc. All Rights Reserved.
+
+using System;
+using System.ComponentModel.DataAnnotations;
+using System.Diagnostics;
+using Horde.Build.Acls;
+using Horde.Build.Utilities;
+
+namespace Horde.Build.Tools
+{
+	using ToolId = StringId<ITool>;
+
+	/// <summary>
+	/// Options for configuring a tool
+	/// </summary>
+	[DebuggerDisplay("{Id}")]
+	public class ToolConfig
+	{
+		/// <inheritdoc cref="VersionedDocument{ToolId, Tool}.Id"/>
+		[Required]
+		public ToolId Id { get; set; }
+
+		/// <inheritdoc cref="ITool.Name"/>
+		[Required]
+		public string Name { get; set; }
+
+		/// <inheritdoc cref="ITool.Description"/>
+		public string Description { get; set; }
+
+		/// <inheritdoc cref="ITool.Public"/>
+		public bool Public { get; set; }
+
+		/// <inheritdoc cref="ITool.Acl"/>
+		public AclV2? Acl { get; set; }
+
+		/// <summary>
+		/// Default constructor for serialization
+		/// </summary>
+		public ToolConfig()
+		{
+			Name = String.Empty;
+			Description = String.Empty;
+		}
+
+		/// <summary>
+		/// Constructor
+		/// </summary>
+		public ToolConfig(ToolId id)
+		{
+			Id = id;
+			Name = id.ToString();
+			Description = String.Empty;
+		}
+	}
+
+	/// <summary>
+	/// Options for a new deployment
+	/// </summary>
+	public class ToolDeploymentConfig
+	{
+		/// <summary>
+		/// Default mime types for deployment data
+		/// </summary>
+		public const string DefaultMimeType = "application/zip";
+
+		/// <inheritdoc cref="IToolDeployment.Version"/>
+		public string Version { get; set; } = "Unknown";
+
+		/// <inheritdoc cref="IToolDeployment.Duration"/>
+		public TimeSpan Duration { get; set; }
+
+		/// <inheritdoc cref="IToolDeployment.MimeType"/>
+		public string MimeType { get; set; } = DefaultMimeType;
+
+		/// <summary>
+		/// Whether to create the deployment in a paused state
+		/// </summary>
+		public bool CreatePaused { get; set; }
+	}
+}
diff --git a/Engine/Source/Programs/Horde/Horde.Build/Tools/ToolsController.cs b/Engine/Source/Programs/Horde/Horde.Build/Tools/ToolsController.cs
--- a/Engine/Source/Programs/Horde/Horde.Build/Tools/ToolsController.cs
+++ b/Engine/Source/Programs/Horde/Horde.Build/Tools/ToolsController.cs
@@ -1,6 +1,7 @@
 ﻿// Copyright Epic Games, Inc. All Rights Reserved.
 
 using EpicGames.Core;
+using EpicGames.Horde.Storage;
 using Horde.Build.Acls;
 using Horde.Build.Utilities;
 using HordeCommon;
@@ -9,14 +10,85 @@ using Microsoft.AspNetCore.Http;
 using Microsoft.AspNetCore.Mvc;
 using Microsoft.Extensions.Logging;
 using System;
+using System.Collections.Generic;
 using System.IO;
 using System.Linq;
 using System.Threading.Tasks;
 
 namespace Horde.Build.Tools
 {
-	using ToolId = StringId<Tool>;
-	using ToolDeploymentId = ObjectId<ToolDeployment>;
+	using ToolId = StringId<ITool>;
+	using ToolDeploymentId = ObjectId<IToolDeployment>;
+
+	/// <summary>
+	/// Describes a standalone, external tool hosted and deployed by Horde. Provides basic functionality for performing
+	/// gradual roll-out, versioning, etc...
+	/// </summary>
+	public class GetToolResponse
+	{
+		readonly ITool _tool;
+
+		/// <inheritdoc cref="VersionedDocument{TId, TLatest}.Id"/>
+		public ToolId Id => _tool.Id;
+
+		/// <inheritdoc cref="ITool.Name"/>
+		public string Name => _tool.Name;
+
+		/// <inheritdoc cref="ITool.Description"/>
+		public string Description => _tool.Description;
+
+		/// <inheritdoc cref="ITool.Deployments"/>
+		public List<GetToolDeploymentResponse> Deployments { get; }
+
+		/// <inheritdoc cref="ITool.Public"/>
+		public bool Public => _tool.Public;
+
+		/// <summary>
+		/// Constructor
+		/// </summary>
+		public GetToolResponse(ITool tool)
+		{
+			_tool = tool;
+			Deployments = tool.Deployments.ConvertAll(x => new GetToolDeploymentResponse(x));
+		}
+	}
+
+	/// <summary>
+	/// Response object describing the deployment of a tool
+	/// </summary>
+	public class GetToolDeploymentResponse
+	{
+		readonly IToolDeployment _deployment;
+
+		/// <inheritdoc cref="IToolDeployment.Id"/>
+		public ToolDeploymentId Id => _deployment.Id;
+
+		/// <inheritdoc cref="IToolDeployment.Version"/>
+		public string Version => _deployment.Version;
+
+		/// <inheritdoc cref="IToolDeployment.State"/>
+		public ToolDeploymentState State => _deployment.State;
+
+		/// <inheritdoc cref="IToolDeployment.Progress"/>
+		public double Progress => _deployment.Progress;
+
+		/// <inheritdoc cref="IToolDeployment.StartedAt"/>
+		public DateTime? StartedAt => _deployment.StartedAt;
+
+		/// <inheritdoc cref="IToolDeployment.Duration"/>
+		public TimeSpan Duration => _deployment.Duration;
+
+		/// <inheritdoc cref="IToolDeployment.MimeType"/>
+		public string MimeType => _deployment.MimeType;
+
+		/// <inheritdoc cref="IToolDeployment.RefId"/>
+		public RefId RefId => _deployment.RefId;
+
+		/// <summary>
+		/// Constructor
+		/// </summary>
+		public GetToolDeploymentResponse(IToolDeployment deployment) => _deployment = deployment;
+	}
 
 	/// <summary>
 	/// Response from creating a deployment
@@ -82,9 +154,9 @@ namespace Horde.Build.Tools
 		/// <returns>Information about the registered agent</returns>
 		[HttpPost]
 		[Route("/api/v1/tools/{id}/deployments")]
-		public async Task<ActionResult<CreateDeploymentResponse>> CreateDeploymentAsync(ToolId id, [FromForm] ToolDeploymentOptions options, [FromForm] IFormFile file)
+		public async Task<ActionResult<CreateDeploymentResponse>> CreateDeploymentAsync(ToolId id, [FromForm] ToolDeploymentConfig options, [FromForm] IFormFile file)
 		{
-			Tool? tool = await _toolCollection.GetAsync(id);
+			ITool? tool = await _toolCollection.GetAsync(id);
 			if (tool == null)
 			{
 				return NotFound(id);
@@ -114,7 +186,7 @@ namespace Horde.Build.Tools
 		[Route("/api/v1/tools/{id}/deployments/{deploymentId}")]
 		public async Task<ActionResult> GetDeploymentAsync(ToolId id, ToolDeploymentId deploymentId, [FromQuery] GetToolAction action = GetToolAction.Info)
 		{
-			Tool? tool = await _toolCollection.GetAsync(id);
+			ITool? tool = await _toolCollection.GetAsync(id);
 			if (tool == null)
 			{
 				return NotFound(id);
@@ -124,7 +196,7 @@ namespace Horde.Build.Tools
 				return Forbid(AclAction.DownloadTool, id);
 			}
 
-			ToolDeployment? deployment = tool.Deployments.FirstOrDefault(x => x.Id == deploymentId);
+			IToolDeployment? deployment = tool.Deployments.FirstOrDefault(x => x.Id == deploymentId);
 			if(deployment == null)
 			{
 				return NotFound(id, deploymentId);
@@ -133,11 +205,12 @@ namespace Horde.Build.Tools
 			return await GetDeploymentResponseAsync(tool, deployment, action);
 		}
 
-		private async Task<ActionResult> GetDeploymentResponseAsync(Tool tool, ToolDeployment deployment, GetToolAction action)
+		private async Task<ActionResult> GetDeploymentResponseAsync(ITool tool, IToolDeployment deployment, GetToolAction action)
 		{
 			if (action == GetToolAction.Info)
 			{
-				return Ok(deployment);
+				GetToolDeploymentResponse response = new GetToolDeploymentResponse(deployment);
+				return Ok(response);
 			}
 			else
 			{
@@ -153,7 +226,7 @@ namespace Horde.Build.Tools
 		[Route("/api/v1/tools/{id}/deployments/{deploymentId}")]
 		public async Task<ActionResult> UpdateDeploymentAsync(ToolId id, ToolDeploymentId deploymentId, [FromBody] UpdateDeploymentRequest request)
 		{
-			Tool? tool = await _toolCollection.GetAsync(id);
+			ITool? tool = await _toolCollection.GetAsync(id);
 			if (tool == null)
 			{
 				return NotFound(id);
@@ -203,7 +276,7 @@ namespace Horde.Build.Tools
 		[Route("/api/v1/tools/{id}")]
 		public async Task<ActionResult> GetToolAsync(ToolId id, GetToolAction action = GetToolAction.Info)
 		{
-			Tool? tool = await _toolCollection.GetAsync(id);
+			ITool? tool = await _toolCollection.GetAsync(id);
 			if (tool == null)
 			{
 				return NotFound(id);
@@ -239,7 +312,7 @@ namespace Horde.Build.Tools
 		[Route("/api/v1/tools/{id}/deployments")]
 		public async Task<ActionResult> FindDeploymentAsync(ToolId id, [FromQuery] double phase = 0.0, [FromQuery] GetToolAction action = GetToolAction.Info)
 		{
-			Tool? tool = await _toolCollection.GetAsync(id);
+			ITool? tool = await _toolCollection.GetAsync(id);
 			if (tool == null)
 			{
 				return NotFound(id);
@@ -274,13 +347,13 @@ namespace Horde.Build.Tools
 			}
 		}
 
-		async Task<ActionResult> GetDeploymentPayloadAsync(Tool tool, ToolDeployment deployment)
+		async Task<ActionResult> GetDeploymentPayloadAsync(ITool tool, IToolDeployment deployment)
 		{
 			Stream stream = await _toolCollection.GetDeploymentPayloadAsync(tool, deployment);
 			return new FileStreamResult(stream, deployment.MimeType);
 		}
 
-		async ValueTask<bool> AuthorizeDownloadAsync(Tool tool)
+		async ValueTask<bool> AuthorizeDownloadAsync(ITool tool)
 		{
 			if (!tool.Public)
 			{
diff --git a/Engine/Source/Programs/Horde/Horde.Build/Utilities/HordeControllerBase.cs b/Engine/Source/Programs/Horde/Horde.Build/Utilities/HordeControllerBase.cs
--- a/Engine/Source/Programs/Horde/Horde.Build/Utilities/HordeControllerBase.cs
+++ b/Engine/Source/Programs/Horde/Horde.Build/Utilities/HordeControllerBase.cs
@@ -18,8 +18,8 @@ namespace Horde.Build.Utilities
 	using LeaseId = ObjectId<ILease>;
 	using ProjectId = StringId<IProject>;
 	using StreamId = StringId<IStream>;
-	using ToolId = StringId<Tool>;
-	using ToolDeploymentId = ObjectId<ToolDeployment>;
+	using ToolId = StringId<ITool>;
+	using ToolDeploymentId = ObjectId<IToolDeployment>;
 	using TemplateId = StringId<ITemplateRef>;
 
 	/// <summary>
