Horde: Remove per-agent ACL lists. We don't use these, and it goes against the idea of having centralized configuration data.

diff --git a/Engine/Source/Programs/Horde/Horde.Build/Agents/AgentCollection.cs b/Engine/Source/Programs/Horde/Horde.Build/Agents/AgentCollection.cs
--- a/Engine/Source/Programs/Horde/Horde.Build/Agents/AgentCollection.cs
+++ b/Engine/Source/Programs/Horde/Horde.Build/Agents/AgentCollection.cs
@@ -113,7 +113,6 @@ namespace Horde.Build.Agents
 
 			public AgentCapabilities Capabilities { get; set; } = new AgentCapabilities();
 			public List<AgentLease>? Leases { get; set; }
-			public Acl? Acl { get; set; }
 			public DateTime UpdateTime { get; set; }
 			public uint UpdateIndex { get; set; }
 			public string? Comment { get; set; }
@@ -133,7 +132,6 @@ namespace Horde.Build.Agents
 			public AgentDocument(AgentId id, bool enabled, List<PoolId> pools)
 			{
 				Id = id;
-				Acl = new Acl();
 				Enabled = enabled;
 				Pools = pools;
 			}
@@ -283,7 +281,7 @@ namespace Horde.Build.Agents
 		}
 
 		/// <inheritdoc/>
-		public async Task<IAgent?> TryUpdateSettingsAsync(IAgent agentInterface, bool? enabled = null, bool? requestConform = null, bool? requestFullConform = null, bool? requestRestart = null, bool? requestShutdown = null, string? shutdownReason = null, List<PoolId>? pools = null, Acl? acl = null, string? comment = null)
+		public async Task<IAgent?> TryUpdateSettingsAsync(IAgent agentInterface, bool? enabled = null, bool? requestConform = null, bool? requestFullConform = null, bool? requestRestart = null, bool? requestShutdown = null, string? shutdownReason = null, List<PoolId>? pools = null, string? comment = null)
 		{
 			AgentDocument agent = (AgentDocument)agentInterface;
 
@@ -337,10 +335,6 @@ namespace Horde.Build.Agents
 				updates.Add(updateBuilder.Set(x => x.LastShutdownReason, shutdownReason));
 			}
 
-			if (acl != null)
-			{
-				updates.Add(Acl.CreateUpdate<AgentDocument>(x => x.Acl!, acl));
-			}
 			if (comment != null)
 			{
 				updates.Add(updateBuilder.Set(x => x.Comment, comment));
diff --git a/Engine/Source/Programs/Horde/Horde.Build/Agents/AgentMessages.cs b/Engine/Source/Programs/Horde/Horde.Build/Agents/AgentMessages.cs
--- a/Engine/Source/Programs/Horde/Horde.Build/Agents/AgentMessages.cs
+++ b/Engine/Source/Programs/Horde/Horde.Build/Agents/AgentMessages.cs
@@ -98,11 +98,6 @@ namespace Horde.Build.Agents
 		/// Pools for this agent
 		/// </summary>
 		public List<string>? Pools { get; set; }
-
-		/// <summary>
-		/// New ACL for this agent
-		/// </summary>
-		public AclConfig? Acl { get; set; }
 		
 		/// <summary>
 		/// New comment
@@ -435,11 +430,6 @@ namespace Horde.Build.Agents
 		/// </summary>
 		public List<GetAgentLeaseResponse> Leases { get; }
 
-		/// <summary>
-		/// Per-object permissions
-		/// </summary>
-		public GetAclResponse? Acl { get; set; }
-		
 		/// <summary>
 		/// Comment for this agent
 		/// </summary>
@@ -451,8 +441,7 @@ namespace Horde.Build.Agents
 		/// <param name="agent">The agent to construct from</param>
 		/// <param name="leases">Active leases</param>
 		/// <param name="rate">Rate for this agent</param>
-		/// <param name="includeAcl">Whether to include the ACL in the response</param>
-		public GetAgentResponse(IAgent agent, List<GetAgentLeaseResponse> leases, double? rate, bool includeAcl)
+		public GetAgentResponse(IAgent agent, List<GetAgentLeaseResponse> leases, double? rate)
 		{
 			Id = agent.Id.ToString();
 			Name = agent.Id.ToString();
@@ -478,7 +467,6 @@ namespace Horde.Build.Agents
 			Workspaces = agent.Workspaces.ConvertAll(x => new GetAgentWorkspaceResponse(x));
 			Capabilities = new { Devices = new[] { new { agent.Properties, agent.Resources } } };
 			Leases = leases;
-			Acl = (includeAcl && agent.Acl != null) ? new GetAclResponse(agent.Acl) : null;
 			Comment = agent.Comment;
 		}
 	}
diff --git a/Engine/Source/Programs/Horde/Horde.Build/Agents/AgentService.cs b/Engine/Source/Programs/Horde/Horde.Build/Agents/AgentService.cs
--- a/Engine/Source/Programs/Horde/Horde.Build/Agents/AgentService.cs
+++ b/Engine/Source/Programs/Horde/Horde.Build/Agents/AgentService.cs
@@ -1024,15 +1024,8 @@ namespace Horde.Build.Agents
 		/// <returns>True if the action is authorized</returns>
 		public async Task<bool> AuthorizeAsync(IAgent agent, AclAction action, ClaimsPrincipal user, GlobalPermissionsCache? cache)
 		{
-			bool? result = agent.Acl?.Authorize(action, user);
-			if (result == null)
-			{
-				return await _aclService.AuthorizeAsync(action, user, cache);
-			}
-			else
-			{
-				return result.Value;
-			}
+			_ = agent;
+			return await _aclService.AuthorizeAsync(action, user, cache);
 		}
 
 		/// <summary>
diff --git a/Engine/Source/Programs/Horde/Horde.Build/Agents/AgentsController.cs b/Engine/Source/Programs/Horde/Horde.Build/Agents/AgentsController.cs
--- a/Engine/Source/Programs/Horde/Horde.Build/Agents/AgentsController.cs
+++ b/Engine/Source/Programs/Horde/Horde.Build/Agents/AgentsController.cs
@@ -127,8 +127,7 @@ namespace Horde.Build.Agents
 				leases.Add(new GetAgentLeaseResponse(lease, details));
 			}
 
-			bool includeAcl = await _agentService.AuthorizeAsync(agent, AclAction.ViewPermissions, User, permissionsCache);
-			return new GetAgentResponse(agent, leases, rate, includeAcl).ApplyFilter(filter);
+			return new GetAgentResponse(agent, leases, rate).ApplyFilter(filter);
 		}
 
 		/// <summary>
@@ -157,12 +156,8 @@ namespace Horde.Build.Agents
 				{
 					return Forbid(AclAction.UpdateAgent, agentId);
 				}
-				if (update.Acl != null && !await _agentService.AuthorizeAsync(agent, AclAction.ChangePermissions, User, cache))
-				{
-					return Forbid(AclAction.ChangePermissions, agentId);
-				}
 
-				IAgent? newAgent = await _agentService.Agents.TryUpdateSettingsAsync(agent, update.Enabled, update.RequestConform, update.RequestFullConform, update.RequestRestart, update.RequestShutdown, $"Manual ({userName})", update.Pools?.ConvertAll(x => new PoolId(x)), Acl.Merge(agent.Acl, update.Acl), update.Comment);
+				IAgent? newAgent = await _agentService.Agents.TryUpdateSettingsAsync(agent, update.Enabled, update.RequestConform, update.RequestFullConform, update.RequestRestart, update.RequestShutdown, $"Manual ({userName})", update.Pools?.ConvertAll(x => new PoolId(x)), update.Comment);
 				if (newAgent == null)
 				{
 					continue;
diff --git a/Engine/Source/Programs/Horde/Horde.Build/Agents/IAgent.cs b/Engine/Source/Programs/Horde/Horde.Build/Agents/IAgent.cs
--- a/Engine/Source/Programs/Horde/Horde.Build/Agents/IAgent.cs
+++ b/Engine/Source/Programs/Horde/Horde.Build/Agents/IAgent.cs
@@ -479,11 +479,6 @@ namespace Horde.Build.Agents
 		/// </summary>
 		public IReadOnlyList<AgentLease> Leases { get; }
 
-		/// <summary>
-		/// ACL for modifying this agent
-		/// </summary>
-		public Acl? Acl { get; }
-
 		/// <summary>
 		/// Last time that the agent was modified
 		/// </summary>
diff --git a/Engine/Source/Programs/Horde/Horde.Build/Agents/IAgentCollection.cs b/Engine/Source/Programs/Horde/Horde.Build/Agents/IAgentCollection.cs
--- a/Engine/Source/Programs/Horde/Horde.Build/Agents/IAgentCollection.cs
+++ b/Engine/Source/Programs/Horde/Horde.Build/Agents/IAgentCollection.cs
@@ -80,10 +80,9 @@ namespace Horde.Build.Agents
 		/// <param name="requestShutdown">Whether to request the machine be shut down</param>
 		/// <param name="shutdownReason">The reason for shutting down agent, ex. Autoscaler/Manual/Unexpected</param>
 		/// <param name="pools">List of pools for the agent</param>
-		/// <param name="acl">New ACL for this agent</param>
 		/// <param name="comment">New comment</param>
 		/// <returns>Version of the software that needs to be installed on the agent. Null if the agent is running the correct version.</returns>
-		Task<IAgent?> TryUpdateSettingsAsync(IAgent agent, bool? enabled = null, bool? requestConform = null, bool? requestFullConform = null, bool? requestRestart = null, bool? requestShutdown = null, string? shutdownReason = null, List<PoolId>? pools = null, Acl? acl = null, string? comment = null);
+		Task<IAgent?> TryUpdateSettingsAsync(IAgent agent, bool? enabled = null, bool? requestConform = null, bool? requestFullConform = null, bool? requestRestart = null, bool? requestShutdown = null, string? shutdownReason = null, List<PoolId>? pools = null, string? comment = null);
 
 		/// <summary>
 		/// Update the current workspaces for an agent.
