Horde: Remove unused ACLs from device interfaces.

diff --git a/Engine/Source/Programs/Horde/Horde.Build/Devices/DeviceCollection.cs b/Engine/Source/Programs/Horde/Horde.Build/Devices/DeviceCollection.cs
--- a/Engine/Source/Programs/Horde/Horde.Build/Devices/DeviceCollection.cs
+++ b/Engine/Source/Programs/Horde/Horde.Build/Devices/DeviceCollection.cs
@@ -4,8 +4,6 @@ using System;
 using System.Collections.Generic;
 using System.Linq;
 using System.Threading.Tasks;
-using Horde.Build.Acls;
-using Horde.Build.Jobs;
 using Horde.Build.Projects;
 using Horde.Build.Server;
 using Horde.Build.Streams;
@@ -76,9 +74,6 @@ namespace Horde.Build.Devices
 			[BsonRequired]
 			public string Name { get; set; } = null!;
 
-			[BsonIgnoreIfNull]
-			public Acl? Acl { get; set; }
-
 			[BsonConstructor]
 			private DevicePoolDocument()
 			{
@@ -237,9 +232,6 @@ namespace Horde.Build.Devices
 			[BsonIgnoreIfNull]
 			public List<DeviceUtilizationTelemetry>? Utilization { get; set; }
 
-			[BsonIgnoreIfNull]
-			public Acl? Acl { get; set; }
-
 			[BsonIgnoreIfNull]
 			public int? Version { get; set; }
 
diff --git a/Engine/Source/Programs/Horde/Horde.Build/Devices/IDevice.cs b/Engine/Source/Programs/Horde/Horde.Build/Devices/IDevice.cs
--- a/Engine/Source/Programs/Horde/Horde.Build/Devices/IDevice.cs
+++ b/Engine/Source/Programs/Horde/Horde.Build/Devices/IDevice.cs
@@ -2,7 +2,6 @@
 
 using System;
 using System.Collections.Generic;
-using Horde.Build.Acls;
 using Horde.Build.Projects;
 using Horde.Build.Streams;
 using Horde.Build.Users;
@@ -161,11 +160,6 @@ namespace Horde.Build.Devices
 		/// Friendly name of the pool
 		/// </summary>
 		string Name { get; }
-
-		/// <summary>
-		/// Acl access to the pool
-		/// </summary>
-		public Acl? Acl { get; }
 	}
 
 	/// <summary>
@@ -346,11 +340,6 @@ namespace Horde.Build.Devices
 		/// Device job utilization history 
 		/// </summary>
 		public List<DeviceUtilizationTelemetry>? Utilization { get; set; }
-
-		/// <summary>
-		/// ACL for modifying this device
-		/// </summary>
-		public Acl? Acl { get; }
 	}
 
 	/// <summary>
