Horde: Remove unused job ACL object.

diff --git a/Engine/Source/Programs/Horde/Horde.Build/Jobs/IJob.cs b/Engine/Source/Programs/Horde/Horde.Build/Jobs/IJob.cs
--- a/Engine/Source/Programs/Horde/Horde.Build/Jobs/IJob.cs
+++ b/Engine/Source/Programs/Horde/Horde.Build/Jobs/IJob.cs
@@ -5,7 +5,6 @@ using System.Collections.Generic;
 using System.Diagnostics.CodeAnalysis;
 using System.Linq;
 using EpicGames.Core;
-using Horde.Build.Acls;
 using Horde.Build.Agents;
 using Horde.Build.Agents.Leases;
 using Horde.Build.Agents.Pools;
@@ -731,11 +730,6 @@ namespace Horde.Build.Jobs
 		/// </summary>
 		public DateTime UpdateTimeUtc { get; }
 
-		/// <summary>
-		/// The ACL for this job.
-		/// </summary>
-		public Acl? Acl { get; }
-
 		/// <summary>
 		/// Update counter for this document. Any updates should compare-and-swap based on the value of this counter, or increment it in the case of server-side updates.
 		/// </summary>
@@ -1496,20 +1490,4 @@ namespace Horde.Build.Jobs
 			return response;
 		}
 	}
-
-	/// <summary>
-	/// Projection of a job definition to just include permissions info
-	/// </summary>
-	public interface IJobPermissions
-	{
-		/// <summary>
-		/// ACL for the job
-		/// </summary>
-		public Acl? Acl { get; }
-
-		/// <summary>
-		/// The stream containing 
-		/// </summary>
-		public StreamId StreamId { get; }
-	}
 }
diff --git a/Engine/Source/Programs/Horde/Horde.Build/Jobs/IJobCollection.cs b/Engine/Source/Programs/Horde/Horde.Build/Jobs/IJobCollection.cs
--- a/Engine/Source/Programs/Horde/Horde.Build/Jobs/IJobCollection.cs
+++ b/Engine/Source/Programs/Horde/Horde.Build/Jobs/IJobCollection.cs
@@ -177,13 +177,6 @@ namespace Horde.Build.Jobs
 		/// <returns>Async task</returns>
 		Task RemoveStreamAsync(StreamId streamId);
 
-		/// <summary>
-		/// Gets a job's permissions info by ID
-		/// </summary>
-		/// <param name="jobId">Unique id of the job</param>
-		/// <returns>The job document</returns>
-		Task<IJobPermissions?> GetPermissionsAsync(JobId jobId);
-
 		/// <summary>
 		/// Searches for jobs matching the given criteria
 		/// </summary>
diff --git a/Engine/Source/Programs/Horde/Horde.Build/Jobs/JobCollection.cs b/Engine/Source/Programs/Horde/Horde.Build/Jobs/JobCollection.cs
--- a/Engine/Source/Programs/Horde/Horde.Build/Jobs/JobCollection.cs
+++ b/Engine/Source/Programs/Horde/Horde.Build/Jobs/JobCollection.cs
@@ -284,8 +284,6 @@ namespace Horde.Build.Jobs
 			[BsonIgnoreIfNull]
 			public DateTime? UpdateTimeUtc { get; set; }
 
-			public Acl? Acl { get; set; }
-
 			[BsonRequired]
 			public int UpdateIndex { get; set; }
 
@@ -343,18 +341,6 @@ namespace Horde.Build.Jobs
 			}
 		}
 
-		/// <summary>
-		/// Projection of a job definition to just include permissions info
-		/// </summary>
-		[SuppressMessage("Design", "CA1812: Class is never instantiated")]
-		class JobPermissions : IJobPermissions
-		{
-			public static ProjectionDefinition<JobDocument> Projection { get; } = Builders<JobDocument>.Projection.Include(x => x.Acl).Include(x => x.StreamId);
-
-			public Acl? Acl { get; set; }
-			public StreamId StreamId { get; set; }
-		}
-
 		/// <summary>
 		/// Maximum number of times a step can be retried (after the original run)
 		/// </summary>
@@ -448,12 +434,6 @@ namespace Horde.Build.Jobs
 			await _jobs.DeleteManyAsync(x => x.StreamId == streamId);
 		}
 
-		/// <inheritdoc/>
-		public async Task<IJobPermissions?> GetPermissionsAsync(JobId jobId)
-		{
-			return await _jobs.Find<JobDocument>(x => x.Id == jobId).Project<JobPermissions>(JobPermissions.Projection).FirstOrDefaultAsync();
-		}
-
 		/// <inheritdoc/>
 		public async Task<List<IJob>> FindAsync(JobId[]? jobIds, StreamId? streamId, string? name, TemplateId[]? templates, int? minChange, int? maxChange, int? preflightChange, bool? preflightOnly, UserId ? preflightStartedByUser, UserId? startedByUser, DateTimeOffset? minCreateTime, DateTimeOffset? maxCreateTime, DateTimeOffset? modifiedBefore, DateTimeOffset? modifiedAfter, int? index, int? count, bool consistentRead, string? indexHint, bool? excludeUserJobs)
 		{
diff --git a/Engine/Source/Programs/Horde/Horde.Build/Jobs/JobMessages.cs b/Engine/Source/Programs/Horde/Horde.Build/Jobs/JobMessages.cs
--- a/Engine/Source/Programs/Horde/Horde.Build/Jobs/JobMessages.cs
+++ b/Engine/Source/Programs/Horde/Horde.Build/Jobs/JobMessages.cs
@@ -364,8 +364,7 @@ namespace Horde.Build.Jobs
 		/// <param name="job">Job to create a response for</param>
 		/// <param name="startedByUserInfo">User that started this job</param>
 		/// <param name="abortedByUserInfo">User that aborted this job</param>
-		/// <param name="aclResponse">The ACL response object</param>
-		public GetJobResponse(IJob job, GetThinUserInfoResponse? startedByUserInfo, GetThinUserInfoResponse? abortedByUserInfo, GetAclResponse? aclResponse)
+		public GetJobResponse(IJob job, GetThinUserInfoResponse? startedByUserInfo, GetThinUserInfoResponse? abortedByUserInfo)
 		{
 			Id = job.Id.ToString();
 			StreamId = job.StreamId.ToString();
@@ -393,7 +392,6 @@ namespace Horde.Build.Jobs
 			Arguments = job.Arguments.ToList();
 			UpdateTime = new DateTimeOffset(job.UpdateTimeUtc);
 			UpdateIssues = job.UpdateIssues;
-			Acl = aclResponse;
 		}
 	}
 
diff --git a/Engine/Source/Programs/Horde/Horde.Build/Jobs/JobService.cs b/Engine/Source/Programs/Horde/Horde.Build/Jobs/JobService.cs
--- a/Engine/Source/Programs/Horde/Horde.Build/Jobs/JobService.cs
+++ b/Engine/Source/Programs/Horde/Horde.Build/Jobs/JobService.cs
@@ -70,7 +70,7 @@ namespace Horde.Build.Jobs
 		/// <summary>
 		/// Map of job id to permissions for that job
 		/// </summary>
-		public Dictionary<JobId, IJobPermissions?> Jobs { get; set; } = new Dictionary<JobId, IJobPermissions?>();
+		public Dictionary<JobId, StreamId> Jobs { get; set; } = new Dictionary<JobId, StreamId>();
 	}
 
 	/// <summary>
@@ -535,16 +535,6 @@ namespace Horde.Build.Jobs
 			return _graphs.GetAsync(job.GraphHash);
 		}
 
-		/// <summary>
-		/// Gets a job's permissions info by ID
-		/// </summary>
-		/// <param name="jobId">Unique id of the job</param>
-		/// <returns>The job document</returns>
-		public Task<IJobPermissions?> GetJobPermissionsAsync(JobId jobId)
-		{
-			return _jobs.GetPermissionsAsync(jobId);
-		}
-
 		/// <summary>
 		/// Searches for jobs matching the given criteria
 		/// </summary>
@@ -1350,29 +1340,6 @@ namespace Horde.Build.Jobs
 			return null;
 		}
 
-		/// <summary>
-		/// Determines if the user is authorized to perform an action on a particular stream
-		/// </summary>
-		/// <param name="acl">The ACL to check</param>
-		/// <param name="streamId">The stream containing this job</param>
-		/// <param name="action">The action being performed</param>
-		/// <param name="user">The principal to authorize</param>
-		/// <param name="cache">Cache of stream permissions</param>
-		/// <returns>True if the action is authorized</returns>
-		private Task<bool> AuthorizeAsync(Acl? acl, StreamId streamId, AclAction action, ClaimsPrincipal user, StreamPermissionsCache? cache)
-		{
-			// Do the regular authorization
-			bool? result = acl?.Authorize(action, user);
-			if (result == null)
-			{
-				return _streamService.AuthorizeAsync(streamId, action, user, cache);
-			}
-			else
-			{
-				return Task.FromResult(result.Value);
-			}
-		}
-
 		/// <summary>
 		/// Determines if the user is authorized to perform an action on a particular stream
 		/// </summary>
@@ -1383,7 +1350,7 @@ namespace Horde.Build.Jobs
 		/// <returns>True if the action is authorized</returns>
 		public Task<bool> AuthorizeAsync(IJob job, AclAction action, ClaimsPrincipal user, StreamPermissionsCache? cache)
 		{
-			return AuthorizeAsync(job.Acl, job.StreamId, action, user, cache);
+			return _streamService.AuthorizeAsync(job.StreamId, action, user, cache);
 		}
 
 		/// <summary>
@@ -1396,17 +1363,19 @@ namespace Horde.Build.Jobs
 		/// <returns>True if the action is authorized</returns>
 		public async Task<bool> AuthorizeAsync(JobId jobId, AclAction action, ClaimsPrincipal user, JobPermissionsCache? cache)
 		{
-			IJobPermissions? permissions;
+			StreamId streamId;
 			if (cache == null)
 			{
-				permissions = await GetJobPermissionsAsync(jobId);
+				IJob? job = await GetJobAsync(jobId);
+				streamId = job?.StreamId ?? StreamId.Empty;
 			}
-			else if (!cache.Jobs.TryGetValue(jobId, out permissions))
+			else if (!cache.Jobs.TryGetValue(jobId, out streamId))
 			{
-				permissions = await GetJobPermissionsAsync(jobId);
-				cache.Jobs.Add(jobId, permissions);
+				IJob? job = await GetJobAsync(jobId);
+				streamId = job?.StreamId ?? StreamId.Empty;
+				cache.Jobs.Add(jobId, streamId);
 			}
-			return permissions != null && await AuthorizeAsync(permissions.Acl, permissions.StreamId, action, user, cache);
+			return !streamId.IsEmpty && await _streamService.AuthorizeAsync(streamId, action, user, cache);
 		}
 
 		/// <summary>
diff --git a/Engine/Source/Programs/Horde/Horde.Build/Jobs/JobsController.cs b/Engine/Source/Programs/Horde/Horde.Build/Jobs/JobsController.cs
--- a/Engine/Source/Programs/Horde/Horde.Build/Jobs/JobsController.cs
+++ b/Engine/Source/Programs/Horde/Horde.Build/Jobs/JobsController.cs
@@ -452,13 +452,7 @@ namespace Horde.Build.Jobs
 				abortedByUserInfo = new GetThinUserInfoResponse(await _userCollection.GetCachedUserAsync(job.AbortedByUserId.Value));
 			}
 
-			GetAclResponse? aclResponse = null;
-			if (includeAcl && job.Acl != null)
-			{
-				aclResponse = new GetAclResponse(job.Acl);
-			}
-
-			GetJobResponse response = new GetJobResponse(job, startedByUserInfo, abortedByUserInfo, aclResponse);
+			GetJobResponse response = new GetJobResponse(job, startedByUserInfo, abortedByUserInfo);
 			if (includeBatches || includeLabels)
 			{
 				if (includeBatches)
