Horde: Add acknowledged by and resolved by information in issue modal header
#skipci
#rnx

diff --git a/Engine/Source/Programs/Horde/HordeDashboard/src/backend/Api.ts b/Engine/Source/Programs/Horde/HordeDashboard/src/backend/Api.ts
--- a/Engine/Source/Programs/Horde/HordeDashboard/src/backend/Api.ts
+++ b/Engine/Source/Programs/Horde/HordeDashboard/src/backend/Api.ts
@@ -2866,7 +2866,7 @@ export type GetIssueResponse = {
 	promoted: boolean;
 
 	/** Owner of the issue */
-	ownerInfo: GetThinUserInfoResponse;
+	ownerInfo?: GetThinUserInfoResponse;
 
 	/** Use that nominated the current owner */
 	nominatedByInfo: GetThinUserInfoResponse;
@@ -2881,7 +2881,7 @@ export type GetIssueResponse = {
 	resolvedAt?: Date | string;
 
 	/** Use info for the person that resolved the issue */
-	resolvedByInfo: GetThinUserInfoResponse;
+	resolvedByInfo?: GetThinUserInfoResponse;
 
 	/**  List of stream paths affected by this issue */
 	streams: string[];
diff --git a/Engine/Source/Programs/Horde/HordeDashboard/src/components/IssueViewV2.tsx b/Engine/Source/Programs/Horde/HordeDashboard/src/components/IssueViewV2.tsx
--- a/Engine/Source/Programs/Horde/HordeDashboard/src/components/IssueViewV2.tsx
+++ b/Engine/Source/Programs/Horde/HordeDashboard/src/components/IssueViewV2.tsx
@@ -1136,14 +1136,25 @@ const StreamCanvas: React.FC = () => {
    </Stack>
 }
 
-const IssueHeader: React.FC = () => {
+
+type SummaryItem = {
+   title: string;
+   text: string;
+   link?: string;
+   strike?: boolean;
+}
+
+const IssueHeader: React.FC<{ items?: SummaryItem[] }> = ({ items }) => {
 
    const issue = details.issue!;
 
+   if (!items?.length && !issue.description) {
+      return null;
+   }
    const renderSummaryItem = (title: string, text: string | undefined, link?: string, strike?: boolean) => {
 
       if (!text) {
-         return null;
+         return <Stack/>;
       }
 
       return <Stack horizontal verticalAlign="center" tokens={{ childrenGap: 16 }}>
@@ -1161,31 +1172,18 @@ const IssueHeader: React.FC = () => {
       </Stack>
    }
 
-   const jiraIssue = details.jiraIssue;
-
-   let jiraAssignee = jiraIssue?.assigneeDisplayName ?? "";
-   if (jiraIssue && !jiraAssignee) {
-      jiraAssignee = "Unassigned";
+   let summaryItems:JSX.Element[] | undefined;
+   if (items?.length) {
+      summaryItems = items.map(i => renderSummaryItem(i.title, i.text, i.link, i.strike));
    }
-   const jiraStatus = jiraIssue?.resolutionName ?? jiraIssue?.statusName ?? "";
 
    return <Stack horizontal tokens={{ childrenGap: 48 }} >
-      {(!!jiraIssue || !!issue.quarantineTimeUtc || !!issue.forceClosedByUserInfo) && <Stack style={{ width: 340, padding: 0 }} tokens={{ childrenGap: 12 }}>
-         {!!jiraIssue && <Stack tokens={{ childrenGap: 4 }}>
-            {renderSummaryItem(`Jira`, jiraIssue.key, jiraIssue.link!, !!jiraIssue.resolutionName)}
-            {(!!jiraAssignee) && renderSummaryItem(`Jira Assignee`, jiraAssignee, jiraIssue!.link!, !!jiraIssue!.resolutionName)}
-            {(!!jiraStatus) && renderSummaryItem(`Jira Status`, jiraStatus, jiraIssue!.link, !!jiraIssue!.resolutionName)}
-         </Stack>}
-         {!!issue.quarantinedByUserInfo && <Stack>
-            {renderSummaryItem(`Quarantined By`, `${issue.quarantinedByUserInfo.name} on ${getShortNiceTime(issue.quarantineTimeUtc)}`)}
-         </Stack>}
-         {!!issue.forceClosedByUserInfo && <Stack>
-            {renderSummaryItem(`Force Closed By`, `${issue.forceClosedByUserInfo.name}`)}
-         </Stack>}
+      {!!summaryItems && <Stack style={{ width: 540, padding: 0 }} tokens={{ childrenGap: 6 }}>
+         {summaryItems}
       </Stack>}
       {!!issue.description && <Stack>
          <Label style={{ padding: 0 }}>{`Description:`}</Label>
-         <Markdown styles={{ root: { maxHeight: 240, maxWidth: 1200, overflow: "auto", th: { fontSize: 12 } } }}>{issue.description}</Markdown>
+         <Markdown styles={{ root: { maxHeight: 240, maxWidth: 1000, overflow: "auto", th: { fontSize: 12 } } }}>{issue.description}</Markdown>
       </Stack>}
    </Stack>
 
@@ -1195,13 +1193,83 @@ const IssueSummaryPanel: React.FC = () => {
 
    const issue = details.issue!;
 
-   if (!issue || (!details.jiraIssue && !issue.description && !issue.quarantinedByUserInfo && !issue.forceClosedByUserInfo)) {
+   if (!issue) {
       return null;
    }
 
+   const items: SummaryItem[] = [];
+
+   if (!!issue.quarantineTimeUtc && !!issue.quarantinedByUserInfo) {
+      items.push({
+         title: "Quarantined By",
+         text: `${issue.quarantinedByUserInfo.name} on ${getShortNiceTime(issue.quarantineTimeUtc)}`
+      })
+   }
+
+   if (!!issue.forceClosedByUserInfo) {
+      items.push({
+         title: "Force Closed By",
+         text: `${issue.forceClosedByUserInfo.name}`
+      })
+   }
+
+   if (!!issue.resolvedAt) {
+      items.push({
+         title: "Resolved By",
+         text: `${issue.resolvedByInfo?.name ?? "Horde"} on ${getShortNiceTime(issue.resolvedAt, true, true)}`
+      })
+   }
+
+   if (!issue.resolvedAt && issue.acknowledgedAt && issue.ownerInfo) {
+      items.push({
+         title: "Acknowledged By",
+         text: `${issue.ownerInfo.name} on ${getShortNiceTime(issue.acknowledgedAt, true, true)}`
+      })
+   }
+
+   const jiraIssue = details.jiraIssue;
+
+   let jiraAssignee = jiraIssue?.assigneeDisplayName ?? "";
+   if (jiraIssue && !jiraAssignee) {
+      jiraAssignee = "Unassigned";
+   }
+   const jiraStatus = jiraIssue?.resolutionName ?? jiraIssue?.statusName ?? "";
+
+   if (!!jiraIssue) {
+      items.push({
+         title: "Jira",
+         text: jiraIssue.key,
+         link: jiraIssue.link!,
+         strike: !!jiraIssue.resolutionName
+      })
+
+      if (jiraAssignee) {
+         items.push({
+            title: "Jira Assignee",
+            text: jiraAssignee,
+            link: jiraIssue.link!,
+            strike: !!jiraIssue.resolutionName
+         })
+      }
+
+      if (jiraStatus) {
+         items.push({
+            title: "Jira Status",
+            text: jiraStatus,
+            link: jiraIssue.link!,
+            strike: !!jiraIssue.resolutionName
+         })
+      }
+   }
+
+   if (!items.length && !issue.description) {
+      return null;
+
+   }
+
    return <Stack style={{ flexBasis: "70px", flexShrink: 0 }}>
       <Stack className={hordeClasses.raised} style={{ maxHeight: "100%" }}>
-         <IssueHeader />
+         <IssueHeader items={items.length ? items : undefined} />
       </Stack>
    </Stack>
 }
@@ -1683,7 +1751,7 @@ const IssueCommandBar: React.FC = () => {
             <Stack styles={{ root: { paddingBottom: 0, paddingRight: 32 } }}>
                <Stack className={hordeClasses.commandBarSmall} horizontal verticalAlign="center" tokens={{ childrenGap: 8 }}>
                   <Stack horizontal style={{ paddingLeft: 12 }}>
-                     <Stack horizontal verticalAlign="center" style={{paddingRight: 18 }}>
+                     <Stack horizontal verticalAlign="center" style={{ paddingRight: 18 }}>
                         {!!suspectRange && <a href={suspectRange} target="blank"> <Stack>
                            <CommandBarButton className={hordeClasses.commandBarSmall} styles={{ root: { padding: "10px 8px 10px 8px" } }} iconProps={{ iconName: "Locate" }} text="Suspects" />
                         </Stack>
@@ -2942,7 +3010,7 @@ export const IssueForceCloseModal: React.FC<{ onClose: () => void }> = ({ onClos
             </Stack>
             <Stack grow />
             <Stack>
-               {!!issue.forceClosedByUserInfo && <Stack horizontal tokens={{ childrenGap: 12 }} style={{ paddingRight: 24 }}>                  
+               {!!issue.forceClosedByUserInfo && <Stack horizontal tokens={{ childrenGap: 12 }} style={{ paddingRight: 24 }}>
                   <DefaultButton text="Ok" disabled={state.submitting} onClick={() => { onClose(); }} />
                </Stack>}
                {!issue.forceClosedByUserInfo && <Stack horizontal tokens={{ childrenGap: 12 }} style={{ paddingRight: 24 }}>
