Horde: Do not add preflight tests to temporal data
#rnx

diff --git a/Engine/Source/Programs/Horde/Horde.Build/Jobs/TestData/TestDataCollection.cs b/Engine/Source/Programs/Horde/Horde.Build/Jobs/TestData/TestDataCollection.cs
--- a/Engine/Source/Programs/Horde/Horde.Build/Jobs/TestData/TestDataCollection.cs
+++ b/Engine/Source/Programs/Horde/Horde.Build/Jobs/TestData/TestDataCollection.cs
@@ -1046,6 +1046,13 @@ namespace Horde.Build.Jobs.TestData
 
 		async Task AddTestReportData(IJob job, IJobStep step, List<TestDataDocument> documents)
 		{
+
+			// do not add preflight to temporal data
+			if (job.PreflightChange != 0)
+			{
+				return;
+			}
+
 			List<AutomatedTestSessionData> sessions = new List<AutomatedTestSessionData>();
 			List<UnrealAutomatedTestData> tests = new List<UnrealAutomatedTestData>();
 
