Horde: Add support for sending telemetry from the Horde server to an external consumer. Disabled by default, but intended for exporting data to Tableau for analysis.

diff --git a/Engine/Source/Programs/Horde/Horde.Build.Tests/TestSetup.cs b/Engine/Source/Programs/Horde/Horde.Build.Tests/TestSetup.cs
--- a/Engine/Source/Programs/Horde/Horde.Build.Tests/TestSetup.cs
+++ b/Engine/Source/Programs/Horde/Horde.Build.Tests/TestSetup.cs
@@ -62,6 +62,7 @@ using Horde.Build.Storage.Backends;
 using Horde.Build.Compute;
 using Horde.Build.Devices;
 using Moq;
+using Horde.Build.Telemetry;
 
 namespace Horde.Build.Tests
 {
@@ -176,6 +177,7 @@ namespace Horde.Build.Tests
 
 			services.AddSingleton(typeof(IAuditLogFactory<>), typeof(AuditLogFactory<>));
 			services.AddSingleton<IAuditLog<AgentId>>(sp => sp.GetRequiredService<IAuditLogFactory<AgentId>>().Create("Agents.Log", "AgentId"));
+			services.AddSingleton<ITelemetrySink, NullTelemetrySink>();
 
 			services.AddSingleton<IAgentCollection, AgentCollection>();
 			services.AddSingleton<IAgentSoftwareCollection, AgentSoftwareCollection>();
diff --git a/Engine/Source/Programs/Horde/Horde.Build/Issues/IssueCollection.cs b/Engine/Source/Programs/Horde/Horde.Build/Issues/IssueCollection.cs
--- a/Engine/Source/Programs/Horde/Horde.Build/Issues/IssueCollection.cs
+++ b/Engine/Source/Programs/Horde/Horde.Build/Issues/IssueCollection.cs
@@ -12,6 +12,7 @@ using Horde.Build.Jobs.Graphs;
 using Horde.Build.Logs;
 using Horde.Build.Server;
 using Horde.Build.Streams;
+using Horde.Build.Telemetry;
 using Horde.Build.Users;
 using Horde.Build.Utilities;
 using Microsoft.Extensions.Caching.Memory;
@@ -422,12 +423,14 @@ namespace Horde.Build.Issues
 		readonly IMongoCollection<IssueStep> _issueSteps;
 		readonly IMongoCollection<IssueSuspect> _issueSuspects;
 		readonly IAuditLog<int> _auditLog;
+		readonly ITelemetrySink _telemetrySink;
 		readonly ILogger _logger;
 
-		public IssueCollection(MongoService mongoService, RedisService redisService, IUserCollection userCollection, IAuditLogFactory<int> auditLogFactory, ILogger<IssueCollection> logger)
+		public IssueCollection(MongoService mongoService, RedisService redisService, IUserCollection userCollection, IAuditLogFactory<int> auditLogFactory, ITelemetrySink telemetrySink, ILogger<IssueCollection> logger)
 		{
 			_redisService = redisService;
 			_userCollection = userCollection;
+			_telemetrySink = telemetrySink;
 			_logger = logger;
 
 			_ledgerSingleton = new SingletonDocument<IssueLedger>(mongoService);
@@ -455,12 +458,6 @@ namespace Horde.Build.Issues
 			_issueSuspects = mongoService.GetCollection<IssueSuspect>("IssuesV2.Suspects", issueSuspectIndexes);
 
 			_auditLog = auditLogFactory.Create("IssuesV2.History", "IssueId");
-
-			if (!mongoService.ReadOnlyMode)
-			{
-				
-
-			}
 		}
 
 		/// <inheritdoc/>
@@ -490,7 +487,13 @@ namespace Horde.Build.Issues
 			update = update.Set(x => x.UpdateIndex, prevUpdateIndex + 1);
 
 			FindOneAndUpdateOptions<Issue, Issue> options = new FindOneAndUpdateOptions<Issue, Issue> { ReturnDocument = ReturnDocument.After };
-			return await _issues.FindOneAndUpdateAsync<Issue>(x => x.Id == issueDocument.Id && x.UpdateIndex == prevUpdateIndex, update, options);
+
+			Issue? newIssue = await _issues.FindOneAndUpdateAsync<Issue>(x => x.Id == issueDocument.Id && x.UpdateIndex == prevUpdateIndex, update, options);
+			if (newIssue != null)
+			{
+				SendTelemetry(newIssue);
+			}
+			return newIssue;
 		}
 
 		async Task<IssueSpan?> TryUpdateSpanAsync(IIssueSpan issueSpan, UpdateDefinition<IssueSpan> update)
@@ -501,7 +504,13 @@ namespace Horde.Build.Issues
 			update = update.Set(x => x.UpdateIndex, prevUpdateIndex + 1);
 
 			FindOneAndUpdateOptions<IssueSpan, IssueSpan> options = new FindOneAndUpdateOptions<IssueSpan, IssueSpan> { ReturnDocument = ReturnDocument.After };
-			return await _issueSpans.FindOneAndUpdateAsync<IssueSpan>(x => x.Id == issueSpanDocument.Id && x.UpdateIndex == prevUpdateIndex, update, options);
+
+			IssueSpan? newIssueSpan = await _issueSpans.FindOneAndUpdateAsync<IssueSpan>(x => x.Id == issueSpanDocument.Id && x.UpdateIndex == prevUpdateIndex, update, options);
+			if (newIssueSpan != null)
+			{
+				SendTelemetry(newIssueSpan);
+			}
+			return newIssueSpan;
 		}
 
 		#region Issues
@@ -513,6 +522,7 @@ namespace Horde.Build.Issues
 
 			Issue newIssue = new Issue(ledger.NextId, summary);
 			await _issues.InsertOneAsync(newIssue);
+			SendTelemetry(newIssue);
 
 			ILogger issueLogger = GetLogger(newIssue.Id);
 			issueLogger.LogInformation("Created issue {IssueId}", newIssue.Id);
@@ -520,6 +530,30 @@ namespace Horde.Build.Issues
 			return newIssue;
 		}
 
+		void SendTelemetry(IIssue issue)
+		{
+			if (_telemetrySink.Enabled)
+			{
+				_telemetrySink.SendEvent("State.Issue", new
+				{
+					Id = issue.Id,
+					AcknowledgedAt = issue.AcknowledgedAt,
+					CreatedAt = issue.CreatedAt,
+					OwnerId = issue.OwnerId,
+					Fingerprints = issue.Fingerprints,
+					FixChange = issue.FixChange,
+					LastSeenAt = issue.LastSeenAt,
+					NominatedAt = issue.NominatedAt,
+					NominatedById = issue.NominatedById,
+					ResolvedAt = issue.ResolvedAt,
+					ResolvedById = issue.ResolvedById,
+					Severity = issue.Severity,
+					Summary = issue.Summary,
+					VerifiedAt = issue.VerifiedAt
+				});
+			}
+		}
+
 		async ValueTask<string> GetUserNameAsync(UserId? UserId)
 		{
 			if (UserId == null)
@@ -1251,9 +1285,28 @@ namespace Horde.Build.Issues
 		{
 			IssueSpan span = new IssueSpan(issueId, newSpan);
 			await _issueSpans.InsertOneAsync(span);
+			SendTelemetry(span);
 			return span;
 		}
 
+		void SendTelemetry(IIssueSpan issueSpan)
+		{
+			if (_telemetrySink.Enabled)
+			{
+				_telemetrySink.SendEvent("State.IssueSpan", new
+				{
+					Id = issueSpan.Id,
+					IssueId = issueSpan.IssueId,
+					Fingerprint = new { Type = issueSpan.Fingerprint.Type, Keys = issueSpan.Fingerprint.Keys },
+					FirstFailure = new { JobId = issueSpan.FirstFailure.JobId, JobName = issueSpan.FirstFailure.JobName, Change = issueSpan.FirstFailure.Change, StepId = issueSpan.FirstFailure.SpanId },
+					LastFailure = (issueSpan.LastFailure != null) ? new { JobId = issueSpan.LastFailure.JobId, JobName = issueSpan.LastFailure.JobName, Change = issueSpan.LastFailure.Change, StepId = issueSpan.LastFailure.SpanId } : null,
+					StreamId = issueSpan.StreamId,
+					StreamName = issueSpan.StreamName,
+					TemplateRefId = issueSpan.TemplateRefId
+				});
+			}
+		}
+
 		/// <inheritdoc/>
 		public async Task<IIssueSpan?> GetSpanAsync(ObjectId spanId)
 		{
@@ -1319,6 +1372,7 @@ namespace Horde.Build.Issues
 				{
 					logger.LogInformation("Added failure for span {SpanId} in job {JobId} at CL {Change}", newSpan.Id, newFailure.JobId, newFailure.Change);
 				}
+				SendTelemetry(newSpan);
 			}
 			return newSpan;
 		}
diff --git a/Engine/Source/Programs/Horde/Horde.Build/Jobs/JobCollection.cs b/Engine/Source/Programs/Horde/Horde.Build/Jobs/JobCollection.cs
--- a/Engine/Source/Programs/Horde/Horde.Build/Jobs/JobCollection.cs
+++ b/Engine/Source/Programs/Horde/Horde.Build/Jobs/JobCollection.cs
@@ -16,6 +16,7 @@ using Horde.Build.Jobs.Graphs;
 using Horde.Build.Logs;
 using Horde.Build.Server;
 using Horde.Build.Streams;
+using Horde.Build.Telemetry;
 using Horde.Build.Users;
 using Horde.Build.Utilities;
 using HordeCommon;
@@ -350,6 +351,7 @@ namespace Horde.Build.Jobs
 		readonly MongoIndex<JobDocument> _createTimeIndex;
 		readonly MongoIndex<JobDocument> _updateTimeIndex;
 		readonly MongoIndex<JobDocument> _streamThenTemplateThenCreationTimeIndex;
+		readonly ITelemetrySink _telemetrySink;
 		readonly IClock _clock;
 		readonly ILogger<JobCollection> _logger;
 
@@ -358,10 +360,12 @@ namespace Horde.Build.Jobs
 		/// </summary>
 		/// <param name="mongoService">The database service singleton</param>
 		/// <param name="clock"></param>
+		/// <param name="telemetrySink">Telemetry sink for data</param>
 		/// <param name="logger">The logger instance</param>
-		public JobCollection(MongoService mongoService, IClock clock, ILogger<JobCollection> logger)
+		public JobCollection(MongoService mongoService, IClock clock, ITelemetrySink telemetrySink, ILogger<JobCollection> logger)
 		{
 			_clock = clock;
+			_telemetrySink = telemetrySink;
 			_logger = logger;
 
 			List<MongoIndex<JobDocument>> indexes = new List<MongoIndex<JobDocument>>();
@@ -407,6 +411,24 @@ namespace Horde.Build.Jobs
 
 			await _jobs.InsertOneAsync(newJob);
 
+			if (_telemetrySink.Enabled)
+			{
+				_telemetrySink.SendEvent("State.Job", new
+				{
+					Arguments = newJob.Arguments,
+					AutoSubmit = newJob.AutoSubmit,
+					Change = newJob.Change,
+					CodeChange = newJob.CodeChange,
+					CreateTimeUtc = newJob.CreateTimeUtc,
+					GraphHash = newJob.GraphHash,
+					Name = newJob.Name,
+					PreflightChange = newJob.PreflightChange,
+					Priority = newJob.Priority,
+					StartedByUserId = newJob.StartedByUserId,
+					TemplateId = newJob.TemplateId
+				});
+			}
+
 			return newJob;
 		}
 
diff --git a/Engine/Source/Programs/Horde/Horde.Build/Jobs/JobStepRefCollection.cs b/Engine/Source/Programs/Horde/Horde.Build/Jobs/JobStepRefCollection.cs
--- a/Engine/Source/Programs/Horde/Horde.Build/Jobs/JobStepRefCollection.cs
+++ b/Engine/Source/Programs/Horde/Horde.Build/Jobs/JobStepRefCollection.cs
@@ -9,6 +9,7 @@ using Horde.Build.Agents.Pools;
 using Horde.Build.Logs;
 using Horde.Build.Server;
 using Horde.Build.Streams;
+using Horde.Build.Telemetry;
 using Horde.Build.Utilities;
 using HordeCommon;
 using MongoDB.Bson.Serialization.Attributes;
@@ -97,14 +98,17 @@ namespace Horde.Build.Jobs
 		}
 
 		readonly IMongoCollection<JobStepRef> _jobStepRefs;
+		readonly ITelemetrySink _telemetrySink;
 
 		/// <summary>
 		/// Constructor
 		/// </summary>
 		/// <param name="mongoService">The database service instance</param>
-		public JobStepRefCollection(MongoService mongoService)
+		/// <param name="telemetrySink">Telemetry sink</param>
+		public JobStepRefCollection(MongoService mongoService, ITelemetrySink telemetrySink)
 		{
 			_jobStepRefs = mongoService.GetCollection<JobStepRef>("JobStepRefs", keys => keys.Ascending(x => x.StreamId).Ascending(x => x.TemplateId).Ascending(x => x.Name).Descending(x => x.Change));
+			_telemetrySink = telemetrySink;
 		}
 
 		/// <inheritdoc/>
@@ -112,6 +116,28 @@ namespace Horde.Build.Jobs
 		{
 			JobStepRef newJobStepRef = new JobStepRef(id, jobName, stepName, streamId, templateId, change, logId, poolId, agentId, outcome, updateIssues, lastSuccess, lastWarning, waitTime, initTime, startTimeUtc, finishTimeUtc);
 			await _jobStepRefs.ReplaceOneAsync(Builders<JobStepRef>.Filter.Eq(x => x.Id, newJobStepRef.Id), newJobStepRef, new ReplaceOptions { IsUpsert = true });
+
+			if (_telemetrySink.Enabled)
+			{
+				_telemetrySink.SendEvent("State.JobStepRef", new
+				{
+					Id = id,
+					AgentId = agentId,
+					BatchInitTime = initTime,
+					BatchWaitTime = waitTime,
+					Change = change,
+					FinishTime = finishTimeUtc,
+					JobName = jobName,
+					StepName = stepName,
+					Outcome = outcome,
+					PoolId = poolId,
+					StartTime = startTimeUtc,
+					StreamId = streamId,
+					TemplateId = templateId,
+					UpdateIssues = updateIssues
+				});
+			}
+
 			return newJobStepRef;
 		}
 
diff --git a/Engine/Source/Programs/Horde/Horde.Build/ServerSettings.cs b/Engine/Source/Programs/Horde/Horde.Build/ServerSettings.cs
--- a/Engine/Source/Programs/Horde/Horde.Build/ServerSettings.cs
+++ b/Engine/Source/Programs/Horde/Horde.Build/ServerSettings.cs
@@ -9,6 +9,7 @@ using EpicGames.Perforce;
 using Horde.Build.Agents.Fleet;
 using Horde.Build.Server;
 using Horde.Build.Storage.Backends;
+using Horde.Build.Telemetry;
 using Horde.Build.Utilities;
 using MongoDB.Driver.Core.Configuration;
 using Serilog.Events;
@@ -162,6 +163,43 @@ namespace Horde.Build
 		/// </summary>
 		Worker
 	}
+
+	/// <summary>
+	/// Type of telemetry provider to use
+	/// </summary>
+	public enum TelemetrySinkType
+	{
+		/// <summary>
+		/// No telemetry sink (default)
+		/// </summary>
+		None,
+
+		/// <summary>
+		/// Use the Epic telemetry sink
+		/// </summary>
+		Epic,
+	}
+
+	/// <summary>
+	/// Configuration for the telemetry sink
+	/// </summary>
+	public class TelemetryConfig : IEpicTelemetrySinkConfig
+	{
+		/// <summary>
+		/// Type of telemetry sink
+		/// </summary>
+		public TelemetrySinkType Type { get; set; } = TelemetrySinkType.None;
+
+		/// <summary>
+		/// Base URL for the telemetry server
+		/// </summary>
+		public Uri? Url { get; set; }
+
+		/// <summary>
+		/// Application name to send in the event messages
+		/// </summary>
+		public string AppId { get; set; } = "Horde";
+	}
 	
 	/// <summary>
 	/// Feature flags to aid rollout of new features
@@ -647,6 +685,11 @@ namespace Horde.Build
 		/// </summary>
 		public CommitSettings Commits { get; set; } = new CommitSettings();
 
+		/// <summary>
+		/// Settings for sending telemetry events
+		/// </summary>
+		public TelemetryConfig Telemetry { get; set; } = new TelemetryConfig();
+
 		/// <summary>
 		/// Helper method to check if this process has activated the given mode
 		/// </summary>
diff --git a/Engine/Source/Programs/Horde/Horde.Build/Startup.cs b/Engine/Source/Programs/Horde/Horde.Build/Startup.cs
--- a/Engine/Source/Programs/Horde/Horde.Build/Startup.cs
+++ b/Engine/Source/Programs/Horde/Horde.Build/Startup.cs
@@ -89,6 +89,7 @@ using Horde.Build.Users;
 using Horde.Build.Perforce;
 using Horde.Build.Projects;
 using Horde.Build.Streams;
+using Horde.Build.Telemetry;
 using Horde.Build.Ugs;
 using Horde.Build.Auditing;
 using Horde.Build.Server.Notices;
@@ -608,6 +609,20 @@ namespace Horde.Build
 					throw new ArgumentException($"Invalid auth method {settings.AuthMethod}");
 			}
 
+			TelemetryConfig telemetryConfig = settings.Telemetry;
+			switch (telemetryConfig.Type)
+			{
+				case TelemetrySinkType.None:
+					services.AddSingleton<ITelemetrySink, NullTelemetrySink>();
+					break;
+				case TelemetrySinkType.Epic:
+					services.AddHttpClient(EpicTelemetrySink.HttpClientName, client => { });
+					services.AddSingleton<EpicTelemetrySink>();
+					services.AddSingleton<ITelemetrySink>(sp => sp.GetRequiredService<EpicTelemetrySink>());
+					services.AddHostedService(sp => sp.GetRequiredService<EpicTelemetrySink>());
+					break;
+			}
+
 			authBuilder.AddScheme<JwtBearerOptions, HordeJwtBearerHandler>(HordeJwtBearerHandler.AuthenticationScheme, options => { });
 			schemes.Add(HordeJwtBearerHandler.AuthenticationScheme);
 
diff --git a/Engine/Source/Programs/Horde/Horde.Build/Telemetry/EpicTelemetrySink.cs b/Engine/Source/Programs/Horde/Horde.Build/Telemetry/EpicTelemetrySink.cs
new file mode 100644
--- /dev/null
+++ b/Engine/Source/Programs/Horde/Horde.Build/Telemetry/EpicTelemetrySink.cs
@@ -0,0 +1,208 @@
+// Copyright Epic Games, Inc. All Rights Reserved.
+
+using System;
+using System.Collections.Generic;
+using System.Net.Http;
+using System.Net.Http.Headers;
+using System.Text;
+using System.Text.Json;
+using System.Threading;
+using System.Threading.Tasks;
+using EpicGames.Core;
+using HordeCommon;
+using Microsoft.AspNetCore.WebUtilities;
+using Microsoft.Extensions.Hosting;
+using Microsoft.Extensions.Logging;
+using Microsoft.Extensions.Options;
+
+namespace Horde.Build.Telemetry
+{
+	/// <summary>
+	/// Options for the Epic telemetry sink
+	/// </summary>
+	public interface IEpicTelemetrySinkConfig
+	{
+		/// <summary>
+		/// Base URL for sending events
+		/// </summary>
+		public Uri? Url { get; }
+
+		/// <summary>
+		/// The application id
+		/// </summary>
+		public string AppId { get; }
+	}
+
+	/// <summary>
+	/// Epic internal telemetry sink using the data router
+	/// </summary>
+	public sealed class EpicTelemetrySink : IHostedService, ITelemetrySink, IDisposable
+	{
+		/// <summary>
+		/// Name of the HTTP client for writing telemetry data
+		/// </summary>
+		public const string HttpClientName = "EpicTelemetrySink";
+
+		static readonly JsonEncodedText s_eventsPropertyName = JsonEncodedText.Encode("Events");
+		static readonly JsonEncodedText s_eventNamePropertyName = JsonEncodedText.Encode("EventName");
+
+		readonly object _lockObject = new object();
+
+		readonly IHttpClientFactory _httpClientFactory;
+		readonly Uri? _uri;
+		readonly JsonSerializerOptions _jsonOptions;
+
+		readonly ArrayMemoryWriter _packetMemoryWriter;
+		readonly Utf8JsonWriter _packetWriter;
+
+		readonly ArrayMemoryWriter _eventMemoryWriter;
+		readonly Utf8JsonWriter _eventWriter;
+
+		readonly ITicker _ticker;
+		readonly ILogger _logger;
+
+		/// <inheritdoc/>
+		public bool Enabled => _uri != null;
+
+		/// <summary>
+		/// Constructor
+		/// </summary>
+		public EpicTelemetrySink(IHttpClientFactory httpClientFactory, IClock clock, IOptions<ServerSettings> serverSettings, ILogger<EpicTelemetrySink> logger)
+		{
+			_httpClientFactory = httpClientFactory;
+
+			_jsonOptions = new JsonSerializerOptions();
+			Startup.ConfigureJsonSerializer(_jsonOptions);
+
+			_packetMemoryWriter = new ArrayMemoryWriter(65536);
+			_packetWriter = new Utf8JsonWriter(_packetMemoryWriter);
+
+			_eventMemoryWriter = new ArrayMemoryWriter(65536);
+			_eventWriter = new Utf8JsonWriter(_eventMemoryWriter);
+
+			_ticker = clock.AddTicker<EpicTelemetrySink>(TimeSpan.FromSeconds(10.0), FlushAsync, logger);
+			_logger = logger;
+
+			IEpicTelemetrySinkConfig config = serverSettings.Value.Telemetry;
+			if (config.Url != null)
+			{
+				Dictionary<string, string?> queryParams = new Dictionary<string, string?>()
+				{
+					["SessionID"] = Guid.NewGuid().ToString(),
+					["AppID"] = config.AppId,
+					["AppVersion"] = Program.Version.ToString(),
+					["UploadType"] = "eteventstream",
+					["AppEnvironment"] = "devlatest"
+				};
+				_uri = new Uri(QueryHelpers.AddQueryString(config.Url.ToString(), queryParams));
+			}
+		}
+
+		/// <inheritdoc/>
+		public void Dispose()
+		{
+			_packetWriter.Dispose();
+			_eventWriter.Dispose();
+			_ticker.Dispose();
+		}
+
+		/// <inheritdoc/>
+		public async Task StartAsync(CancellationToken cancellationToken)
+		{
+			await _ticker.StartAsync();
+		}
+
+		/// <inheritdoc/>
+		public async Task StopAsync(CancellationToken cancellationToken)
+		{
+			await _ticker.StopAsync();
+		}
+
+		/// <inheritdoc/>
+		public async ValueTask FlushAsync(CancellationToken cancellationToken)
+		{
+			// Generate the content for this event
+			byte[] packet;
+			lock (_lockObject)
+			{
+				if (_packetMemoryWriter.Length == 0)
+				{
+					return;
+				}
+
+				_packetWriter.WriteEndArray();
+				_packetWriter.WriteEndObject();
+				_packetWriter.Flush();
+
+				packet = _packetMemoryWriter.WrittenMemory.ToArray();
+
+				_packetMemoryWriter.Clear();
+				_packetWriter.Reset(_packetMemoryWriter);
+			}
+
+			// Make sure the settings are valid
+			if (_uri == null)
+			{
+				return;
+			}
+
+			// Post the event data
+			HttpClient httpClient = _httpClientFactory.CreateClient(HttpClientName);
+			using (HttpRequestMessage request = new HttpRequestMessage())
+			{
+				request.RequestUri = _uri;
+				request.Method = HttpMethod.Post;
+				request.Content = new ByteArrayContent(packet);
+				request.Content.Headers.ContentType = new MediaTypeHeaderValue("application/json");
+
+				using (HttpResponseMessage response = await httpClient.SendAsync(request, cancellationToken))
+				{
+					if (!response.IsSuccessStatusCode)
+					{
+						string content = await response.Content.ReadAsStringAsync(cancellationToken);
+						_logger.LogError("Unable to send telemetry data to server ({Code}): {Message}", response.StatusCode, content);
+					}
+				}
+			}
+		}
+
+		/// <inheritdoc/>
+		public void SendEvent(string eventName, object attributes)
+		{
+			lock (_lockObject)
+			{
+				// If this is the first event written, write the outer events array
+				if (_packetMemoryWriter.Length == 0)
+				{
+					_packetWriter.WriteStartObject();
+					_packetWriter.WriteStartArray(s_eventsPropertyName);
+				}
+
+				// Serialize the event data
+				_eventMemoryWriter.Clear();
+				_eventWriter.Reset(_eventMemoryWriter);
+				JsonSerializer.Serialize(_eventWriter, attributes, _jsonOptions);
+				_eventWriter.Flush();
+
+				// Get the event bytes and check it's a full object
+				Span<byte> eventData = _eventMemoryWriter.WrittenSpan;
+				if (eventData[0] != (byte)'{' || eventData[^1] != (byte)'}')
+				{
+					_logger.LogError("Unexpected output from JSON serialization: {Message}", Encoding.UTF8.GetString(eventData));
+					return;
+				}
+
+				// Modify the serialized data so we can append it to another object
+				eventData[0] = (byte)',';
+				eventData = eventData.Slice(0, eventData.Length - 1);
+
+				// Copy the event data to the buffer along with the event name
+				_packetWriter.WriteStartObject();
+				_packetWriter.WriteString(s_eventNamePropertyName, eventName);
+				_packetWriter.Flush();
+				_packetMemoryWriter.WriteFixedLengthBytes(eventData.Slice(0, eventData.Length - 1));
+				_packetWriter.WriteEndObject();
+			}
+		}
+	}
+}
diff --git a/Engine/Source/Programs/Horde/Horde.Build/Telemetry/ITelemetrySink.cs b/Engine/Source/Programs/Horde/Horde.Build/Telemetry/ITelemetrySink.cs
new file mode 100644
--- /dev/null
+++ b/Engine/Source/Programs/Horde/Horde.Build/Telemetry/ITelemetrySink.cs
@@ -0,0 +1,22 @@
+// Copyright Epic Games, Inc. All Rights Reserved.
+
+namespace Horde.Build.Telemetry
+{
+	/// <summary>
+	/// Interface for a telemetry sink
+	/// </summary>
+	public interface ITelemetrySink
+	{
+		/// <summary>
+		/// Whether the sink is enabled
+		/// </summary>
+		bool Enabled { get; }
+
+		/// <summary>
+		/// Sends a telemetry event with the given information
+		/// </summary>
+		/// <param name="eventName">Name of the event</param>
+		/// <param name="attributes">Arbitrary object to include in the payload</param>
+		void SendEvent(string eventName, object attributes);
+	}
+}
diff --git a/Engine/Source/Programs/Horde/Horde.Build/Telemetry/NullTelemetrySink.cs b/Engine/Source/Programs/Horde/Horde.Build/Telemetry/NullTelemetrySink.cs
new file mode 100644
--- /dev/null
+++ b/Engine/Source/Programs/Horde/Horde.Build/Telemetry/NullTelemetrySink.cs
@@ -0,0 +1,18 @@
+// Copyright Epic Games, Inc. All Rights Reserved.
+
+namespace Horde.Build.Telemetry
+{
+	/// <summary>
+	/// Telemetry sink that discards all events
+	/// </summary>
+	public sealed class NullTelemetrySink : ITelemetrySink
+	{
+		/// <inheritdoc/>
+		public bool Enabled => false;
+
+		/// <inheritdoc/>
+		public void SendEvent(string eventName, object attributes)
+		{
+		}
+	}
+}
