Horde: Add log message for succesfully sending telemetry data.

diff --git a/Engine/Source/Programs/Horde/Horde.Build/Telemetry/EpicTelemetrySink.cs b/Engine/Source/Programs/Horde/Horde.Build/Telemetry/EpicTelemetrySink.cs
--- a/Engine/Source/Programs/Horde/Horde.Build/Telemetry/EpicTelemetrySink.cs
+++ b/Engine/Source/Programs/Horde/Horde.Build/Telemetry/EpicTelemetrySink.cs
@@ -157,7 +157,11 @@ namespace Horde.Build.Telemetry
 
 				using (HttpResponseMessage response = await httpClient.SendAsync(request, cancellationToken))
 				{
-					if (!response.IsSuccessStatusCode)
+					if (response.IsSuccessStatusCode)
+					{
+						_logger.LogInformation("Sending {Size} bytes of telemetry data to {Url}", packet.Length, _uri);
+					}
+					else
 					{
 						string content = await response.Content.ReadAsStringAsync(cancellationToken);
 						_logger.LogError("Unable to send telemetry data to server ({Code}): {Message}", response.StatusCode, content);
