Horde: Fix test failures due to un-mocked property.

diff --git a/Engine/Source/Programs/Horde/Horde.Build.Tests/TestDataTests.cs b/Engine/Source/Programs/Horde/Horde.Build.Tests/TestDataTests.cs
--- a/Engine/Source/Programs/Horde/Horde.Build.Tests/TestDataTests.cs
+++ b/Engine/Source/Programs/Horde/Horde.Build.Tests/TestDataTests.cs
@@ -149,6 +149,7 @@ namespace Horde.Build.Tests
 			job.SetupGet(x => x.StreamId).Returns(streamId);
 			job.SetupGet(x => x.TemplateId).Returns(new TemplateId("test-template"));
 			job.SetupGet(x => x.Change).Returns(change);
+			job.SetupGet(x => x.PreflightChange).Returns(0);
 			job.SetupGet(x => x.Batches).Returns(batches);
 			job.SetupGet(x => x.ShowUgsBadges).Returns(false);
 			job.SetupGet(x => x.ShowUgsAlerts).Returns(false);
