Horde: Add user-agent header to telemetry messages.

diff --git a/Engine/Source/Programs/Horde/Horde.Build/Telemetry/EpicTelemetrySink.cs b/Engine/Source/Programs/Horde/Horde.Build/Telemetry/EpicTelemetrySink.cs
--- a/Engine/Source/Programs/Horde/Horde.Build/Telemetry/EpicTelemetrySink.cs
+++ b/Engine/Source/Programs/Horde/Horde.Build/Telemetry/EpicTelemetrySink.cs
@@ -152,6 +152,7 @@ namespace Horde.Build.Telemetry
 			{
 				request.RequestUri = _uri;
 				request.Method = HttpMethod.Post;
+				request.Headers.UserAgent.Add(new ProductInfoHeaderValue("ue/horde", Program.Version.ToString()));
 				request.Content = new ByteArrayContent(packet);
 				request.Content.Headers.ContentType = new MediaTypeHeaderValue("application/json");
 
