Horde: Reduce frequency of telemetry batches, and downgrade verbosity of log messages.

diff --git a/Engine/Source/Programs/Horde/Horde.Build/Telemetry/EpicTelemetrySink.cs b/Engine/Source/Programs/Horde/Horde.Build/Telemetry/EpicTelemetrySink.cs
--- a/Engine/Source/Programs/Horde/Horde.Build/Telemetry/EpicTelemetrySink.cs
+++ b/Engine/Source/Programs/Horde/Horde.Build/Telemetry/EpicTelemetrySink.cs
@@ -80,7 +80,7 @@ namespace Horde.Build.Telemetry
 			_eventMemoryWriter = new ArrayMemoryWriter(65536);
 			_eventWriter = new Utf8JsonWriter(_eventMemoryWriter);
 
-			_ticker = clock.AddTicker<EpicTelemetrySink>(TimeSpan.FromSeconds(10.0), FlushAsync, logger);
+			_ticker = clock.AddTicker<EpicTelemetrySink>(TimeSpan.FromSeconds(60.0), FlushAsync, logger);
 			_logger = logger;
 
 			IEpicTelemetrySinkConfig config = serverSettings.Value.Telemetry;
@@ -160,7 +160,7 @@ namespace Horde.Build.Telemetry
 				{
 					if (response.IsSuccessStatusCode)
 					{
-						_logger.LogInformation("Sending {Size} bytes of telemetry data to {Url}", packet.Length, _uri);
+						_logger.LogDebug("Sending {Size} bytes of telemetry data to {Url}", packet.Length, _uri);
 					}
 					else
 					{
