Horde: Remove unused server settings.

diff --git a/Engine/Source/Programs/Horde/Horde.Build/ServerSettings.cs b/Engine/Source/Programs/Horde/Horde.Build/ServerSettings.cs
--- a/Engine/Source/Programs/Horde/Horde.Build/ServerSettings.cs
+++ b/Engine/Source/Programs/Horde/Horde.Build/ServerSettings.cs
@@ -524,21 +524,6 @@ namespace Horde.Build
 		/// </summary>
 		public string? JobNotificationChannel { get; set; }
 
-		/// <summary>
-		/// URI to the SmtpServer to use for sending email notifications
-		/// </summary>
-		public string? SmtpServer { get; set; }
-
-		/// <summary>
-		/// The email address to send email notifications from
-		/// </summary>
-		public string? EmailSenderAddress { get; set; }
-
-		/// <summary>
-		/// The name for the sender when sending email notifications
-		/// </summary>
-		public string? EmailSenderName { get; set; }
-
 		/// <summary>
 		/// The URl to use for generating links back to the dashboard.
 		/// </summary>
@@ -554,26 +539,6 @@ namespace Horde.Build
 		/// </summary>
 		public string? HelpSlackChannel { get; set; }
 
-		/// <summary>
-		/// The p4 bridge server
-		/// </summary>
-		public string? P4BridgeServer { get; set; }
-
-		/// <summary>
-		/// The p4 bridge service username
-		/// </summary>
-		public string? P4BridgeServiceUsername { get; set; }
-
-		/// <summary>
-		/// The p4 bridge service password
-		/// </summary>
-		public string? P4BridgeServicePassword { get; set; }
-
-		/// <summary>
-		/// Whether the p4 bridge service account can impersonate other users
-		/// </summary>
-		public bool P4BridgeCanImpersonate { get; set; } = false;
-
 		/// <summary>
 		/// Url of P4 Swarm installation
 		/// </summary>
