Horde: Perform a login when connecting to Perforce.

diff --git a/Engine/Source/Programs/Horde/Horde.Build/Configuration/ConfigSource.cs b/Engine/Source/Programs/Horde/Horde.Build/Configuration/ConfigSource.cs
--- a/Engine/Source/Programs/Horde/Horde.Build/Configuration/ConfigSource.cs
+++ b/Engine/Source/Programs/Horde/Horde.Build/Configuration/ConfigSource.cs
@@ -7,7 +7,6 @@ using System.IO;
 using System.Linq;
 using System.Threading;
 using System.Threading.Tasks;
-using Amazon.Runtime.Internal.Endpoints.StandardLibrary;
 using EpicGames.Core;
 using EpicGames.Perforce;
 using Horde.Build.Users;
@@ -269,7 +268,7 @@ namespace Horde.Build.Configuration
 			Dictionary<Uri, IConfigFile> results = new Dictionary<Uri, IConfigFile>();
 			foreach (IGrouping<string, Uri> group in uris.GroupBy(x => x.Host))
 			{
-				using (IPerforceConnection perforce = await ConnectAsync(group.Key))
+				using (IPerforceConnection perforce = await ConnectAsync(group.Key, cancellationToken))
 				{
 					FileSpecList fileSpec = group.Select(x => x.AbsolutePath).Distinct(StringComparer.OrdinalIgnoreCase).ToList();
 
@@ -314,7 +313,7 @@ namespace Horde.Build.Configuration
 			string cacheKey = $"{nameof(PerforceConfigSource)}:data:{uri}@{change}";
 			if (!_cache.TryGetValue(cacheKey, out ReadOnlyMemory<byte> data))
 			{
-				using (IPerforceConnection perforce = await ConnectAsync(uri.Host))
+				using (IPerforceConnection perforce = await ConnectAsync(uri.Host, cancellationToken))
 				{
 					PerforceResponse<PrintRecord<byte[]>> response = await perforce.TryPrintDataAsync($"{uri.AbsolutePath}@{change}", cancellationToken);
 					response.EnsureSuccess();
@@ -330,7 +329,7 @@ namespace Horde.Build.Configuration
 			return data;
 		}
 
-		async Task<IPerforceConnection> ConnectAsync(string host)
+		async Task<IPerforceConnection> ConnectAsync(string host, CancellationToken cancellationToken)
 		{
 			ServerSettings settings = _settings.CurrentValue;
 
@@ -353,7 +352,12 @@ namespace Horde.Build.Configuration
 				}
 			}
 
-			return await PerforceConnection.CreateAsync(connectionSettings.ToPerforceSettings(), _logger);
+			IPerforceConnection connection = await PerforceConnection.CreateAsync(connectionSettings.ToPerforceSettings(), _logger);
+			if (connectionSettings.Credentials?.Password != null)
+			{
+				await connection.LoginAsync(connectionSettings.Credentials.Password, cancellationToken);
+			}
+			return connection;
 		}
 	}
 }
