Horde: Back out change to do explicit P4 login.

diff --git a/Engine/Source/Programs/Horde/Horde.Build/Configuration/ConfigSource.cs b/Engine/Source/Programs/Horde/Horde.Build/Configuration/ConfigSource.cs
--- a/Engine/Source/Programs/Horde/Horde.Build/Configuration/ConfigSource.cs
+++ b/Engine/Source/Programs/Horde/Horde.Build/Configuration/ConfigSource.cs
@@ -353,10 +353,6 @@ namespace Horde.Build.Configuration
 			}
 
 			IPerforceConnection connection = await PerforceConnection.CreateAsync(connectionSettings.ToPerforceSettings(), _logger);
-			if (connectionSettings.Credentials?.Password != null)
-			{
-				await connection.LoginAsync(connectionSettings.Credentials.Password, cancellationToken);
-			}
 			return connection;
 		}
 	}
