Horde: Add an "Id" field to job telemetry events.

diff --git a/Engine/Source/Programs/Horde/Horde.Build/Jobs/JobCollection.cs b/Engine/Source/Programs/Horde/Horde.Build/Jobs/JobCollection.cs
--- a/Engine/Source/Programs/Horde/Horde.Build/Jobs/JobCollection.cs
+++ b/Engine/Source/Programs/Horde/Horde.Build/Jobs/JobCollection.cs
@@ -415,6 +415,7 @@ namespace Horde.Build.Jobs
 			{
 				_telemetrySink.SendEvent("State.Job", new
 				{
+					Id = newJob.Id,
 					Arguments = newJob.Arguments,
 					AutoSubmit = newJob.AutoSubmit,
 					Change = newJob.Change,
