Horde: Rename StartTime property in telemetry event to StartTimeUtc.

diff --git a/Engine/Source/Programs/Horde/Horde.Build/Jobs/JobStepRefCollection.cs b/Engine/Source/Programs/Horde/Horde.Build/Jobs/JobStepRefCollection.cs
--- a/Engine/Source/Programs/Horde/Horde.Build/Jobs/JobStepRefCollection.cs
+++ b/Engine/Source/Programs/Horde/Horde.Build/Jobs/JobStepRefCollection.cs
@@ -131,7 +131,7 @@ namespace Horde.Build.Jobs
 					StepName = stepName,
 					Outcome = outcome,
 					PoolId = poolId,
-					StartTime = startTimeUtc,
+					StartTimeUtc = startTimeUtc,
 					StreamId = streamId,
 					TemplateId = templateId,
 					UpdateIssues = updateIssues
