Horde: Tweak a few more event names to match schema sent to analytics team.

diff --git a/Engine/Source/Programs/Horde/Horde.Build/Issues/IssueCollection.cs b/Engine/Source/Programs/Horde/Horde.Build/Issues/IssueCollection.cs
--- a/Engine/Source/Programs/Horde/Horde.Build/Issues/IssueCollection.cs
+++ b/Engine/Source/Programs/Horde/Horde.Build/Issues/IssueCollection.cs
@@ -540,7 +540,6 @@ namespace Horde.Build.Issues
 					AcknowledgedAt = issue.AcknowledgedAt,
 					CreatedAt = issue.CreatedAt,
 					OwnerId = issue.OwnerId,
-					Fingerprints = issue.Fingerprints,
 					FixChange = issue.FixChange,
 					LastSeenAt = issue.LastSeenAt,
 					NominatedAt = issue.NominatedAt,
diff --git a/Engine/Source/Programs/Horde/Horde.Build/Jobs/JobStepRefCollection.cs b/Engine/Source/Programs/Horde/Horde.Build/Jobs/JobStepRefCollection.cs
--- a/Engine/Source/Programs/Horde/Horde.Build/Jobs/JobStepRefCollection.cs
+++ b/Engine/Source/Programs/Horde/Horde.Build/Jobs/JobStepRefCollection.cs
@@ -131,7 +131,7 @@ namespace Horde.Build.Jobs
 					StepName = stepName,
 					Outcome = outcome,
 					PoolId = poolId,
-					StartTimeUtc = startTimeUtc,
+					StartTime = startTimeUtc,
 					StreamId = streamId,
 					TemplateId = templateId,
 					UpdateIssues = updateIssues
