Horde: Fix capitalization of property name in globals.json.

diff --git a/Engine/Source/Programs/Horde/Horde.Build/Defaults/globals.json b/Engine/Source/Programs/Horde/Horde.Build/Defaults/globals.json
--- a/Engine/Source/Programs/Horde/Horde.Build/Defaults/globals.json
+++ b/Engine/Source/Programs/Horde/Horde.Build/Defaults/globals.json
@@ -8,7 +8,7 @@
 	"compute": [
 		{
 			"id": "default",
-			"namespaceid": "horde.compute"
+			"namespaceId": "horde.compute"
 		}
 	],
 	"storage": {
